[AudioInsights] Moved (Remove/Add due to folder restrictions) AudioInsights plugin to: Main/Engine/Plugins/

diff --git a/Engine/Plugins/AudioInsights/AudioInsights.uplugin b/Engine/Plugins/AudioInsights/AudioInsights.uplugin
new file mode 100644
--- /dev/null
+++ b/Engine/Plugins/AudioInsights/AudioInsights.uplugin
@@ -0,0 +1,41 @@
+{
+	"FileVersion": 3,
+	"Version": 1,
+	"VersionName": "1.0",
+	"FriendlyName": "AudioInsights",
+	"Description": "Suite of tools to profile, debug, and monitor aspects of audio in the Unreal Engine.",
+	"Category": "Audio",
+	"CreatedBy": "Epic Games, Inc.",
+	"CreatedByURL": "",
+	"DocsURL": "",
+	"MarketplaceURL": "",
+	"SupportURL": "",
+	"CanContainContent": true,
+	"IsBetaVersion": true,
+	"IsExperimentalVersion": false,
+	"Installed": false,
+	"Modules":
+	[
+		{
+			"Name": "AudioInsights",
+			"Type": "Editor",
+			"LoadingPhase": "Default",
+			"PlatformAllowList":
+			[
+				"Win64",
+				"Mac"
+			]
+		}
+	],
+	"Plugins":
+	[
+		{
+			"Name": "AudioWidgets",
+			"Enabled": true
+		},
+		{
+			"Name": "GameplayInsights",
+			"Enabled": true
+		}
+	]
+}
\ No newline at end of file
diff --git a/Engine/Plugins/AudioInsights/Source/AudioInsights/AudioInsights.Build.cs b/Engine/Plugins/AudioInsights/Source/AudioInsights/AudioInsights.Build.cs
new file mode 100644
--- /dev/null
+++ b/Engine/Plugins/AudioInsights/Source/AudioInsights/AudioInsights.Build.cs
@@ -0,0 +1,46 @@
+// Copyright Epic Games, Inc. All Rights Reserved.
+
+using UnrealBuildTool;
+
+public class AudioInsights : ModuleRules
+{
+	public AudioInsights(ReadOnlyTargetRules Target) : base(Target)
+	{
+		PCHUsage = ModuleRules.PCHUsageMode.UseExplicitOrSharedPCHs;
+		
+		PublicDependencyModuleNames.AddRange
+		(
+			new string[]
+			{
+				"Core",
+				"TraceServices",
+			}
+		);	
+		
+		PrivateDependencyModuleNames.AddRange
+		(
+			new string[]
+			{
+				"CoreUObject",
+				"Engine",
+				"AudioMixer",
+				"AudioMixerCore",
+				"AudioWidgets",
+				"GameplayInsights",
+				"InputCore",
+				"LevelEditor",
+				"OutputLog",
+				"SignalProcessing",
+				"Slate",
+				"SlateCore",
+				"ToolMenus",
+				"ToolWidgets",
+				"TraceLog",
+				"TraceAnalysis",
+				"TraceInsights",
+				"UnrealEd",
+				"WorkspaceMenuStructure"
+			}
+		);		
+	}
+}
diff --git a/Engine/Plugins/AudioInsights/Source/AudioInsights/Private/Analyzers/AudioMeterAnalyzer.cpp b/Engine/Plugins/AudioInsights/Source/AudioInsights/Private/Analyzers/AudioMeterAnalyzer.cpp
new file mode 100644
--- /dev/null
+++ b/Engine/Plugins/AudioInsights/Source/AudioInsights/Private/Analyzers/AudioMeterAnalyzer.cpp
@@ -0,0 +1,71 @@
+// Copyright Epic Games, Inc. All Rights Reserved.
+#include "AudioMeterAnalyzer.h"
+
+#include "AudioDeviceManager.h"
+#include "AudioInsightsModule.h"
+#include "AudioMeter.h"
+#include "AudioMixerDevice.h"
+#include "Editor.h"
+
+namespace UE::Audio::Insights
+{
+	namespace FAudioMeterAnalyzerPrivate
+	{
+		TSharedRef<AudioWidgets::FAudioMeter> CreateAudioMeter(TWeakObjectPtr<UAudioBus> InExternalAudioBus)
+		{
+			using namespace ::Audio;
+
+			if (FAudioDeviceManager* AudioDeviceManager = FAudioDeviceManager::Get())
+			{
+				const IAudioInsightsModule& InsightsModule = FModuleManager::GetModuleChecked<IAudioInsightsModule>(IAudioInsightsModule::GetName());
+				TArray<UWorld*> DeviceWorlds = AudioDeviceManager->GetWorldsUsingAudioDevice(InsightsModule.GetDeviceId());
+
+				if (UWorld* World = (DeviceWorlds.Num() > 0) ? DeviceWorlds[0] : nullptr)
+				{
+					if (FMixerDevice* MixerDevice = static_cast<FMixerDevice*>(World->GetAudioDeviceRaw()))
+					{
+						return MakeShared<AudioWidgets::FAudioMeter>(InExternalAudioBus.IsValid() ? InExternalAudioBus->GetNumChannels() : MixerDevice->GetNumDeviceChannels(), 
+							*DeviceWorlds[0], 
+							InExternalAudioBus.Get());
+					}
+				}
+			}
+			
+			return MakeShared<AudioWidgets::FAudioMeter>(1, *GEditor->GetWorld());
+		}
+	}
+
+	FAudioMeterAnalyzer::FAudioMeterAnalyzer(TWeakObjectPtr<UAudioBus> InExternalAudioBus)
+		: AudioMeter(FAudioMeterAnalyzerPrivate::CreateAudioMeter(InExternalAudioBus))
+	{
+		
+	}
+
+	void FAudioMeterAnalyzer::RebuildAudioMeter(TWeakObjectPtr<UAudioBus> InExternalAudioBus)
+	{
+		using namespace ::Audio;
+
+		FAudioDeviceManager* AudioDeviceManager = FAudioDeviceManager::Get();
+		if (!AudioDeviceManager)
+		{
+			return;
+		}
+
+		const IAudioInsightsModule& InsightsModule = FModuleManager::GetModuleChecked<IAudioInsightsModule>(IAudioInsightsModule::GetName());
+		TArray<UWorld*> DeviceWorlds = AudioDeviceManager->GetWorldsUsingAudioDevice(InsightsModule.GetDeviceId());
+
+		UWorld* World = (DeviceWorlds.Num() > 0) ? DeviceWorlds[0] : nullptr;
+		if (!World)
+		{
+			return;
+		}
+
+		FMixerDevice* MixerDevice = static_cast<FMixerDevice*>(World->GetAudioDeviceRaw());
+		if (!MixerDevice)
+		{
+			return;
+		}
+
+		AudioMeter->Init(InExternalAudioBus.IsValid() ? InExternalAudioBus->GetNumChannels() : MixerDevice->GetNumDeviceChannels(), *World, InExternalAudioBus.Get());
+	}
+} // namespace UE::Audio::Insights
diff --git a/Engine/Plugins/AudioInsights/Source/AudioInsights/Private/Analyzers/AudioMeterAnalyzer.h b/Engine/Plugins/AudioInsights/Source/AudioInsights/Private/Analyzers/AudioMeterAnalyzer.h
new file mode 100644
--- /dev/null
+++ b/Engine/Plugins/AudioInsights/Source/AudioInsights/Private/Analyzers/AudioMeterAnalyzer.h
@@ -0,0 +1,26 @@
+// Copyright Epic Games, Inc. All Rights Reserved.
+#pragma once
+
+#include "Templates/SharedPointer.h"
+#include "UObject/WeakObjectPtr.h"
+
+namespace AudioWidgets { class FAudioMeter; }
+
+class UAudioBus;
+
+namespace UE::Audio::Insights
+{
+	class FAudioMeterAnalyzer : public TSharedFromThis<FAudioMeterAnalyzer>
+	{
+	public:
+		explicit FAudioMeterAnalyzer(TWeakObjectPtr<UAudioBus> InExternalAudioBus = nullptr);
+		virtual ~FAudioMeterAnalyzer() = default;
+
+		TSharedRef<AudioWidgets::FAudioMeter> GetAudioMeter() { return AudioMeter; };
+
+		void RebuildAudioMeter(TWeakObjectPtr<UAudioBus> InExternalAudioBus = nullptr);
+
+	private:
+		TSharedRef<AudioWidgets::FAudioMeter> AudioMeter;
+	};
+} // namespace UE::Audio::Insights
diff --git a/Engine/Plugins/AudioInsights/Source/AudioInsights/Private/Analyzers/AudioMeterSubmixAnalyzer.cpp b/Engine/Plugins/AudioInsights/Source/AudioInsights/Private/Analyzers/AudioMeterSubmixAnalyzer.cpp
new file mode 100644
--- /dev/null
+++ b/Engine/Plugins/AudioInsights/Source/AudioInsights/Private/Analyzers/AudioMeterSubmixAnalyzer.cpp
@@ -0,0 +1,144 @@
+// Copyright Epic Games, Inc. All Rights Reserved.
+#include "AudioMeterSubmixAnalyzer.h"
+
+#include "AudioBusSubsystem.h"
+#include "AudioDeviceManager.h"
+#include "AudioInsightsModule.h"
+#include "AudioMeter.h"
+#include "AudioMixerDevice.h"
+#include "AudioMixerSubmix.h"
+#include "Editor.h"
+
+namespace UE::Audio::Insights
+{
+	FAudioMeterSubmixAnalyzer::FAudioMeterSubmixAnalyzer(TWeakObjectPtr<USoundSubmix> InSoundSubmix)
+	{
+		SetSubmix(InSoundSubmix);
+	}
+
+	FAudioMeterSubmixAnalyzer::~FAudioMeterSubmixAnalyzer()
+	{
+		UnregisterAudioBusFromSubmix();
+	}
+
+	void FAudioMeterSubmixAnalyzer::SetSubmix(TWeakObjectPtr<USoundSubmix> InSoundSubmix)
+	{
+		using namespace ::Audio;
+
+		if (SoundSubmix.IsValid())
+		{
+			UnregisterAudioBusFromSubmix();
+			SoundSubmix.Reset();
+		}
+
+		AudioMeterAnalyzer.RebuildAudioMeter();
+
+		SoundSubmix = InSoundSubmix;
+
+		if (!SoundSubmix.IsValid())
+		{
+			return;
+		}
+		
+		FAudioDeviceManager* AudioDeviceManager = FAudioDeviceManager::Get();
+		if (!AudioDeviceManager)
+		{
+			return;
+		}
+
+		const IAudioInsightsModule& InsightsModule = FModuleManager::GetModuleChecked<IAudioInsightsModule>(IAudioInsightsModule::GetName());
+		TArray<UWorld*> DeviceWorlds = AudioDeviceManager->GetWorldsUsingAudioDevice(InsightsModule.GetDeviceId());
+
+		UWorld* World = (DeviceWorlds.Num() > 0) ? DeviceWorlds[0] : nullptr;
+		if (!World)
+		{
+			return;
+		}
+
+		FMixerDevice* MixerDevice = static_cast<FMixerDevice*>(World->GetAudioDeviceRaw());
+		if (!MixerDevice)
+		{
+			return;
+		}
+
+		FMixerSubmixPtr MixerSubmix = MixerDevice->GetSubmixInstance(SoundSubmix.Get()).Pin();
+		if (!MixerSubmix.IsValid())
+		{
+			return;
+		}
+
+		// Register audio bus in submix
+		const TObjectPtr<UAudioBus> AudioBus = AudioMeterAnalyzer.GetAudioMeter()->GetAudioBus();
+		if (!AudioBus)
+		{
+			return;
+		}
+
+		const FAudioBusKey AudioBusKey(AudioBus->GetUniqueID());
+		const int32 AudioBusNumChannels = AudioBus->GetNumChannels();
+
+		FAudioThread::RunCommandOnAudioThread([MixerDevice, MixerSubmix, AudioBusKey, AudioBusNumChannels]()
+		{
+			TObjectPtr<UAudioBusSubsystem> AudioBusSubsystem = MixerDevice->GetSubsystem<UAudioBusSubsystem>();
+			check(AudioBusSubsystem);
+
+			MixerSubmix->RegisterAudioBus(AudioBusKey, AudioBusSubsystem->AddPatchInputForAudioBus(AudioBusKey, MixerDevice->GetNumOutputFrames(), AudioBusNumChannels));
+		});
+	}
+
+	void FAudioMeterSubmixAnalyzer::UnregisterAudioBusFromSubmix()
+	{
+		using namespace ::Audio;
+
+		FAudioDeviceManager* AudioDeviceManager = FAudioDeviceManager::Get();
+		if (!AudioDeviceManager)
+		{
+			return;
+		}
+
+		const IAudioInsightsModule& InsightsModule = FModuleManager::GetModuleChecked<IAudioInsightsModule>(IAudioInsightsModule::GetName());
+		TArray<UWorld*> DeviceWorlds = AudioDeviceManager->GetWorldsUsingAudioDevice(InsightsModule.GetDeviceId());
+
+		UWorld* World = (DeviceWorlds.Num() > 0) ? DeviceWorlds[0] : nullptr;
+		if (!World)
+		{
+			return;
+		}
+
+		FMixerDevice* MixerDevice = static_cast<FMixerDevice*>(World->GetAudioDeviceRaw());
+		if (!MixerDevice)
+		{
+			return;
+		}
+
+		if (!SoundSubmix.IsValid())
+		{
+			return;
+		}
+
+		FMixerSubmixPtr MixerSubmix = MixerDevice->GetSubmixInstance(SoundSubmix.Get()).Pin();
+		if (!MixerSubmix.IsValid())
+		{
+			return;
+		}
+
+		// Unregister audio bus from submix
+		const TObjectPtr<UAudioBus> AudioBus = AudioMeterAnalyzer.GetAudioMeter()->GetAudioBus();
+		if (!AudioBus)
+		{
+			return;
+		}
+
+		const FAudioBusKey AudioBusKey(AudioBus->GetUniqueID());
+
+		FAudioThread::RunCommandOnAudioThread([MixerDevice, MixerSubmix, AudioBusKey]()
+		{
+			MixerSubmix->UnregisterAudioBus(AudioBusKey);
+		});
+	}
+
+	TSharedRef<SAudioMeter> FAudioMeterSubmixAnalyzer::GetWidget()
+	{
+		return AudioMeterAnalyzer.GetAudioMeter()->GetWidget();
+	}
+} // namespace UE::Audio::Insights
diff --git a/Engine/Plugins/AudioInsights/Source/AudioInsights/Private/Analyzers/AudioMeterSubmixAnalyzer.h b/Engine/Plugins/AudioInsights/Source/AudioInsights/Private/Analyzers/AudioMeterSubmixAnalyzer.h
new file mode 100644
--- /dev/null
+++ b/Engine/Plugins/AudioInsights/Source/AudioInsights/Private/Analyzers/AudioMeterSubmixAnalyzer.h
@@ -0,0 +1,31 @@
+// Copyright Epic Games, Inc. All Rights Reserved.
+#pragma once
+
+#include "AudioMeterAnalyzer.h"
+#include "Templates/SharedPointer.h"
+#include "UObject/WeakObjectPtr.h"
+
+namespace AudioWidgets { class FAudioMeter; }
+
+class SAudioMeter;
+class USoundSubmix;
+
+namespace UE::Audio::Insights
+{
+	class FAudioMeterSubmixAnalyzer final : public TSharedFromThis<FAudioMeterSubmixAnalyzer>
+	{
+	public:
+		FAudioMeterSubmixAnalyzer(TWeakObjectPtr<USoundSubmix> InSoundSubmix);
+		~FAudioMeterSubmixAnalyzer();
+
+		// Rebuilds the audio meter and registers the newly selected submix
+		void SetSubmix(TWeakObjectPtr<USoundSubmix> InSoundSubmix);
+
+		TSharedRef<SAudioMeter> GetWidget();
+	private:
+		void UnregisterAudioBusFromSubmix();
+
+		TWeakObjectPtr<USoundSubmix> SoundSubmix;
+		FAudioMeterAnalyzer AudioMeterAnalyzer;
+	};
+} // namespace UE::Audio::Insights
diff --git a/Engine/Plugins/AudioInsights/Source/AudioInsights/Private/Analyzers/AudioOscilloscopeAnalyzer.cpp b/Engine/Plugins/AudioInsights/Source/AudioInsights/Private/Analyzers/AudioOscilloscopeAnalyzer.cpp
new file mode 100644
--- /dev/null
+++ b/Engine/Plugins/AudioInsights/Source/AudioInsights/Private/Analyzers/AudioOscilloscopeAnalyzer.cpp
@@ -0,0 +1,180 @@
+// Copyright Epic Games, Inc. All Rights Reserved.
+#include "AudioOscilloscopeAnalyzer.h"
+
+#include "AudioBusSubsystem.h"
+#include "AudioDeviceManager.h"
+#include "AudioInsightsModule.h"
+#include "AudioOscilloscope.h"
+#include "AudioMixerDevice.h"
+#include "AudioMixerSubmix.h"
+#include "Editor.h"
+
+namespace UE::Audio::Insights
+{
+	namespace FAudioOscilloscopeAnalyzerPrivate
+	{
+		TSharedRef<AudioWidgets::FAudioOscilloscope> CreateAudioOscilloscope(const float InTimeWindowMs,
+			const float InMaxTimeWindowMs,
+			const float InAnalysisPeriodMs,
+			const EAudioPanelLayoutType InPanelLayoutType)
+		{
+			using namespace ::Audio;
+
+			if (FAudioDeviceManager* AudioDeviceManager = FAudioDeviceManager::Get())
+			{
+				const IAudioInsightsModule& InsightsModule = FModuleManager::GetModuleChecked<IAudioInsightsModule>(IAudioInsightsModule::GetName());
+				TArray<UWorld*> DeviceWorlds = AudioDeviceManager->GetWorldsUsingAudioDevice(InsightsModule.GetDeviceId());
+
+				if (UWorld* World = (DeviceWorlds.Num() > 0) ? DeviceWorlds[0] : nullptr)
+				{
+					if (FMixerDevice* MixerDevice = static_cast<FMixerDevice*>(World->GetAudioDeviceRaw()))
+					{
+						return MakeShared<AudioWidgets::FAudioOscilloscope>(DeviceWorlds[0],
+							MixerDevice->GetNumDeviceChannels(),
+							InTimeWindowMs,
+							InMaxTimeWindowMs,
+							InAnalysisPeriodMs,
+							InPanelLayoutType);
+					}
+				}
+			}
+			
+			return MakeShared<AudioWidgets::FAudioOscilloscope>(GEditor->GetWorld(), 1, InTimeWindowMs, InMaxTimeWindowMs, InAnalysisPeriodMs, InPanelLayoutType);
+		}
+	}
+
+	FAudioOscilloscopeAnalyzer::FAudioOscilloscopeAnalyzer(TWeakObjectPtr<USoundSubmix> InSoundSubmix)
+		: AudioOscilloscope(FAudioOscilloscopeAnalyzerPrivate::CreateAudioOscilloscope(TimeWindowMs, MaxTimeWindowMs, AnalysisPeriodMs, PanelLayoutType))
+	{
+		RebuildAudioOscilloscope(InSoundSubmix);
+	}
+
+	FAudioOscilloscopeAnalyzer::~FAudioOscilloscopeAnalyzer()
+	{
+		CleanupAudioOscilloscope();
+	}
+
+	void FAudioOscilloscopeAnalyzer::RebuildAudioOscilloscope(TWeakObjectPtr<USoundSubmix> InSoundSubmix)
+	{
+		using namespace ::Audio;
+
+		if (SoundSubmix.IsValid())
+		{
+			CleanupAudioOscilloscope();
+		}
+
+		SoundSubmix = InSoundSubmix;
+
+		FAudioDeviceManager* AudioDeviceManager = FAudioDeviceManager::Get();
+		if (!AudioDeviceManager)
+		{
+			return;
+		}
+
+		const IAudioInsightsModule& InsightsModule = FModuleManager::GetModuleChecked<IAudioInsightsModule>(IAudioInsightsModule::GetName());
+		TArray<UWorld*> DeviceWorlds = AudioDeviceManager->GetWorldsUsingAudioDevice(InsightsModule.GetDeviceId());
+
+		UWorld* World = (DeviceWorlds.Num() > 0) ? DeviceWorlds[0] : nullptr;
+		if (!World)
+		{
+			return;
+		}
+
+		FMixerDevice* MixerDevice = static_cast<FMixerDevice*>(World->GetAudioDeviceRaw());
+		if (!MixerDevice)
+		{
+			return;
+		}
+
+		if (!SoundSubmix.IsValid())
+		{
+			return;
+		}
+
+		FMixerSubmixPtr MixerSubmix = MixerDevice->GetSubmixInstance(SoundSubmix.Get()).Pin();
+		if (!MixerSubmix.IsValid())
+		{
+			return;
+		}
+
+		AudioOscilloscope->CreateDataProvider(World, TimeWindowMs, MaxTimeWindowMs, AnalysisPeriodMs, PanelLayoutType);
+		AudioOscilloscope->CreateOscilloscopeWidget(MixerDevice->GetNumDeviceChannels(), PanelLayoutType);
+
+		// Start processing
+		AudioOscilloscope->StartProcessing();
+		
+		// Register audio bus in submix
+		const TObjectPtr<UAudioBus> AudioBus = AudioOscilloscope->GetAudioBus();
+		if (!AudioBus)
+		{
+			return;
+		}
+
+		const FAudioBusKey AudioBusKey(AudioBus->GetUniqueID());
+		const int32 AudioBusNumChannels = AudioBus->GetNumChannels();
+
+		FAudioThread::RunCommandOnAudioThread([MixerDevice, MixerSubmix, AudioBusKey, AudioBusNumChannels]()
+		{
+			TObjectPtr<UAudioBusSubsystem> AudioBusSubsystem = MixerDevice->GetSubsystem<UAudioBusSubsystem>();
+			check(AudioBusSubsystem);
+
+			MixerSubmix->RegisterAudioBus(AudioBusKey, AudioBusSubsystem->AddPatchInputForAudioBus(AudioBusKey, MixerDevice->GetNumOutputFrames(), AudioBusNumChannels));
+		});
+	}
+
+	void FAudioOscilloscopeAnalyzer::CleanupAudioOscilloscope()
+	{
+		using namespace ::Audio;
+
+		FAudioDeviceManager* AudioDeviceManager = FAudioDeviceManager::Get();
+		if (!AudioDeviceManager)
+		{
+			return;
+		}
+
+		const IAudioInsightsModule& InsightsModule = FModuleManager::GetModuleChecked<IAudioInsightsModule>(IAudioInsightsModule::GetName());
+		TArray<UWorld*> DeviceWorlds = AudioDeviceManager->GetWorldsUsingAudioDevice(InsightsModule.GetDeviceId());
+
+		UWorld* World = (DeviceWorlds.Num() > 0) ? DeviceWorlds[0] : nullptr;
+		if (!World)
+		{
+			return;
+		}
+
+		FMixerDevice* MixerDevice = static_cast<FMixerDevice*>(World->GetAudioDeviceRaw());
+		if (!MixerDevice)
+		{
+			return;
+		}
+
+		if (!SoundSubmix.IsValid())
+		{
+			return;
+		}
+
+		FMixerSubmixPtr MixerSubmix = MixerDevice->GetSubmixInstance(SoundSubmix.Get()).Pin();
+		if (!MixerSubmix.IsValid())
+		{
+			return;
+		}
+
+		// Unregister audio bus from submix
+		const TObjectPtr<UAudioBus> AudioBus = AudioOscilloscope->GetAudioBus();
+		if (!AudioBus)
+		{
+			return;
+		}
+
+		const FAudioBusKey AudioBusKey(AudioBus->GetUniqueID());
+
+		FAudioThread::RunCommandOnAudioThread([MixerDevice, MixerSubmix, AudioBusKey]()
+		{
+			MixerSubmix->UnregisterAudioBus(AudioBusKey);
+		});
+
+		// Stop processing
+		AudioOscilloscope->StopProcessing();
+
+		SoundSubmix.Reset();
+	}
+} // namespace UE::Audio::Insights
diff --git a/Engine/Plugins/AudioInsights/Source/AudioInsights/Private/Analyzers/AudioOscilloscopeAnalyzer.h b/Engine/Plugins/AudioInsights/Source/AudioInsights/Private/Analyzers/AudioOscilloscopeAnalyzer.h
new file mode 100644
--- /dev/null
+++ b/Engine/Plugins/AudioInsights/Source/AudioInsights/Private/Analyzers/AudioOscilloscopeAnalyzer.h
@@ -0,0 +1,35 @@
+// Copyright Epic Games, Inc. All Rights Reserved.
+#pragma once
+
+#include "AudioWidgetsEnums.h"
+#include "Templates/SharedPointer.h"
+#include "UObject/WeakObjectPtr.h"
+
+namespace AudioWidgets { class FAudioOscilloscope; }
+
+class USoundSubmix;
+
+namespace UE::Audio::Insights
+{
+	class FAudioOscilloscopeAnalyzer : public TSharedFromThis<FAudioOscilloscopeAnalyzer>
+	{
+	public:
+		FAudioOscilloscopeAnalyzer(TWeakObjectPtr<USoundSubmix> InSoundSubmix);
+		virtual ~FAudioOscilloscopeAnalyzer();
+
+		TSharedRef<AudioWidgets::FAudioOscilloscope> GetAudioOscilloscope() { return AudioOscilloscope; };
+
+		void RebuildAudioOscilloscope(TWeakObjectPtr<USoundSubmix> InSoundSubmix);
+
+	private:
+		void CleanupAudioOscilloscope();
+
+		TSharedRef<AudioWidgets::FAudioOscilloscope> AudioOscilloscope;
+		TWeakObjectPtr<USoundSubmix> SoundSubmix;
+		
+		static constexpr float TimeWindowMs     = 10.0f;
+		static constexpr float MaxTimeWindowMs  = 10.0f;
+		static constexpr float AnalysisPeriodMs = 10.0f;
+		static constexpr EAudioPanelLayoutType PanelLayoutType = EAudioPanelLayoutType::Basic;
+	};
+} // namespace UE::Audio::Insights
diff --git a/Engine/Plugins/AudioInsights/Source/AudioInsights/Private/AudioInsightsDashboardAssetCommands.cpp b/Engine/Plugins/AudioInsights/Source/AudioInsights/Private/AudioInsightsDashboardAssetCommands.cpp
new file mode 100644
--- /dev/null
+++ b/Engine/Plugins/AudioInsights/Source/AudioInsights/Private/AudioInsightsDashboardAssetCommands.cpp
@@ -0,0 +1,108 @@
+// Copyright Epic Games, Inc. All Rights Reserved.
+#include "AudioInsightsDashboardAssetCommands.h"
+
+#include "AudioInsightsModule.h"
+#include "AudioInsightsStyle.h"
+#include "Framework/Commands/Commands.h"
+#include "Framework/Commands/InputChord.h"
+#include "Framework/Commands/UICommandInfo.h"
+#include "Misc/Attribute.h"
+
+#define LOCTEXT_NAMESPACE "AudioInsights"
+
+
+namespace UE::Audio::Insights
+{
+	FDashboardAssetCommands::FDashboardAssetCommands()
+		: TCommands<FDashboardAssetCommands>("AudioInsightsDashboardAssetCommands", LOCTEXT("AudioInsightsDashboardAssetCommands", "Dashboard Asset Commands"), NAME_None, FSlateStyle::GetStyleName())
+	{
+	}
+
+	void FDashboardAssetCommands::RegisterCommands()
+	{
+		UI_COMMAND(BrowserSync, "Browse", "Browses to the selected asset in the content browser.", EUserInterfaceActionType::Button, FInputChord());
+		UI_COMMAND(Open, "Open", "Opens the selected asset(s) in respective editor(s).", EUserInterfaceActionType::Button, FInputChord());
+
+		UI_COMMAND(Start, "Start Trace", "Starts the active trace session used by Audio Insights.", EUserInterfaceActionType::Button, FInputChord());
+		UI_COMMAND(Stop, "Stop Trace", "Stops the active trace session used by Audio Insights.", EUserInterfaceActionType::Button, FInputChord());
+
+		UI_COMMAND(Mute, "Mute (Toggle)", "Mutes the selected item in the Mixer Source View.", EUserInterfaceActionType::Button, FInputChord());
+		UI_COMMAND(Solo, "Solo (Toggle)", "Solo the selected item in the Mixer Source View.", EUserInterfaceActionType::Button, FInputChord());
+		UI_COMMAND(ClearMuteSolo, "Clear All Mute/Solo", "Clears any Mute/Solo state from any of the items in the Mixer Source View.", EUserInterfaceActionType::Button, FInputChord());
+	}
+
+	void FDashboardAssetCommands::AddAssetCommands(FToolBarBuilder& OutToolbarBuilder) const
+	{
+		OutToolbarBuilder.AddToolBarButton(
+			BrowserSync,
+			NAME_None,
+			TAttribute<FText>(),
+			TAttribute<FText>(),
+			TAttribute<FSlateIcon>::Create([]() { return FSlateStyle::Get().CreateIcon("AudioInsights.Icon.ContentBrowser"); }),
+			"BrowserSync"
+		);
+
+		OutToolbarBuilder.AddToolBarButton(
+			Open,
+			NAME_None,
+			TAttribute<FText>(),
+			TAttribute<FText>(),
+			TAttribute<FSlateIcon>::Create([]() { return FSlateStyle::Get().CreateIcon("AudioInsights.Icon.Open"); }),
+			"Open"
+		);
+	}
+
+	TSharedPtr<const FUICommandInfo> FDashboardAssetCommands::GetBrowserSyncCommand() const
+	{
+		return BrowserSync;
+	};
+
+	TSharedPtr<const FUICommandInfo> FDashboardAssetCommands::GetOpenCommand() const
+	{
+		return Open;
+	}
+
+	TSharedPtr<const FUICommandInfo> FDashboardAssetCommands::GetStartCommand() const
+	{
+		return Start;
+	}
+
+	TSharedPtr<const FUICommandInfo> FDashboardAssetCommands::GetStopCommand() const
+	{
+		return Stop;
+	}
+
+	TSharedPtr<const FUICommandInfo> FDashboardAssetCommands::GetMuteCommand() const
+	{
+		return Mute;
+	};
+
+	TSharedPtr<const FUICommandInfo> FDashboardAssetCommands::GetSoloCommand() const
+	{
+		return Solo;
+	}
+
+	TSharedPtr<const FUICommandInfo> FDashboardAssetCommands::GetClearMuteSoloCommand() const
+	{
+		return ClearMuteSolo;
+	}
+
+	FSlateIcon FDashboardAssetCommands::GetStartIcon() const
+	{
+		const FTraceModule& TraceModule = FAudioInsightsModule::GetChecked().GetTraceModule();
+		return FSlateStyle::Get().CreateIcon(TraceModule.IsTraceAnalysisActive()
+			? "AudioInsights.Icon.Start.Inactive"
+			: "AudioInsights.Icon.Start.Active"
+		);
+	}
+
+	FSlateIcon FDashboardAssetCommands::GetStopIcon() const
+	{
+		const FTraceModule& TraceModule = FAudioInsightsModule::GetChecked().GetTraceModule();
+		return FSlateStyle::Get().CreateIcon(TraceModule.IsTraceAnalysisActive()
+			? "AudioInsights.Icon.Stop.Active"
+			: "AudioInsights.Icon.Stop.Inactive"
+		);
+	}
+} // namespace UE::Audio::Insights
+#undef LOCTEXT_NAMESPACE
diff --git a/Engine/Plugins/AudioInsights/Source/AudioInsights/Private/AudioInsightsDashboardFactory.cpp b/Engine/Plugins/AudioInsights/Source/AudioInsights/Private/AudioInsightsDashboardFactory.cpp
new file mode 100644
--- /dev/null
+++ b/Engine/Plugins/AudioInsights/Source/AudioInsights/Private/AudioInsightsDashboardFactory.cpp
@@ -0,0 +1,585 @@
+// Copyright Epic Games, Inc. All Rights Reserved.
+#include "AudioInsightsDashboardFactory.h"
+
+#include "AudioDevice.h"
+#include "AudioDeviceManager.h"
+#include "AudioInsightsModule.h"
+#include "Editor.h"
+#include "Engine/World.h"
+#include "Framework/Docking/TabManager.h"
+#include "Internationalization/Text.h"
+#include "IPropertyTypeCustomization.h"
+#include "Kismet2/DebuggerCommands.h"
+#include "Templates/SharedPointer.h"
+#include "ToolMenu.h"
+#include "ToolMenus.h"
+#include "Widgets/Input/SCheckBox.h"
+
+#define LOCTEXT_NAMESPACE "AudioInsights"
+
+
+namespace UE::Audio::Insights
+{
+	namespace DashboardFactoryPrivate
+	{
+		static const FText ToolName = LOCTEXT("AudioDashboard_ToolName", "Audio Dashboard");
+
+		static const FLazyName MainToolbarName = "MainToolbar";
+		static const FText MainToolbarDisplayName = LOCTEXT("AudioDashboard_MainToolbarDisplayName", "Dashboard Transport");
+
+		static const FText PreviewDeviceDisplayName = LOCTEXT("AudioDashboard_PreviewDevice", "[Preview Audio]");
+		static const FText DashboardWorldSelectDescription = LOCTEXT("AudioDashboard_SelectWorldDescription", "Select world(s) to monitor (worlds may share audio output).");
+
+		FText GetDebugNameFromDeviceId(::Audio::FDeviceId InDeviceId)
+		{
+			FString WorldName;
+			if (FAudioDeviceManager* DeviceManager = FAudioDeviceManager::Get())
+			{
+				TArray<UWorld*> DeviceWorlds = DeviceManager->GetWorldsUsingAudioDevice(InDeviceId);
+				for (const UWorld* World : DeviceWorlds)
+				{
+					if (!WorldName.IsEmpty())
+					{
+						WorldName += TEXT(", ");
+					}
+					WorldName += World->GetDebugDisplayName();
+				}
+			}
+
+			if (WorldName.IsEmpty())
+			{
+				return PreviewDeviceDisplayName;
+			}
+
+			return FText::FromString(WorldName);
+		}
+	} // namespace DashboardFactoryPrivate
+
+	FDashboardFactory::FDashboardFactory()
+	{
+	}
+
+	void FDashboardFactory::OnWorldRegisteredToAudioDevice(const UWorld* InWorld, ::Audio::FDeviceId InDeviceId)
+	{
+		if (InDeviceId != INDEX_NONE)
+		{
+			if (bStartWithPIE)
+			{
+				const FTraceModule& TraceModule = FAudioInsightsModule::GetChecked().GetTraceModule();
+				TraceModule.StartTraceAnalysis();
+				ActiveDeviceId = InDeviceId;
+			}
+		}
+
+		RefreshDeviceSelector();
+	}
+
+	void FDashboardFactory::OnPIEStarted(bool bSimulating)
+	{
+		if (bStartWithPIE)
+		{
+			const FTraceModule& TraceModule = FAudioInsightsModule::GetChecked().GetTraceModule();
+			TraceModule.StartTraceAnalysis();
+		}
+	}
+
+	void FDashboardFactory::OnPostPIEStarted(bool bSimulating)
+	{
+		OnActiveAudioDeviceChanged.Broadcast();
+	}
+
+	void FDashboardFactory::OnPIEStopped(bool bSimulating)
+	{
+		if (bStopWithPIE)
+		{
+			const FTraceModule& TraceModule = FAudioInsightsModule::GetChecked().GetTraceModule();
+			TraceModule.StopTraceAnalysis();
+
+		}
+
+		RefreshDeviceSelector();
+	}
+
+	void FDashboardFactory::OnWorldUnregisteredFromAudioDevice(const UWorld* InWorld, ::Audio::FDeviceId InDeviceId)
+	{
+		RefreshDeviceSelector();
+	}
+
+	void FDashboardFactory::OnDeviceDestroyed(::Audio::FDeviceId InDeviceId)
+	{
+		if (ActiveDeviceId == InDeviceId)
+		{
+			if (FAudioDeviceManager* DeviceManager = FAudioDeviceManager::Get())
+			{
+				ActiveDeviceId = DeviceManager->GetMainAudioDeviceID();
+			}
+		}
+
+		AudioDeviceIds.RemoveAll([InDeviceId](const TSharedPtr<::Audio::FDeviceId>& DeviceIdPtr)
+		{
+			return *DeviceIdPtr.Get() == InDeviceId;
+		});
+
+		if (AudioDeviceComboBox.IsValid())
+		{
+			AudioDeviceComboBox->RefreshOptions();
+		}
+
+		OnActiveAudioDeviceChanged.Broadcast();
+	}
+
+	void FDashboardFactory::RefreshDeviceSelector()
+	{
+		if (FAudioDeviceManager* DeviceManager = FAudioDeviceManager::Get())
+		{
+			if (!DeviceManager->IsValidAudioDevice(ActiveDeviceId))
+			{
+				ActiveDeviceId = DeviceManager->GetMainAudioDeviceID();
+			}
+		}
+
+		AudioDeviceIds.Empty();
+		if (const FAudioDeviceManager* DeviceManager = FAudioDeviceManager::Get())
+		{
+			DeviceManager->IterateOverAllDevices([this, &DeviceManager](::Audio::FDeviceId DeviceId, const FAudioDevice* AudioDevice)
+			{
+				AudioDeviceIds.Add(MakeShared<::Audio::FDeviceId>(DeviceId));
+			});
+		}
+
+		if (AudioDeviceComboBox.IsValid())
+		{
+			AudioDeviceComboBox->RefreshOptions();
+		}
+	}
+
+	void FDashboardFactory::ResetDelegates()
+	{
+		if (OnWorldRegisteredToAudioDeviceHandle.IsValid())
+		{
+			FAudioDeviceWorldDelegates::OnWorldRegisteredToAudioDevice.Remove(OnWorldRegisteredToAudioDeviceHandle);
+			OnWorldRegisteredToAudioDeviceHandle.Reset();
+		}
+
+		if (OnWorldUnregisteredFromAudioDeviceHandle.IsValid())
+		{
+			FAudioDeviceWorldDelegates::OnWorldUnregisteredWithAudioDevice.Remove(OnWorldUnregisteredFromAudioDeviceHandle);
+			OnWorldUnregisteredFromAudioDeviceHandle.Reset();
+		}
+
+		if (OnDeviceDestroyedHandle.IsValid())
+		{
+			FAudioDeviceManagerDelegates::OnAudioDeviceDestroyed.Remove(OnDeviceDestroyedHandle);
+			OnDeviceDestroyedHandle.Reset();
+		}
+
+		if (OnPIEStartedHandle.IsValid())
+		{
+			FEditorDelegates::PreBeginPIE.Remove(OnPIEStartedHandle);
+			OnPIEStartedHandle.Reset();
+		}
+
+		if (OnPostPIEStartedHandle.IsValid())
+		{
+			FEditorDelegates::PostPIEStarted.Remove(OnPostPIEStartedHandle);
+			OnPostPIEStartedHandle.Reset();
+		}
+
+		if (OnPIEStoppedHandle.IsValid())
+		{
+			FEditorDelegates::EndPIE.Remove(OnPIEStoppedHandle);
+			OnPIEStoppedHandle.Reset();
+		}
+	}
+
+	::Audio::FDeviceId FDashboardFactory::GetDeviceId() const
+	{
+		return ActiveDeviceId;
+	}
+
+	TSharedRef<SDockTab> FDashboardFactory::MakeDockTabWidget(const FSpawnTabArgs& Args)
+	{
+		UnregisterTabSpawners();
+
+		InitDelegates();
+
+		InitTabLayout();
+
+		const TSharedRef<SDockTab> DockTab = SNew(SDockTab)
+			.Label(DashboardFactoryPrivate::ToolName)
+			.Clipping(EWidgetClipping::ClipToBounds)
+			.TabRole(ETabRole::NomadTab);
+
+		TSharedPtr<SWindow> Window = Args.GetOwnerWindow();
+
+		DashboardTabManager = FGlobalTabmanager::Get()->NewTabManager(DockTab);
+
+		RegisterTabSpawners();
+		RefreshDeviceSelector();
+
+		TSharedPtr<SWidget> Content = DashboardTabManager->RestoreFrom(TabLayout->AsShared(), Window);
+		DockTab->SetContent(Content->AsShared());
+		DockTab->SetOnTabClosed(SDockTab::FOnTabClosedCallback::CreateLambda([this](TSharedRef<SDockTab> TabClosed)
+		{
+			ResetDelegates();
+		}));
+		return DockTab;
+	}
+
+	TSharedRef<SWidget> FDashboardFactory::MakeMainToolbarWidget()
+	{
+		static const FName PlayWorldToolBarName = "Kismet.DebuggingViewToolBar";
+		if (!UToolMenus::Get()->IsMenuRegistered(PlayWorldToolBarName))
+		{
+			UToolMenu* ToolBar = UToolMenus::Get()->RegisterMenu(PlayWorldToolBarName, NAME_None, EMultiBoxType::SlimHorizontalToolBar);
+			FToolMenuSection& Section = ToolBar->AddSection("Debug");
+			FPlayWorldCommands::BuildToolbar(Section);
+		}
+
+		return SNew(SHorizontalBox)
+		+ SHorizontalBox::Slot()
+		.HAlign(HAlign_Left)
+		.VAlign(VAlign_Center)
+		.AutoWidth()
+		[
+			SNew(SBorder)
+			.BorderImage(FAppStyle::Get().GetBrush("NoBorder"))
+			[
+				UToolMenus::Get()->GenerateWidget(PlayWorldToolBarName, { FPlayWorldCommands::GlobalPlayWorldActions })
+			]
+		]
+
+		+ SHorizontalBox::Slot()
+		.HAlign(HAlign_Right)
+		.VAlign(VAlign_Center)
+		.AutoWidth()
+		.Padding(2.0f, 0.0f)
+		[
+			SNew(STextBlock)
+			.Text(LOCTEXT("StartOnPIE_DisplayName", "Start with PIE:"))
+			.Font(IPropertyTypeCustomizationUtils::GetRegularFont())
+		]
+		+ SHorizontalBox::Slot()
+		.HAlign(HAlign_Left)
+		.VAlign(VAlign_Center)
+		.AutoWidth()
+		.Padding(2.0f, 0.0f)
+		[
+			SNew(SCheckBox)
+			.IsChecked_Lambda([this]() { return bStartWithPIE ? ECheckBoxState::Checked : ECheckBoxState::Unchecked; })
+			.OnCheckStateChanged_Lambda([this](ECheckBoxState NewState) { bStartWithPIE = NewState == ECheckBoxState::Checked; })
+		]
+		+ SHorizontalBox::Slot()
+		.HAlign(HAlign_Right)
+		.VAlign(VAlign_Center)
+		.AutoWidth()
+		.Padding(2.0f, 0.0f)
+		[
+			SNew(STextBlock)
+			.Text(LOCTEXT("StopOnPIE_DisplayName", "Stop with PIE:"))
+			.Font(IPropertyTypeCustomizationUtils::GetRegularFont())
+		]
+		+ SHorizontalBox::Slot()
+		.HAlign(HAlign_Left)
+		.VAlign(VAlign_Center)
+		.AutoWidth()
+		.Padding(2.0f, 0.0f)
+		[
+			SNew(SCheckBox)
+			.IsChecked_Lambda([this]() { return bStopWithPIE ? ECheckBoxState::Checked : ECheckBoxState::Unchecked; })
+			.OnCheckStateChanged_Lambda([this](ECheckBoxState NewState) { bStopWithPIE = NewState == ECheckBoxState::Checked; })
+		]
+		+ SHorizontalBox::Slot()
+		.HAlign(HAlign_Right)
+		.VAlign(VAlign_Center)
+		.Padding(2.0f, 0.0f)
+		[
+			SNew(STextBlock)
+			.Text(LOCTEXT("SelectDashboardWorld_DisplayName", "World Filter:"))
+			.ToolTipText(DashboardFactoryPrivate::DashboardWorldSelectDescription)
+			.Font(IPropertyTypeCustomizationUtils::GetRegularFont())
+		]
+		+ SHorizontalBox::Slot()
+		.HAlign(HAlign_Left)
+		.VAlign(VAlign_Center)
+		.AutoWidth()
+		.Padding(2.0f, 0.0f)
+		[
+			SAssignNew(AudioDeviceComboBox, SComboBox<TSharedPtr<::Audio::FDeviceId>>)
+			.ToolTipText(DashboardFactoryPrivate::DashboardWorldSelectDescription)
+			.OptionsSource(&AudioDeviceIds)
+			.OnGenerateWidget_Lambda([](const TSharedPtr<::Audio::FDeviceId>& WidgetDeviceId)
+			{
+				FText NameText = DashboardFactoryPrivate::GetDebugNameFromDeviceId(*WidgetDeviceId);
+				return SNew(STextBlock)
+					.Text(NameText)
+					.Font(IPropertyTypeCustomizationUtils::GetRegularFont());
+			})
+			.OnSelectionChanged_Lambda([this](TSharedPtr<::Audio::FDeviceId> NewDeviceId, ESelectInfo::Type)
+			{
+				if (NewDeviceId.IsValid())
+				{
+					ActiveDeviceId = *NewDeviceId;
+					RefreshDeviceSelector();
+
+					OnActiveAudioDeviceChanged.Broadcast();
+				}
+			})
+			[
+				SNew(STextBlock)
+				.Font(IPropertyTypeCustomizationUtils::GetRegularFont())
+				.Text_Lambda([this]()
+				{
+					return DashboardFactoryPrivate::GetDebugNameFromDeviceId(ActiveDeviceId);
+				})
+			]
+		];
+	}
+
+	void FDashboardFactory::InitDelegates()
+	{
+		if (!OnWorldRegisteredToAudioDeviceHandle.IsValid())
+		{
+			OnWorldRegisteredToAudioDeviceHandle = FAudioDeviceWorldDelegates::OnWorldRegisteredToAudioDevice.AddSP(this, &FDashboardFactory::OnWorldRegisteredToAudioDevice);
+		}
+
+		if (!OnWorldUnregisteredFromAudioDeviceHandle.IsValid())
+		{
+			OnWorldUnregisteredFromAudioDeviceHandle = FAudioDeviceWorldDelegates::OnWorldUnregisteredWithAudioDevice.AddSP(this, &FDashboardFactory::OnWorldUnregisteredFromAudioDevice);
+		}
+
+		if (!OnDeviceDestroyedHandle.IsValid())
+		{
+			OnDeviceDestroyedHandle = FAudioDeviceManagerDelegates::OnAudioDeviceDestroyed.AddSP(this, &FDashboardFactory::OnDeviceDestroyed);
+		}
+
+		if (!OnPIEStartedHandle.IsValid())
+		{
+			OnPIEStartedHandle = FEditorDelegates::PreBeginPIE.AddSP(this, &FDashboardFactory::OnPIEStarted);
+		}
+
+		if (!OnPostPIEStartedHandle.IsValid())
+		{
+			OnPostPIEStartedHandle = FEditorDelegates::PostPIEStarted.AddSP(this, &FDashboardFactory::OnPostPIEStarted);
+		}
+
+		if (!OnPIEStoppedHandle.IsValid())
+		{
+			OnPIEStoppedHandle = FEditorDelegates::EndPIE.AddSP(this, &FDashboardFactory::OnPIEStopped);
+		}
+	}
+
+	void FDashboardFactory::InitTabLayout()
+	{
+		using namespace DashboardFactoryPrivate;
+
+		TabLayout.Reset();
+
+		TSharedRef<FTabManager::FStack> MainMenuTabStack = FTabManager::NewStack();
+		TSharedRef<FTabManager::FStack> ViewportTabStack = FTabManager::NewStack();
+		TSharedRef<FTabManager::FStack> LogTabStack = FTabManager::NewStack();
+		TSharedRef<FTabManager::FStack> AnalysisTabStack = FTabManager::NewStack();
+		TSharedRef<FTabManager::FStack> AudioMetersTabStack = FTabManager::NewStack();
+		TSharedRef<FTabManager::FStack> AudioMeterTabStack = FTabManager::NewStack();
+		TSharedRef<FTabManager::FStack> OscilloscopeTabStack = FTabManager::NewStack();
+
+		MainMenuTabStack->AddTab(MainToolbarName, ETabState::OpenedTab);
+
+		for (const TPair<FName, TSharedPtr<IDashboardViewFactory>>& Factory : DashboardViewFactories)
+		{
+			EDefaultDashboardTabStack DefaultTabStack = Factory.Value->GetDefaultTabStack();
+			switch (DefaultTabStack)
+			{
+				case EDefaultDashboardTabStack::Viewport:
+				{
+					ViewportTabStack->AddTab(Factory.Key, ETabState::OpenedTab);
+				}
+				break;
+
+				case EDefaultDashboardTabStack::Log:
+				{
+					LogTabStack->AddTab(Factory.Key, ETabState::OpenedTab);
+				}
+				break;
+
+				case EDefaultDashboardTabStack::Analysis:
+				{
+					AnalysisTabStack->AddTab(Factory.Key, ETabState::OpenedTab);
+				}
+				break;
+
+				case EDefaultDashboardTabStack::AudioMeters:
+				{
+					AudioMetersTabStack->AddTab(Factory.Key, ETabState::OpenedTab);
+				}
+				break;
+
+				case EDefaultDashboardTabStack::AudioMeter:
+				{
+					AudioMeterTabStack->AddTab(Factory.Key, ETabState::OpenedTab);
+				}
+				break;
+
+				case EDefaultDashboardTabStack::Oscilloscope:
+				{
+					OscilloscopeTabStack->AddTab(Factory.Key, ETabState::OpenedTab);
+				}
+				break;
+
+				default:
+					break;
+			}
+		}
+
+		AnalysisTabStack->SetForegroundTab(FName("MixerSources"));
+
+		TabLayout = FTabManager::NewLayout("AudioDashboard_Layout_v1")
+		->AddArea
+		(
+			FTabManager::NewPrimaryArea()
+			->SetOrientation(Orient_Vertical)
+			->Split
+			(
+				MainMenuTabStack
+				->SetSizeCoefficient(0.065)
+				->SetHideTabWell(true)
+			)
+			->Split
+			(
+				// Left column
+				FTabManager::NewSplitter()
+				->SetOrientation(Orient_Horizontal)
+				->Split
+				(
+					// Top
+					FTabManager::NewSplitter()
+					->SetOrientation(Orient_Vertical)
+					->SetSizeCoefficient(0.25f) // Column width
+					->Split
+					(
+						ViewportTabStack
+						->SetSizeCoefficient(0.5f)
+					)
+					// Bottom
+					->Split
+					(
+						LogTabStack
+						->SetSizeCoefficient(0.5f)
+					)
+				)
+				
+				// Middle column
+				->Split
+				(
+					// Top
+					FTabManager::NewSplitter()
+					->SetOrientation(Orient_Vertical)
+					->SetSizeCoefficient(0.6f) // Column width
+					->Split
+					(
+						FTabManager::NewSplitter()
+						->SetOrientation(Orient_Horizontal)
+						->Split
+						(
+							AnalysisTabStack
+							->SetSizeCoefficient(0.58f)
+						)
+					)
+					// Bottom
+					->Split
+					(
+						AudioMetersTabStack
+						->SetSizeCoefficient(0.42f)
+					)
+				)
+				// Right column
+				->Split
+				(
+					// Top
+					FTabManager::NewSplitter()
+					->SetOrientation(Orient_Vertical)
+					->SetSizeCoefficient(0.15f) // Column width
+					->Split
+					(
+						AudioMeterTabStack
+						->SetSizeCoefficient(0.7f)
+						->SetHideTabWell(true)
+					)
+					// Bottom
+					->Split
+					(
+						OscilloscopeTabStack
+						->SetSizeCoefficient(0.3f)
+						->SetHideTabWell(true)
+					)
+				)
+			)
+		);
+	}
+
+	void FDashboardFactory::RegisterTabSpawners()
+	{
+		DashboardWorkspace = DashboardTabManager->AddLocalWorkspaceMenuCategory(DashboardFactoryPrivate::ToolName);
+		DashboardTabManager->RegisterTabSpawner(DashboardFactoryPrivate::MainToolbarName, FOnSpawnTab::CreateLambda([this](const FSpawnTabArgs& Args)
+		{
+			return SNew(SDockTab)
+				.Clipping(EWidgetClipping::ClipToBounds)
+				.Label(DashboardFactoryPrivate::MainToolbarDisplayName)
+				[
+					MakeMainToolbarWidget()
+				];
+		}))
+		.SetDisplayName(DashboardFactoryPrivate::MainToolbarDisplayName)
+		.SetGroup(DashboardWorkspace->AsShared());
+
+		for (const TPair<FName, TSharedPtr<IDashboardViewFactory>>& KVP : DashboardViewFactories)
+		{
+			const FName FactoryName = KVP.Value->GetName();
+			DashboardTabManager->RegisterTabSpawner(FactoryName, FOnSpawnTab::CreateLambda([this, Factory = KVP.Value](const FSpawnTabArgs& Args)
+			{
+				TSharedPtr<SWidget> DashboardView = Factory->MakeWidget();
+				return SNew(SDockTab)
+					.Clipping(EWidgetClipping::ClipToBounds)
+					.Label(Factory->GetDisplayName())
+					[
+						DashboardView->AsShared()
+					];
+			}))
+			.SetDisplayName(KVP.Value->GetDisplayName())
+			.SetGroup(DashboardWorkspace->AsShared())
+			.SetIcon(KVP.Value->GetIcon());
+		}
+	}
+
+	void FDashboardFactory::RegisterViewFactory(TSharedRef<IDashboardViewFactory> InFactory)
+	{
+		const FName Name = InFactory->GetName();
+		if (ensureAlwaysMsgf(!DashboardViewFactories.Contains(Name), TEXT("Failed to register Audio Dashboard '%s': Dashboard with name already registered"), *Name.ToString()))
+		{
+			DashboardViewFactories.Add(Name, InFactory);
+		}
+	}
+
+	void FDashboardFactory::UnregisterViewFactory(FName InName)
+	{
+		DashboardViewFactories.Remove(InName);
+	}
+
+	void FDashboardFactory::UnregisterTabSpawners()
+	{
+		using namespace DashboardFactoryPrivate;
+
+		if (DashboardTabManager.IsValid())
+		{
+			for (const TPair<FName, TSharedPtr<IDashboardViewFactory>>& Factory : DashboardViewFactories)
+			{
+				DashboardTabManager->UnregisterTabSpawner(Factory.Value->GetName());
+			}
+
+			DashboardTabManager.Reset();
+		}
+
+		DashboardWorkspace.Reset();
+	}
+} // namespace UE::Audio::Insights
+#undef LOCTEXT_NAMESPACE
diff --git a/Engine/Plugins/AudioInsights/Source/AudioInsights/Private/AudioInsightsDashboardFactory.h b/Engine/Plugins/AudioInsights/Source/AudioInsights/Private/AudioInsightsDashboardFactory.h
new file mode 100644
--- /dev/null
+++ b/Engine/Plugins/AudioInsights/Source/AudioInsights/Private/AudioInsightsDashboardFactory.h
@@ -0,0 +1,77 @@
+// Copyright Epic Games, Inc. All Rights Reserved.
+#pragma once
+
+#include "AudioDefines.h"
+#include "Containers/Map.h"
+#include "Engine/World.h"
+#include "Framework/Docking/TabManager.h"
+#include "IAudioInsightsDashboardFactory.h"
+#include "Templates/SharedPointer.h"
+#include "Views/DashboardViewFactory.h"
+#include "Widgets/Docking/SDockTab.h"
+#include "Widgets/Input/SComboBox.h"
+
+namespace UE::Audio::Insights
+{
+	class FDashboardFactory : public IDashboardFactory, public TSharedFromThis<FDashboardFactory>
+	{
+		class FSceneViewport;
+		class FEditorViewportClient;
+
+	public:
+		FDashboardFactory();
+		virtual ~FDashboardFactory() = default;
+
+		TSharedRef<SDockTab> MakeDockTabWidget(const FSpawnTabArgs& Args);
+
+		virtual void RegisterViewFactory(TSharedRef<IDashboardViewFactory> InFactory) override;
+		virtual void UnregisterViewFactory(FName InDashboardName) override;
+		virtual ::Audio::FDeviceId GetDeviceId() const override;
+
+		DECLARE_MULTICAST_DELEGATE(FOnActiveAudioDeviceChanged);
+		inline static FOnActiveAudioDeviceChanged OnActiveAudioDeviceChanged;
+
+	private:
+		void OnWorldRegisteredToAudioDevice(const UWorld* InWorld, ::Audio::FDeviceId InDeviceId);
+		void OnWorldUnregisteredFromAudioDevice(const UWorld* InWorld, ::Audio::FDeviceId InDeviceId);
+		void OnDeviceDestroyed(::Audio::FDeviceId InDeviceId);
+		void OnPIEStarted(bool bSimulating);
+		void OnPostPIEStarted(bool bSimulating);
+		void OnPIEStopped(bool bSimulating);
+
+		void RefreshDeviceSelector();
+		void ResetDelegates();
+
+		TSharedRef<SWidget> MakeMainToolbarWidget();
+
+		void InitDelegates();
+		void InitTabLayout();
+
+		void RegisterTabSpawners();
+		void UnregisterTabSpawners();
+
+		FDelegateHandle OnDeviceDestroyedHandle;
+		FDelegateHandle OnPIEStartedHandle;
+		FDelegateHandle OnPostPIEStartedHandle;
+		FDelegateHandle OnPIEStoppedHandle;
+		FDelegateHandle OnWorldRegisteredToAudioDeviceHandle;
+		FDelegateHandle OnWorldUnregisteredFromAudioDeviceHandle;
+
+		bool bStartWithPIE = true;
+		bool bStopWithPIE = true;
+
+		TSharedPtr<FTabManager> DashboardTabManager;
+		TSharedPtr<FWorkspaceItem> DashboardWorkspace;
+		TSharedPtr<FTabManager::FLayout> TabLayout;
+
+		TSharedPtr<FSceneViewport> SceneViewport;
+		TSharedPtr<FEditorViewportClient> ViewportClient;
+
+		TArray<TSharedPtr<::Audio::FDeviceId>> AudioDeviceIds;
+		TSharedPtr<SComboBox<TSharedPtr<::Audio::FDeviceId>>> AudioDeviceComboBox;
+
+		::Audio::FDeviceId ActiveDeviceId = INDEX_NONE;
+
+		TMap<FName, TSharedPtr<IDashboardViewFactory>> DashboardViewFactories;
+	};
+} // namespace UE::Audio::Insights
diff --git a/Engine/Plugins/AudioInsights/Source/AudioInsights/Private/AudioInsightsModule.cpp b/Engine/Plugins/AudioInsights/Source/AudioInsights/Private/AudioInsightsModule.cpp
new file mode 100644
--- /dev/null
+++ b/Engine/Plugins/AudioInsights/Source/AudioInsights/Private/AudioInsightsModule.cpp
@@ -0,0 +1,125 @@
+// Copyright Epic Games, Inc. All Rights Reserved.
+#include "AudioInsightsModule.h"
+
+#include "AudioInsightsDashboardAssetCommands.h"
+#include "AudioInsightsDashboardFactory.h"
+#include "AudioInsightsLog.h"
+#include "AudioInsightsStyle.h"
+#include "AudioInsightsTraceModule.h"
+#include "Features/IModularFeatures.h"
+#include "Framework/Docking/TabManager.h"
+#include "Modules/ModuleManager.h"
+#include "Templates/SharedPointer.h"
+#include "TraceServices/ModuleService.h"
+#include "UObject/NameTypes.h"
+#include "Views/AudioBusesDashboardViewFactory.h"
+#include "Views/AudioMetersDashboardViewFactory.h"
+#include "Views/LogDashboardViewFactory.h"
+#include "Views/MixerSourceDashboardViewFactory.h"
+#include "Views/OutputMeterDashboardViewFactory.h"
+#include "Views/OutputOscilloscopeDashboardViewFactory.h"
+#include "Views/SubmixesDashboardViewFactory.h"
+#include "Views/ViewportDashboardViewFactory.h"
+#include "Views/VirtualLoopDashboardViewFactory.h"
+#include "WorkspaceMenuStructure.h"
+#include "WorkspaceMenuStructureModule.h"
+
+#define LOCTEXT_NAMESPACE "AudioInsights"
+DEFINE_LOG_CATEGORY(LogAudioInsights);
+
+
+namespace UE::Audio::Insights
+{
+	void FAudioInsightsModule::StartupModule()
+	{
+		IModularFeatures::Get().RegisterModularFeature(TraceServices::ModuleFeatureName, &TraceModule);
+
+		RegisterMenus();
+
+		DashboardFactory = MakeShared<FDashboardFactory>();
+		DashboardFactory->RegisterViewFactory(MakeShared<FViewportDashboardViewFactory>());
+		DashboardFactory->RegisterViewFactory(MakeShared<FLogDashboardViewFactory>());
+		DashboardFactory->RegisterViewFactory(MakeShared<FMixerSourceDashboardViewFactory>());
+		DashboardFactory->RegisterViewFactory(MakeShared<FVirtualLoopDashboardViewFactory>());
+		DashboardFactory->RegisterViewFactory(MakeShared<FSubmixesDashboardViewFactory>());
+		DashboardFactory->RegisterViewFactory(MakeShared<FAudioBusesDashboardViewFactory>());
+		DashboardFactory->RegisterViewFactory(MakeShared<FAudioMetersDashboardViewFactory>());
+		DashboardFactory->RegisterViewFactory(MakeShared<FOutputMeterDashboardViewFactory>());
+		DashboardFactory->RegisterViewFactory(MakeShared<FOutputOscilloscopeDashboardViewFactory>());
+
+		FDashboardAssetCommands::Register();
+	}
+
+	void FAudioInsightsModule::ShutdownModule()
+	{
+		DashboardFactory.Reset();
+		IModularFeatures::Get().UnregisterModularFeature(TraceServices::ModuleFeatureName, &TraceModule);
+
+		FDashboardAssetCommands::Unregister();
+	}
+
+	void FAudioInsightsModule::RegisterDashboardViewFactory(TSharedRef<IDashboardViewFactory> InDashboardFactory)
+	{
+		DashboardFactory->RegisterViewFactory(InDashboardFactory);
+	}
+
+	void FAudioInsightsModule::UnregisterDashboardViewFactory(FName InName)
+	{
+		DashboardFactory->UnregisterViewFactory(InName);
+	}
+
+	::Audio::FDeviceId FAudioInsightsModule::GetDeviceId() const
+	{
+		return DashboardFactory->GetDeviceId();
+	}
+
+	FAudioInsightsModule& FAudioInsightsModule::GetChecked()
+	{
+		return static_cast<FAudioInsightsModule&>(FModuleManager::GetModuleChecked<IAudioInsightsModule>(GetName()));
+	}
+
+	TSharedRef<FDashboardFactory> FAudioInsightsModule::GetDashboardFactory()
+	{
+		return DashboardFactory->AsShared();
+	}
+
+	const TSharedRef<FDashboardFactory> FAudioInsightsModule::GetDashboardFactory() const
+	{
+		return DashboardFactory->AsShared();
+	}
+
+	FTraceModule& FAudioInsightsModule::GetTraceModule()
+	{
+		return TraceModule;
+	}
+
+	TSharedRef<SDockTab> FAudioInsightsModule::CreateDashboardTabWidget(const FSpawnTabArgs& Args)
+	{
+		return DashboardFactory->MakeDockTabWidget(Args);
+	}
+
+	TSharedRef<SDockTab> FAudioInsightsModule::CreateInsightsTabWidget(const FSpawnTabArgs& Args)
+	{
+		// TODO: Implement
+		return SNew(SDockTab);
+	}
+
+	void FAudioInsightsModule::RegisterMenus()
+	{
+		const IWorkspaceMenuStructure& MenuStructure = WorkspaceMenu::GetMenuStructure();
+		FGlobalTabmanager::Get()->RegisterNomadTabSpawner("AudioDashboard", FOnSpawnTab::CreateRaw(this, &FAudioInsightsModule::CreateDashboardTabWidget))
+			.SetDisplayName(LOCTEXT("OpenDashboard_TabDisplayName", "Audio Dashboard"))
+			.SetTooltipText(LOCTEXT("OpenDashboard_TabTooltip", "Opens Audio Dashboard, an extensible suite of tools and visualizers which enable monitoring and debugging audio in the Unreal Engine."))
+			.SetGroup(MenuStructure.GetToolsCategory())
+			.SetIcon(FSlateStyle::Get().CreateIcon("AudioInsights.Icon.Dashboard"));
+
+		FGlobalTabmanager::Get()->RegisterNomadTabSpawner(GetName(), FOnSpawnTab::CreateRaw(this, &FAudioInsightsModule::CreateInsightsTabWidget))
+			.SetDisplayName(LOCTEXT("OpenAudioInsights_TabDisplayName", "Audio Insights"))
+			.SetTooltipText(LOCTEXT("OpenAudioInsights_TabTooltip", "Opens Insights audio profile mode."))
+			.SetGroup(MenuStructure.GetDeveloperToolsProfilingCategory())
+			.SetIcon(FSlateStyle::Get().CreateIcon("AudioInsights.Icon"));
+	};
+} // namespace UE::Audio::Insights
+#undef LOCTEXT_NAMESPACE // AudioInsights
+
+IMPLEMENT_MODULE(UE::Audio::Insights::FAudioInsightsModule, AudioInsights)
diff --git a/Engine/Plugins/AudioInsights/Source/AudioInsights/Private/AudioInsightsModule.h b/Engine/Plugins/AudioInsights/Source/AudioInsights/Private/AudioInsightsModule.h
new file mode 100644
--- /dev/null
+++ b/Engine/Plugins/AudioInsights/Source/AudioInsights/Private/AudioInsightsModule.h
@@ -0,0 +1,46 @@
+// Copyright Epic Games, Inc. All Rights Reserved.
+#pragma once 
+
+#include "AudioInsightsDashboardFactory.h"
+#include "AudioInsightsTraceModule.h"
+#include "Framework/Docking/TabManager.h"
+#include "IAudioInsightsModule.h"
+#include "Templates/SharedPointer.h"
+#include "Views/DashboardViewFactory.h"
+#include "Widgets/Docking/SDockTab.h"
+
+
+namespace UE::Audio::Insights
+{
+	// Forward Declarations
+	class FMixerSourceTraceProvider;
+
+	class FAudioInsightsModule final : public IAudioInsightsModule
+	{
+		TSharedPtr<FDashboardFactory> DashboardFactory;
+		FTraceModule TraceModule;
+
+	public:
+		FAudioInsightsModule() = default;
+
+		virtual ~FAudioInsightsModule() = default;
+
+		virtual void StartupModule() override;
+		virtual void ShutdownModule() override;
+		virtual void RegisterDashboardViewFactory(TSharedRef<IDashboardViewFactory> InDashboardFactory) override;
+		virtual void UnregisterDashboardViewFactory(FName InName) override;
+		virtual ::Audio::FDeviceId GetDeviceId() const override;
+
+		FTraceModule& GetTraceModule();
+
+		TSharedRef<FDashboardFactory> GetDashboardFactory();
+		const TSharedRef<FDashboardFactory> GetDashboardFactory() const;
+
+		static FAudioInsightsModule& GetChecked();
+
+	private:
+		TSharedRef<SDockTab> CreateDashboardTabWidget(const FSpawnTabArgs& Args);
+		TSharedRef<SDockTab> CreateInsightsTabWidget(const FSpawnTabArgs& Args);
+		void RegisterMenus();
+	};
+} // namespace UE::Audio::Insights
diff --git a/Engine/Plugins/AudioInsights/Source/AudioInsights/Private/AudioInsightsStyle.h b/Engine/Plugins/AudioInsights/Source/AudioInsights/Private/AudioInsightsStyle.h
new file mode 100644
--- /dev/null
+++ b/Engine/Plugins/AudioInsights/Source/AudioInsights/Private/AudioInsightsStyle.h
@@ -0,0 +1,153 @@
+// Copyright Epic Games, Inc. All Rights Reserved.
+#pragma once
+
+#include "Misc/Paths.h"
+#include "Styling/AppStyle.h"
+#include "Styling/SlateStyle.h"
+#include "Styling/SlateStyleMacros.h"
+#include "Styling/SlateStyleRegistry.h"
+#include "Textures/SlateIcon.h"
+
+#define LOCTEXT_NAMESPACE "AudioInsights"
+
+namespace UE::Audio::Insights
+{
+	class FSlateStyle final : public FSlateStyleSet
+	{
+	public:
+		static FSlateStyle& Get()
+		{
+			static FSlateStyle InsightsStyle;
+			return InsightsStyle;
+		}
+
+		static FName GetStyleName()
+		{
+			const FLazyName StyleName = "AudioInsights";
+			return StyleName.Resolve();
+		}
+
+		const FNumberFormattingOptions* GetAmpFloatFormat()
+		{
+			static FNumberFormattingOptions FloatFormat;
+			FloatFormat.MinimumIntegralDigits = 1;
+			FloatFormat.MaximumIntegralDigits = 1;
+			FloatFormat.MinimumFractionalDigits = 3;
+			FloatFormat.MaximumFractionalDigits = 3;
+			return &FloatFormat;
+		};
+
+		const FNumberFormattingOptions* GetDefaultFloatFormat()
+		{
+			static FNumberFormattingOptions FloatFormat;
+			FloatFormat.MinimumIntegralDigits = 1;
+			FloatFormat.MinimumFractionalDigits = 4;
+			FloatFormat.MaximumFractionalDigits = 4;
+			return &FloatFormat;
+		};
+
+		const FNumberFormattingOptions* GetFreqFloatFormat()
+		{
+			static FNumberFormattingOptions FloatFormat;
+			FloatFormat.MinimumIntegralDigits = 1;
+			FloatFormat.MaximumIntegralDigits = 5;
+			FloatFormat.MinimumFractionalDigits = 0;
+			FloatFormat.MaximumFractionalDigits = 2;
+			return &FloatFormat;
+		};
+
+		const FNumberFormattingOptions* GetPitchFloatFormat()
+		{
+			static FNumberFormattingOptions FloatFormat;
+			FloatFormat.MinimumIntegralDigits = 1;
+			FloatFormat.MaximumIntegralDigits = 3;
+			FloatFormat.MinimumFractionalDigits = 3;
+			FloatFormat.MaximumFractionalDigits = 3;
+			return &FloatFormat;
+		};
+
+		const FNumberFormattingOptions* GetTimeFormat()
+		{
+			static FNumberFormattingOptions FloatFormat;
+			FloatFormat.MinimumIntegralDigits = 1;
+			FloatFormat.MinimumFractionalDigits = 3;
+			FloatFormat.MaximumFractionalDigits = 3;
+			return &FloatFormat;
+		};
+
+		FText FormatSecondsAsTime(float InTimeSec)
+		{
+			return FText::Format(LOCTEXT("TimeInSecondsFormat", "{0}s"), FText::AsNumber(InTimeSec, GetTimeFormat()));
+		}
+
+		FText FormatMillisecondsAsTime(float InTimeMS)
+		{
+			return FText::Format(LOCTEXT("TimeInMillisecondsFormat", "{0}ms"), FText::AsNumber(InTimeMS, GetTimeFormat()));
+		}
+
+		FSlateIcon CreateIcon(FName InName)
+		{
+			return { GetStyleName(), InName};
+		}
+
+		const FSlateBrush& GetBrushEnsured(FName InName)
+		{
+			const ISlateStyle* AudioInsightsStyle = FSlateStyleRegistry::FindSlateStyle(GetStyleName());
+			if (ensureMsgf(AudioInsightsStyle, TEXT("Missing slate style '%s'"), *GetStyleName().ToString()))
+			{
+				const FSlateBrush* Brush = AudioInsightsStyle->GetBrush(InName);
+				if (ensureMsgf(Brush, TEXT("Missing brush '%s'"), *InName.ToString()))
+				{
+					return *Brush;
+				}
+			}
+
+			if (const FSlateBrush* NoBrush = FAppStyle::GetBrush("NoBrush"))
+			{
+				return *NoBrush;
+			}
+
+			return *DefaultBrush;
+		}
+
+		FSlateStyle()
+			: FSlateStyleSet(GetStyleName())
+		{
+			SetParentStyleName(FAppStyle::GetAppStyleSetName());
+
+			const FString PluginsDir = FPaths::EnginePluginsDir();
+
+			SetContentRoot(PluginsDir / TEXT("AudioInsights/Content"));
+			SetCoreContentRoot(PluginsDir / TEXT("Slate"));
+
+			Set("AudioInsights.Analyzers.BackgroundColor", FLinearColor(0.0075f, 0.0075f, 0.0075f, 1.0f));
+
+			const FVector2D Icon16(16.0f, 16.0f);
+			const FVector2D Icon20(20.0f, 20.0f);
+			const FVector2D Icon40(40.0f, 40.0f);
+			const FVector2D Icon64(64.0f, 64.0f);
+
+			Set("AudioInsights.Icon", new IMAGE_BRUSH_SVG(TEXT("Icons/audio_insights_icon"), Icon16));
+			Set("AudioInsights.Icon.Dashboard", new IMAGE_BRUSH_SVG(TEXT("Icons/audio_dashboard"), Icon16));
+			Set("AudioInsights.Icon.Event", new IMAGE_BRUSH_SVG(TEXT("Icons/audio_event"), Icon16));
+			Set("AudioInsights.Icon.Log", new IMAGE_BRUSH_SVG(TEXT("Icons/audio_log"), Icon16));
+			Set("AudioInsights.Icon.Sources", new IMAGE_BRUSH_SVG(TEXT("Icons/audio_sources"), Icon16));
+			Set("AudioInsights.Icon.Submix", new IMAGE_BRUSH_SVG(TEXT("Icons/audio_submix"), Icon16));
+			Set("AudioInsights.Icon.VirtualLoop", new IMAGE_BRUSH_SVG(TEXT("Icons/audio_virtualloop"), Icon16));
+			Set("AudioInsights.Icon.Viewport", new IMAGE_BRUSH_SVG(TEXT("Icons/viewport"), Icon16));
+			Set("AudioInsights.Icon.Open", new IMAGE_BRUSH_SVG(TEXT("Icons/open"), Icon20));
+			Set("AudioInsights.Icon.ContentBrowser", new IMAGE_BRUSH_SVG(TEXT("Icons/content_browser"), Icon20));
+			Set("AudioInsights.Icon.Start.Active", new IMAGE_BRUSH_SVG(TEXT("Icons/start_active"), Icon20));
+			Set("AudioInsights.Icon.Start.Inactive", new IMAGE_BRUSH_SVG(TEXT("Icons/start_inactive"), Icon20));
+			Set("AudioInsights.Icon.Record.Active", new IMAGE_BRUSH_SVG(TEXT("Icons/record_active"), Icon20));
+			Set("AudioInsights.Icon.Record.Inactive", new IMAGE_BRUSH_SVG(TEXT("Icons/record_inactive"), Icon20));
+			Set("AudioInsights.Icon.Stop.Active", new IMAGE_BRUSH_SVG(TEXT("Icons/stop_active"), Icon20));
+			Set("AudioInsights.Icon.Stop.Inactive", new IMAGE_BRUSH_SVG(TEXT("Icons/stop_inactive"), Icon20));
+
+			Set("AudioInsights.Thumbnail", new IMAGE_BRUSH_SVG(TEXT("Icons/audio_insights"), Icon16));
+
+			FSlateStyleRegistry::RegisterSlateStyle(*this);
+		}
+	};
+} // namespace UE::Audio::Insights
+#undef LOCTEXT_NAMESPACE
diff --git a/Engine/Plugins/AudioInsights/Source/AudioInsights/Private/AudioInsightsTraceChannelHandle.cpp b/Engine/Plugins/AudioInsights/Source/AudioInsights/Private/AudioInsightsTraceChannelHandle.cpp
new file mode 100644
--- /dev/null
+++ b/Engine/Plugins/AudioInsights/Source/AudioInsights/Private/AudioInsightsTraceChannelHandle.cpp
@@ -0,0 +1,115 @@
+// Copyright Epic Games, Inc. All Rights Reserved.
+#pragma once
+
+#include "AudioInsightsTraceChannelHandle.h"
+
+#include "Misc/ScopeLock.h"
+#include "Trace/Trace.h"
+
+
+namespace UE::Audio::Insights
+{
+	FTraceChannelHandle::FTraceChannelHandle(FString&& InName, TSharedRef<FTraceChannelManager> InManager)
+		: ChannelName(MoveTemp(InName))
+		, Manager(InManager)
+	{
+		Init();
+	}
+
+	FTraceChannelHandle::FTraceChannelHandle(const FTraceChannelHandle& Other)
+		: ChannelName(Other.ChannelName)
+		, Manager(Other.Manager)
+	{
+		Init();
+	}
+
+	FTraceChannelHandle::FTraceChannelHandle(FTraceChannelHandle&& Other)
+		: ChannelName(MoveTemp(Other.ChannelName))
+		, Manager(MoveTemp(Other.Manager))
+	{
+	}
+
+	FTraceChannelHandle::~FTraceChannelHandle()
+	{
+		Reset();
+	}
+
+	bool FTraceChannelHandle::operator==(const FTraceChannelHandle& Other) const
+	{
+		return ChannelName.Compare(Other.ChannelName, ESearchCase::IgnoreCase) == 0 && Manager == Other.Manager;
+	}
+
+	bool FTraceChannelHandle::operator!=(const FTraceChannelHandle& Other) const
+	{
+		return ChannelName.Compare(Other.ChannelName, ESearchCase::IgnoreCase) != 0 || Manager != Other.Manager;
+	}
+
+	FTraceChannelHandle& FTraceChannelHandle::operator=(const FTraceChannelHandle& Other)
+	{
+		if (*this != Other)
+		{
+			ChannelName = Other.ChannelName;
+			Manager = Other.Manager;
+			Init();
+		}
+		return *this;
+	}
+
+	FTraceChannelHandle& FTraceChannelHandle::operator=(FTraceChannelHandle&& Other)
+	{
+		if (*this == Other)
+		{
+			Reset();
+		}
+
+		ChannelName = MoveTemp(Other.ChannelName);
+		Manager = Other.Manager;
+		return *this;
+	}
+
+	void FTraceChannelHandle::Init()
+	{
+		if (IsValid())
+		{
+			FScopeLock Lock(&Manager->CritSect);
+
+			if (uint32* RefCount = Manager->ActiveChannels.Find(ChannelName))
+			{
+				(*RefCount)++;
+			}
+			else
+			{
+				Manager->ActiveChannels.Add(ChannelName, 1);
+				constexpr bool bIsEnabled = true;
+				Trace::ToggleChannel(*ChannelName, bIsEnabled);
+			}
+		}
+	}
+
+	void FTraceChannelHandle::Reset()
+	{
+		if (IsValid())
+		{
+			FScopeLock Lock(&Manager->CritSect);
+
+			uint32& RefCount = Manager->ActiveChannels.FindChecked(ChannelName);
+			RefCount--;
+			if (RefCount == 0)
+			{
+				constexpr bool bIsEnabled = false;
+				Trace::ToggleChannel(*ChannelName, bIsEnabled);
+				Manager->ActiveChannels.Remove(ChannelName);
+			}
+		}
+	}
+
+	bool FTraceChannelHandle::IsValid() const
+	{
+		return !ChannelName.IsEmpty();
+	}
+
+	FTraceChannelHandle FTraceChannelManager::CreateHandle(FString&& InName)
+	{
+		return FTraceChannelHandle(MoveTemp(InName), AsShared());
+	}
+} // namespace UE::Audio::Insights
\ No newline at end of file
diff --git a/Engine/Plugins/AudioInsights/Source/AudioInsights/Private/AudioInsightsTraceModule.cpp b/Engine/Plugins/AudioInsights/Source/AudioInsights/Private/AudioInsightsTraceModule.cpp
new file mode 100644
--- /dev/null
+++ b/Engine/Plugins/AudioInsights/Source/AudioInsights/Private/AudioInsightsTraceModule.cpp
@@ -0,0 +1,82 @@
+// Copyright Epic Games, Inc. All Rights Reserved.
+#include "AudioInsightsTraceModule.h"
+
+#include "Insights/IUnrealInsightsModule.h"
+#include "ProfilingDebugging/TraceAuxiliary.h"
+#include "Providers/MixerSourceTraceProvider.h"
+#include "Providers/VirtualLoopTraceProvider.h"
+#include "Templates/SharedPointer.h"
+
+
+namespace UE::Audio::Insights
+{
+	FTraceModule::FTraceModule()
+		: ChannelManager(MakeShared<FTraceChannelManager>())
+	{
+		TSharedPtr<FMixerSourceTraceProvider> SourceProvider = MakeShared<FMixerSourceTraceProvider>(ChannelManager);
+		TSharedPtr<FVirtualLoopTraceProvider> VirtualLoopProvider = MakeShared<FVirtualLoopTraceProvider>(ChannelManager);
+
+		TraceProviders.Add(SourceProvider->GetName(), StaticCastSharedPtr<FTraceProviderBase>(SourceProvider));
+		TraceProviders.Add(VirtualLoopProvider->GetName(), StaticCastSharedPtr<FTraceProviderBase>(VirtualLoopProvider));
+	}
+
+	void FTraceModule::GetModuleInfo(TraceServices::FModuleInfo& OutModuleInfo)
+	{
+		OutModuleInfo.Name = GetName();
+		OutModuleInfo.DisplayName = TEXT("Audio");
+	}
+
+	TSharedRef<FTraceChannelManager> FTraceModule::GetChannelManager()
+	{
+		return ChannelManager;
+	}
+
+	const FName FTraceModule::GetName()
+	{
+		const FLazyName TraceName = { "TraceModule_AudioTrace" };
+		return TraceName.Resolve();
+	}
+
+	void FTraceModule::OnAnalysisBegin(TraceServices::IAnalysisSession& InSession)
+	{
+		for (const TPair<FName, TSharedPtr<FTraceProviderBase>>& Pair : TraceProviders)
+		{
+			InSession.AddProvider(Pair.Key, nullptr, Pair.Value);
+			InSession.AddAnalyzer(Pair.Value->ConstructAnalyzer());
+		}
+	}
+
+	void FTraceModule::StartTraceAnalysis() const
+	{
+		if (!FTraceAuxiliary::IsConnected())
+		{
+			FTraceAuxiliary::Start(FTraceAuxiliary::EConnectionType::Network, TEXT("localhost"), nullptr);
+
+			IUnrealInsightsModule& UnrealInsightsModule = FModuleManager::LoadModuleChecked<IUnrealInsightsModule>("TraceInsights");
+			UnrealInsightsModule.StartAnalysisForLastLiveSession();
+		}
+	}
+
+	void FTraceModule::StopTraceAnalysis() const
+	{
+		if (FTraceAuxiliary::IsConnected())
+		{
+			FTraceAuxiliary::Stop();
+		}
+	}
+
+	bool FTraceModule::IsTraceAnalysisActive() const
+	{
+		return FTraceAuxiliary::IsConnected();
+	}
+
+	void FTraceModule::GetLoggers(TArray<const TCHAR *>& OutLoggers)
+	{
+		OutLoggers.Add(TEXT("Audio"));
+	}
+
+	void FTraceModule::GenerateReports(const TraceServices::IAnalysisSession& Session, const TCHAR* CmdLine, const TCHAR* OutputDirectory)
+	{
+
+	}
+} // namespace UE::Audio::Insights
diff --git a/Engine/Plugins/AudioInsights/Source/AudioInsights/Private/AudioInsightsTraceModule.h b/Engine/Plugins/AudioInsights/Source/AudioInsights/Private/AudioInsightsTraceModule.h
new file mode 100644
--- /dev/null
+++ b/Engine/Plugins/AudioInsights/Source/AudioInsights/Private/AudioInsightsTraceModule.h
@@ -0,0 +1,46 @@
+// Copyright Epic Games, Inc. All Rights Reserved.
+#pragma once
+
+#include "AudioInsightsTraceChannelHandle.h"
+#include "AudioInsightsTraceProviderBase.h"
+#include "Containers/Map.h"
+#include "Templates/SharedPointer.h"
+#include "TraceServices/Model/AnalysisSession.h"
+#include "TraceServices/ModuleService.h"
+#include "UObject/NameTypes.h"
+
+namespace UE::Audio::Insights
+{
+	class FTraceModule : public TraceServices::IModule
+	{
+	public:
+		FTraceModule();
+		virtual ~FTraceModule() = default;
+
+		//~ Begin TraceServices::IModule interface
+		virtual void GetModuleInfo(TraceServices::FModuleInfo& OutModuleInfo) override;
+		virtual void OnAnalysisBegin(TraceServices::IAnalysisSession& Session) override;
+		virtual void GetLoggers(TArray<const TCHAR *>& OutLoggers) override;
+		virtual void GenerateReports(const TraceServices::IAnalysisSession& Session, const TCHAR* CmdLine, const TCHAR* OutputDirectory) override;
+		virtual const TCHAR* GetCommandLineArgument() override { return TEXT("audiotrace"); }
+		//~ End TraceServices::IModule interface
+
+		template <typename TraceProviderType>
+		TSharedPtr<TraceProviderType> FindAudioTraceProvider() const
+		{
+			return StaticCastSharedPtr<TraceProviderType>(TraceProviders.FindRef(TraceProviderType::GetName_Static()));
+		}
+
+		TSharedRef<FTraceChannelManager> GetChannelManager();
+
+		void StartTraceAnalysis() const;
+		bool IsTraceAnalysisActive() const;
+		void StopTraceAnalysis() const;
+
+	private:
+		static const FName GetName();
+
+		TSharedRef<FTraceChannelManager> ChannelManager;
+		TMap<FName, TSharedPtr<FTraceProviderBase>> TraceProviders;
+	};
+} // namespace UE::Audio::Insights
diff --git a/Engine/Plugins/AudioInsights/Source/AudioInsights/Private/AudioInsightsTraceProviderBase.cpp b/Engine/Plugins/AudioInsights/Source/AudioInsights/Private/AudioInsightsTraceProviderBase.cpp
new file mode 100644
--- /dev/null
+++ b/Engine/Plugins/AudioInsights/Source/AudioInsights/Private/AudioInsightsTraceProviderBase.cpp
@@ -0,0 +1,43 @@
+// Copyright Epic Games, Inc. All Rights Reserved.
+
+#include "AudioInsightsTraceProviderBase.h"
+
+
+namespace UE::Audio::Insights
+{
+	FTraceProviderBase::FTraceProviderBase(FName InName)
+		: Name(InName)
+	{
+	}
+
+	FName FTraceProviderBase::GetName() const
+	{
+		return Name;
+	}
+
+	FTraceProviderBase::FTraceAnalyzerBase::FTraceAnalyzerBase(TSharedRef<FTraceProviderBase> InProvider)
+		: Provider(InProvider)
+	{
+	}
+
+	void FTraceProviderBase::FTraceAnalyzerBase::OnAnalysisBegin(const FOnAnalysisContext& Context)
+	{
+		Provider->Reset();
+	}
+
+	bool FTraceProviderBase::FTraceAnalyzerBase::OnEventSuccess(uint16 RouteId, EStyle Style, const FOnEventContext& Context)
+	{
+		++(Provider->LastMessageId);
+		return true;
+	}
+
+	bool FTraceProviderBase::FTraceAnalyzerBase::OnEventFailure(uint16 RouteId, EStyle Style, const FOnEventContext& Context)
+	{
+		++(Provider->LastMessageId);
+
+		const FString ProviderName = Provider->GetName().ToString();
+		constexpr bool bEventSucceeded = false;
+		ensureAlwaysMsgf(bEventSucceeded, TEXT("'%s' TraceProvider's Analyzer message with RouteId '%u' event not handled"), *ProviderName, RouteId);
+		return bEventSucceeded;
+	}
+} // namespace UE::Audio::Insights
diff --git a/Engine/Plugins/AudioInsights/Source/AudioInsights/Private/Messages/MixerSourceTraceMessages.h b/Engine/Plugins/AudioInsights/Source/AudioInsights/Private/Messages/MixerSourceTraceMessages.h
new file mode 100644
--- /dev/null
+++ b/Engine/Plugins/AudioInsights/Source/AudioInsights/Private/Messages/MixerSourceTraceMessages.h
@@ -0,0 +1,97 @@
+// Copyright Epic Games, Inc. All Rights Reserved.
+#pragma once
+
+#include "AudioDefines.h"
+#include "Messages/AnalyzerMessageQueue.h"
+#include "Trace/Analyzer.h"
+#include "Views//TableDashboardViewFactory.h"
+
+
+namespace UE::Audio::Insights
+{
+	struct FMixerSourceMessageBase
+	{
+		FMixerSourceMessageBase() = default;
+		FMixerSourceMessageBase(const Trace::IAnalyzer::FOnEventContext& InContext)
+		{
+			const Trace::IAnalyzer::FEventData& EventData = InContext.EventData;
+			DeviceId = static_cast<::Audio::FDeviceId>(EventData.GetValue<uint32>("DeviceId"));
+			Timestamp = InContext.EventTime.AsSeconds(EventData.GetValue<uint64>("Timestamp"));
+			PlayOrder = EventData.GetValue<uint32>("PlayOrder");
+		}
+
+		::Audio::FDeviceId DeviceId = INDEX_NONE;
+		uint32 PlayOrder = INDEX_NONE;
+		uint64 ComponentId = TNumericLimits<uint64>::Max();
+		double Timestamp = 0.0;
+	};
+
+	using FMixerSourceStopMessage = FMixerSourceMessageBase;
+
+	struct FMixerSourceStartMessage : public FMixerSourceMessageBase
+	{
+		FMixerSourceStartMessage() = default;
+		FMixerSourceStartMessage(const Trace::IAnalyzer::FOnEventContext& InContext)
+			: FMixerSourceMessageBase(InContext)
+		{
+			const Trace::IAnalyzer::FEventData& EventData = InContext.EventData;
+			EventData.GetString("Name", Name);
+			ComponentId = EventData.GetValue<uint64>("ComponentId");
+			SourceId = EventData.GetValue<int32>("SourceId");
+		}
+
+		FString Name;
+		int32 SourceId = INDEX_NONE;
+		uint64 ComponentId = TNumericLimits<uint64>::Max();
+	};
+
+#define DEFINE_MIXERSOURCE_PARAM_MESSAGE(ClassName, ParamName, Type, Default)			\
+	struct ClassName : public FMixerSourceMessageBase									\
+	{																					\
+		ClassName() = default;															\
+		ClassName(const Trace::IAnalyzer::FOnEventContext& InContext)					\
+			: FMixerSourceMessageBase(InContext)										\
+		{																				\
+			const Trace::IAnalyzer::FEventData& EventData = InContext.EventData;		\
+			ParamName = EventData.GetValue<Type>(#ParamName);							\
+		}																				\
+		Type ParamName = Default;														\
+	};
+
+	DEFINE_MIXERSOURCE_PARAM_MESSAGE(FMixerSourceDistanceAttenuationMessage, DistanceAttenuation, float, 0.0f)
+	DEFINE_MIXERSOURCE_PARAM_MESSAGE(FMixerSourceEnvelopeMessage, Envelope, float, 0.0f)
+	DEFINE_MIXERSOURCE_PARAM_MESSAGE(FMixerSourceHPFFreqMessage, HPFFrequency, float, MIN_FILTER_FREQUENCY)
+	DEFINE_MIXERSOURCE_PARAM_MESSAGE(FMixerSourceLPFFreqMessage, LPFFrequency, float, MAX_FILTER_FREQUENCY)
+	DEFINE_MIXERSOURCE_PARAM_MESSAGE(FMixerSourcePitchMessage, Pitch, float, 1.0f)
+	DEFINE_MIXERSOURCE_PARAM_MESSAGE(FMixerSourceVolumeMessage, Volume, float, 1.0f)
+#undef DEFINE_MIXERSOURCE_PARAM_MESSAGE
+
+	class FMixerSourceDashboardEntry : public FSoundAssetDashboardEntry
+	{
+	public:
+		FMixerSourceDashboardEntry() = default;
+		virtual ~FMixerSourceDashboardEntry() = default;
+
+		int32 SourceId = INDEX_NONE;
+
+		float Volume = 1.0f;
+		float Pitch = 1.0f;
+		float LPFFreq = MIN_FILTER_FREQUENCY;
+		float HPFFreq = MAX_FILTER_FREQUENCY;
+		float Envelope = 0.0f;
+		float DistanceAttenuation = 0.0f;
+	};
+
+	class FMixerSourceMessages
+	{
+	public:
+		TAnalyzerMessageQueue<FMixerSourceDistanceAttenuationMessage> DistanceAttenuationMessages { 0.1 };
+		TAnalyzerMessageQueue<FMixerSourceEnvelopeMessage> EnvelopeMessages { 0.1 };
+		TAnalyzerMessageQueue<FMixerSourceHPFFreqMessage> HPFFreqMessages { 0.1 };
+		TAnalyzerMessageQueue<FMixerSourceLPFFreqMessage> LPFFreqMessages { 0.1 };
+		TAnalyzerMessageQueue<FMixerSourcePitchMessage> PitchMessages { 0.1 };
+		TAnalyzerMessageQueue<FMixerSourceStartMessage> StartMessages { 0.1 };
+		TAnalyzerMessageQueue<FMixerSourceStopMessage> StopMessages { 0.1 };
+		TAnalyzerMessageQueue<FMixerSourceVolumeMessage> VolumeMessages { 0.1 };
+	};
+} // namespace UE::Audio::Insights
diff --git a/Engine/Plugins/AudioInsights/Source/AudioInsights/Private/Messages/VirtualLoopTraceMessages.h b/Engine/Plugins/AudioInsights/Source/AudioInsights/Private/Messages/VirtualLoopTraceMessages.h
new file mode 100644
--- /dev/null
+++ b/Engine/Plugins/AudioInsights/Source/AudioInsights/Private/Messages/VirtualLoopTraceMessages.h
@@ -0,0 +1,111 @@
+// Copyright Epic Games, Inc. All Rights Reserved.
+#pragma once
+
+#include "AudioDefines.h"
+#include "Math/NumericLimits.h"
+#include "Messages/AnalyzerMessageQueue.h"
+#include "Trace/Analyzer.h"
+#include "Views/TableDashboardViewFactory.h"
+
+
+namespace UE::Audio::Insights
+{
+	struct FVirtualLoopMessageBase
+	{
+		FVirtualLoopMessageBase() = default;
+		FVirtualLoopMessageBase(const Trace::IAnalyzer::FOnEventContext& InContext)
+		{
+			const Trace::IAnalyzer::FEventData& EventData = InContext.EventData;
+			DeviceId = static_cast<::Audio::FDeviceId>(EventData.GetValue<uint32>("DeviceId"));
+			Timestamp = InContext.EventTime.AsSeconds(EventData.GetValue<uint64>("Timestamp"));
+			PlayOrder = EventData.GetValue<uint32>("PlayOrder");
+		}
+
+		::Audio::FDeviceId DeviceId = INDEX_NONE;
+		uint32 PlayOrder = INDEX_NONE;
+		double Timestamp = 0.0;
+	};
+
+	using FVirtualLoopRealizeMessage = FVirtualLoopMessageBase;
+	using FVirtualLoopStopMessage = FVirtualLoopMessageBase;
+
+	struct FVirtualLoopVirtualizeMessage : public FVirtualLoopMessageBase
+	{
+		FVirtualLoopVirtualizeMessage() = default;
+		FVirtualLoopVirtualizeMessage(const Trace::IAnalyzer::FOnEventContext& InContext)
+			: FVirtualLoopMessageBase(InContext)
+		{
+			const Trace::IAnalyzer::FEventData& EventData = InContext.EventData;
+			EventData.GetString("Name", Name);
+			ComponentId = EventData.GetValue<uint64>("ComponentId");
+		}
+
+		FString Name;
+		uint64 ComponentId = TNumericLimits<uint64>::Max();
+	};
+
+	struct FVirtualLoopUpdateMessage : public FVirtualLoopMessageBase
+	{
+		FVirtualLoopUpdateMessage() = default;
+		FVirtualLoopUpdateMessage(const Trace::IAnalyzer::FOnEventContext& InContext)
+			: FVirtualLoopMessageBase(InContext)
+		{
+			const Trace::IAnalyzer::FEventData& EventData = InContext.EventData;
+			TimeVirtualized = EventData.GetValue<float>("TimeVirtualized");
+			PlaybackTime = EventData.GetValue<float>("PlaybackTime");
+			UpdateInterval = EventData.GetValue<float>("UpdateInterval");
+
+			LocationX = EventData.GetValue<double>("LocationX");
+			LocationY = EventData.GetValue<double>("LocationY");
+			LocationZ = EventData.GetValue<double>("LocationZ");
+
+			RotatorPitch = EventData.GetValue<double>("RotatorPitch");
+			RotatorYaw = EventData.GetValue<double>("RotatorYaw");
+			RotatorRoll = EventData.GetValue<double>("RotatorRoll");
+		}
+
+		float TimeVirtualized = 0.0f;
+		float PlaybackTime = 0.0f;
+		float UpdateInterval = 0.0f;
+
+		double LocationX = 0.0;
+		double LocationY = 0.0;
+		double LocationZ = 0.0;
+
+		double RotatorPitch = 0.0;
+		double RotatorYaw = 0.0;
+		double RotatorRoll = 0.0;
+	};
+
+	class FVirtualLoopDashboardEntry : public FSoundAssetDashboardEntry
+	{
+	public:
+		FVirtualLoopDashboardEntry() = default;
+		virtual ~FVirtualLoopDashboardEntry() = default;
+
+		virtual bool IsValid() const override
+		{
+			return PlayOrder != static_cast<uint32>(INDEX_NONE);
+		}
+
+		uint32 PlayOrder = static_cast<uint32>(INDEX_NONE);
+		uint64 ComponentId = TNumericLimits<uint64>::Max();
+
+		float TimeVirtualized = 0.0f;
+		float PlaybackTime = 0.0f;
+		float UpdateInterval = 0.0f;
+
+		FVector Location = FVector::ZeroVector;
+		FRotator Rotator = FRotator::ZeroRotator;
+	};
+
+	class FVirtualLoopMessages
+	{
+		TAnalyzerMessageQueue<FVirtualLoopVirtualizeMessage> VirtualizeMessages { 5.0 };
+		TAnalyzerMessageQueue<FVirtualLoopRealizeMessage> RealizeMessages { 5.0 };
+		TAnalyzerMessageQueue<FVirtualLoopStopMessage> StopMessages { 5.0 };
+		TAnalyzerMessageQueue<FVirtualLoopUpdateMessage> UpdateMessages { 0.1 };
+
+		friend class FVirtualLoopTraceProvider;
+	};
+} // namespace UE::Audio::Insights
diff --git a/Engine/Plugins/AudioInsights/Source/AudioInsights/Private/Providers/AudioBusProvider.cpp b/Engine/Plugins/AudioInsights/Source/AudioInsights/Private/Providers/AudioBusProvider.cpp
new file mode 100644
--- /dev/null
+++ b/Engine/Plugins/AudioInsights/Source/AudioInsights/Private/Providers/AudioBusProvider.cpp
@@ -0,0 +1,166 @@
+// Copyright Epic Games, Inc. All Rights Reserved.
+
+#include "Providers/AudioBusProvider.h"
+
+#include "AssetRegistry/AssetData.h"
+#include "AssetRegistry/AssetRegistryModule.h"
+#include "AudioInsightsDashboardFactory.h"
+#include "Sound/AudioBus.h"
+
+namespace UE::Audio::Insights
+{
+	FAudioBusProvider::FAudioBusProvider()
+		: TDeviceDataMapTraceProvider<uint32, TSharedPtr<FAudioBusAssetDashboardEntry>>(GetName_Static())
+	{
+		FAssetRegistryModule& AssetRegistryModule = FModuleManager::LoadModuleChecked<FAssetRegistryModule>("AssetRegistry");
+
+		AssetRegistryModule.Get().OnAssetAdded().AddRaw(this, &FAudioBusProvider::OnAssetAdded);
+		AssetRegistryModule.Get().OnAssetRemoved().AddRaw(this, &FAudioBusProvider::OnAssetRemoved);
+		AssetRegistryModule.Get().OnFilesLoaded().AddRaw(this, &FAudioBusProvider::OnFilesLoaded);
+
+		FDashboardFactory::OnActiveAudioDeviceChanged.AddRaw(this, &FAudioBusProvider::OnActiveAudioDeviceChanged);
+	}
+	
+	FAudioBusProvider::~FAudioBusProvider()
+	{
+		if (FAssetRegistryModule* AssetRegistryModule = FModuleManager::GetModulePtr<FAssetRegistryModule>("AssetRegistry"))
+		{
+			AssetRegistryModule->Get().OnAssetAdded().RemoveAll(this);
+			AssetRegistryModule->Get().OnAssetRemoved().RemoveAll(this);
+			AssetRegistryModule->Get().OnFilesLoaded().RemoveAll(this);
+		}
+
+		FDashboardFactory::OnActiveAudioDeviceChanged.RemoveAll(this);
+	}
+
+	FName FAudioBusProvider::GetName_Static()
+	{
+		return "AudioBusProvider";
+	}
+
+	void FAudioBusProvider::OnAssetAdded(const FAssetData& InAssetData)
+	{
+		if (bAreFilesLoaded && InAssetData.AssetClassPath == FTopLevelAssetPath(UAudioBus::StaticClass()))
+		{
+			AddAudioBusAsset(InAssetData);
+		}
+	}
+
+	void FAudioBusProvider::OnAssetRemoved(const FAssetData& InAssetData)
+	{
+		if (InAssetData.AssetClassPath == FTopLevelAssetPath(UAudioBus::StaticClass()))
+		{
+			RemoveAudioBusAsset(InAssetData);
+		}
+	}
+
+	void FAudioBusProvider::OnFilesLoaded()
+	{
+		bAreFilesLoaded = true;
+		UpdateAudioBusAssetNames();
+	}
+
+	void FAudioBusProvider::OnActiveAudioDeviceChanged()
+	{
+		UpdateAudioBusAssetNames();
+	}
+
+	void FAudioBusProvider::AddAudioBusAsset(const FAssetData& InAssetData)
+	{
+		const bool bIsAudioBusAssetAlreadAdded = AudioBusDataViewEntries.ContainsByPredicate(
+			[AssetName = InAssetData.GetObjectPathString()](const TSharedPtr<FAudioBusAssetDashboardEntry>& AudioBusAssetDashboardEntry)
+			{
+				return AudioBusAssetDashboardEntry->Name == AssetName;
+			});
+
+		if (!bIsAudioBusAssetAlreadAdded)
+		{
+			const IAudioInsightsModule& InsightsModule = FModuleManager::GetModuleChecked<IAudioInsightsModule>(IAudioInsightsModule::GetName());
+			const ::Audio::FDeviceId DeviceId = InsightsModule.GetDeviceId();
+
+			TSharedPtr<FAudioBusAssetDashboardEntry> AudioBusAssetDashboardEntry = MakeShared<FAudioBusAssetDashboardEntry>();
+			AudioBusAssetDashboardEntry->DeviceId = DeviceId;
+			AudioBusAssetDashboardEntry->Name     = InAssetData.GetObjectPathString();
+			AudioBusAssetDashboardEntry->AudioBus = Cast<UAudioBus>(InAssetData.GetAsset());
+
+			AudioBusDataViewEntries.Add(MoveTemp(AudioBusAssetDashboardEntry));
+
+			OnAudioBusAssetAdded.Broadcast(InAssetData.GetAsset());
+			++LastUpdateId;
+		}
+	}
+
+	void FAudioBusProvider::RemoveAudioBusAsset(const FAssetData& InAssetData)
+	{
+		const int32 FoundAudioBusAssetNameIndex = AudioBusDataViewEntries.IndexOfByPredicate(
+			[AssetName = InAssetData.GetObjectPathString()](const TSharedPtr<FAudioBusAssetDashboardEntry>& AudioBusAssetDashboardEntry)
+			{
+				return AudioBusAssetDashboardEntry->Name == AssetName;
+			});
+
+		if (FoundAudioBusAssetNameIndex != INDEX_NONE)
+		{
+			const IAudioInsightsModule& InsightsModule = FModuleManager::GetModuleChecked<IAudioInsightsModule>(IAudioInsightsModule::GetName());
+			const ::Audio::FDeviceId DeviceId = InsightsModule.GetDeviceId();
+
+			RemoveDeviceEntry(DeviceId, AudioBusDataViewEntries[FoundAudioBusAssetNameIndex]->AudioBus->GetUniqueID());
+
+			AudioBusDataViewEntries.RemoveAt(FoundAudioBusAssetNameIndex);
+
+			OnAudioBusAssetRemoved.Broadcast(InAssetData.GetAsset());
+			++LastUpdateId;
+		}
+	}
+
+	void FAudioBusProvider::UpdateAudioBusAssetNames()
+	{
+		// Get all UAudioBus assets
+		TArray<FAssetData> AssetDataArray;
+
+		FAssetRegistryModule& AssetRegistryModule = FModuleManager::LoadModuleChecked<FAssetRegistryModule>("AssetRegistry");
+		AssetRegistryModule.Get().GetAssetsByClass(FTopLevelAssetPath(UAudioBus::StaticClass()), AssetDataArray);
+
+		// Build AudioBusDataViewEntries
+		Reset();
+		AudioBusDataViewEntries.Empty();
+
+		for (const FAssetData& AssetData : AssetDataArray)
+		{
+			AddAudioBusAsset(AssetData);
+		}
+
+		AudioBusDataViewEntries.Sort([](const TSharedPtr<FAudioBusAssetDashboardEntry>& A, const TSharedPtr<FAudioBusAssetDashboardEntry>& B)
+		{
+			return A->GetDisplayName().CompareToCaseIgnored(B->GetDisplayName()) < 0;
+		});
+
+		OnAudioBusAssetListUpdated.Broadcast();
+	}
+
+	bool FAudioBusProvider::ProcessMessages()
+	{
+		const IAudioInsightsModule& InsightsModule = FModuleManager::GetModuleChecked<IAudioInsightsModule>(IAudioInsightsModule::GetName());
+		const ::Audio::FDeviceId DeviceId = InsightsModule.GetDeviceId();
+
+		for (uint32 Index = 0;
+			 const TSharedPtr<FAudioBusAssetDashboardEntry>& AudioBusDataViewEntry : AudioBusDataViewEntries)
+		{
+			UpdateDeviceEntry(DeviceId, AudioBusDataViewEntry->AudioBus->GetUniqueID(), [&AudioBusDataViewEntry](TSharedPtr<FAudioBusAssetDashboardEntry>& Entry)
+			{
+				if (!Entry.IsValid())
+				{
+					Entry = AudioBusDataViewEntry;
+				}
+			});
+
+			++Index;
+		}
+
+		return true;
+	}
+
+	UE::Trace::IAnalyzer* FAudioBusProvider::ConstructAnalyzer()
+	{
+		return nullptr;
+	}
+} // namespace UE::Audio::Insights
diff --git a/Engine/Plugins/AudioInsights/Source/AudioInsights/Private/Providers/AudioBusProvider.h b/Engine/Plugins/AudioInsights/Source/AudioInsights/Private/Providers/AudioBusProvider.h
new file mode 100644
--- /dev/null
+++ b/Engine/Plugins/AudioInsights/Source/AudioInsights/Private/Providers/AudioBusProvider.h
@@ -0,0 +1,46 @@
+// Copyright Epic Games, Inc. All Rights Reserved.
+#pragma once
+
+#include "Delegates/Delegate.h"
+#include "UObject/NameTypes.h"
+#include "Views/AudioBusAssetDashboardEntry.h"
+
+namespace UE::Audio::Insights
+{
+	class FAudioBusProvider : public TDeviceDataMapTraceProvider<uint32, TSharedPtr<FAudioBusAssetDashboardEntry>>, public TSharedFromThis<FAudioBusProvider>
+	{
+	public:
+		FAudioBusProvider();
+		virtual ~FAudioBusProvider();
+
+		static FName GetName_Static();
+
+		virtual UE::Trace::IAnalyzer* ConstructAnalyzer() override;
+		
+		DECLARE_MULTICAST_DELEGATE_OneParam(FOnAudioBusAssetAdded, const TWeakObjectPtr<UObject> /*Asset*/);
+		inline static FOnAudioBusAssetAdded OnAudioBusAssetAdded;
+
+		DECLARE_MULTICAST_DELEGATE_OneParam(FOnAudioBusAssetRemoved, const TWeakObjectPtr<UObject> /*Asset*/);
+		inline static FOnAudioBusAssetRemoved OnAudioBusAssetRemoved;
+
+		DECLARE_MULTICAST_DELEGATE(FOnAudioBusAssetListUpdated);
+		inline static FOnAudioBusAssetListUpdated OnAudioBusAssetListUpdated;
+
+	private:
+		void OnAssetAdded(const FAssetData& InAssetData);
+		void OnAssetRemoved(const FAssetData& InAssetData);
+		void OnFilesLoaded();
+		void OnActiveAudioDeviceChanged();
+
+		void AddAudioBusAsset(const FAssetData& InAssetData);
+		void RemoveAudioBusAsset(const FAssetData& InAssetData);
+
+		void UpdateAudioBusAssetNames();
+
+		virtual bool ProcessMessages() override;
+
+		bool bAreFilesLoaded = false;
+
+		TArray<TSharedPtr<FAudioBusAssetDashboardEntry>> AudioBusDataViewEntries;
+	};
+} // namespace UE::Audio::Insights
diff --git a/Engine/Plugins/AudioInsights/Source/AudioInsights/Private/Providers/MixerSourceTraceProvider.cpp b/Engine/Plugins/AudioInsights/Source/AudioInsights/Private/Providers/MixerSourceTraceProvider.cpp
new file mode 100644
--- /dev/null
+++ b/Engine/Plugins/AudioInsights/Source/AudioInsights/Private/Providers/MixerSourceTraceProvider.cpp
@@ -0,0 +1,197 @@
+// Copyright Epic Games, Inc. All Rights Reserved.
+#include "Providers/MixerSourceTraceProvider.h"
+
+#include "Trace/Analyzer.h"
+#include "TraceServices/Model/AnalysisSession.h"
+#include "TraceServices/ModuleService.h"
+
+
+namespace UE::Audio::Insights
+{
+	FName FMixerSourceTraceProvider::GetName_Static()
+	{
+		return "MixerSourceProvider";
+	}
+
+	bool FMixerSourceTraceProvider::ProcessMessages()
+	{
+		auto BumpEntryFunc = [this](const FMixerSourceMessageBase& Msg)
+		{
+			TSharedPtr<FMixerSourceDashboardEntry>* ToReturn = nullptr;
+			UpdateDeviceEntry(Msg.DeviceId, Msg.PlayOrder, [&ToReturn, &Msg](TSharedPtr<FMixerSourceDashboardEntry>& Entry)
+			{
+				if (!Entry.IsValid())
+				{
+					Entry = MakeShared<FMixerSourceDashboardEntry>();
+					Entry->DeviceId = Msg.DeviceId;
+					Entry->PlayOrder = Msg.PlayOrder;
+				}
+				Entry->Timestamp = Msg.Timestamp;
+				ToReturn = &Entry;
+			});
+
+			return ToReturn;
+		};
+
+		ProcessMessageQueue<FMixerSourceStartMessage>(TraceMessages.StartMessages, BumpEntryFunc,
+		[](const FMixerSourceStartMessage& Msg, TSharedPtr<FMixerSourceDashboardEntry>* OutEntry)
+		{
+			FMixerSourceDashboardEntry& EntryRef = *OutEntry->Get();
+			EntryRef.Name = *Msg.Name;
+			EntryRef.ComponentId = Msg.ComponentId;
+			EntryRef.SourceId = Msg.SourceId;
+		});
+
+		ProcessMessageQueue<FMixerSourceVolumeMessage>(TraceMessages.VolumeMessages, BumpEntryFunc,
+		[](const FMixerSourceVolumeMessage& Msg, TSharedPtr<FMixerSourceDashboardEntry>* OutEntry)
+		{
+			(*OutEntry)->Volume = Msg.Volume;
+		});
+
+		ProcessMessageQueue<FMixerSourcePitchMessage>(TraceMessages.PitchMessages, BumpEntryFunc,
+		[](const FMixerSourcePitchMessage& Msg, TSharedPtr<FMixerSourceDashboardEntry>* OutEntry)
+		{
+			(*OutEntry)->Pitch = Msg.Pitch;
+		});
+
+		ProcessMessageQueue<FMixerSourceLPFFreqMessage>(TraceMessages.LPFFreqMessages, BumpEntryFunc,
+		[](const FMixerSourceLPFFreqMessage& Msg, TSharedPtr<FMixerSourceDashboardEntry>* OutEntry)
+		{
+			(*OutEntry)->LPFFreq = Msg.LPFFrequency;
+		});
+
+		ProcessMessageQueue<FMixerSourceHPFFreqMessage>(TraceMessages.HPFFreqMessages, BumpEntryFunc,
+		[](const FMixerSourceHPFFreqMessage& Msg, TSharedPtr<FMixerSourceDashboardEntry>* OutEntry)
+		{
+			(*OutEntry)->HPFFreq = Msg.HPFFrequency;
+		});
+
+		MaxEnvsMap.Reset();
+		ProcessMessageQueue<FMixerSourceEnvelopeMessage>(TraceMessages.EnvelopeMessages, BumpEntryFunc,
+		[this](const FMixerSourceEnvelopeMessage& Msg, TSharedPtr<FMixerSourceDashboardEntry>* OutEntry)
+		{
+			float& MaxEnv = MaxEnvsMap.FindOrAdd(Msg.PlayOrder);
+			MaxEnv = FMath::Max(MaxEnv, Msg.Envelope);
+			(*OutEntry)->Envelope = MaxEnv;
+		});
+
+		ProcessMessageQueue<FMixerSourceDistanceAttenuationMessage>(TraceMessages.DistanceAttenuationMessages, BumpEntryFunc,
+		[](const FMixerSourceDistanceAttenuationMessage& Msg, TSharedPtr<FMixerSourceDashboardEntry>* OutEntry)
+		{
+			(*OutEntry)->DistanceAttenuation = Msg.DistanceAttenuation;
+		});
+
+		auto GetEntry = [this](const FMixerSourceMessageBase& Msg)
+		{
+			return FindDeviceEntry(Msg.DeviceId, Msg.PlayOrder);
+		};
+
+		ProcessMessageQueue<FMixerSourceStopMessage>(TraceMessages.StopMessages, GetEntry,
+		[this](const FMixerSourceStopMessage& Msg, TSharedPtr<FMixerSourceDashboardEntry>* OutEntry)
+		{
+			if (OutEntry && (*OutEntry)->Timestamp < Msg.Timestamp)
+			{
+				RemoveDeviceEntry(Msg.DeviceId, Msg.PlayOrder);
+			}
+		});
+
+		return true;
+	}
+
+	UE::Trace::IAnalyzer* FMixerSourceTraceProvider::ConstructAnalyzer()
+	{
+		class FMixerSourceTraceAnalyzer : public FTraceAnalyzerBase
+		{
+		public:
+			FMixerSourceTraceAnalyzer(TSharedRef<FMixerSourceTraceProvider> InProvider)
+				: FTraceAnalyzerBase(InProvider)
+			{
+			}
+
+			virtual void OnAnalysisBegin(const FOnAnalysisContext& Context) override
+			{
+				UE::Trace::IAnalyzer::FInterfaceBuilder& Builder = Context.InterfaceBuilder;
+				Builder.RouteEvent(RouteId_DistanceAttenuation, "Audio", "MixerSourceDistanceAttenuation");
+				Builder.RouteEvent(RouteId_Envelope, "Audio", "MixerSourceEnvelope");
+				Builder.RouteEvent(RouteId_Filters, "Audio", "MixerSourceFilters");
+				Builder.RouteEvent(RouteId_Pitch, "Audio", "MixerSourcePitch");
+				Builder.RouteEvent(RouteId_Start, "Audio", "MixerSourceStart");
+				Builder.RouteEvent(RouteId_Stop, "Audio", "MixerSourceStop");
+				Builder.RouteEvent(RouteId_Volume, "Audio", "MixerSourceVolume");
+			}
+
+			virtual bool OnEvent(uint16 RouteId, EStyle Style, const FOnEventContext& Context) override
+			{
+				LLM_SCOPE_BYNAME(TEXT("Insights/FMixerSourceTraceAnalyzer"));
+
+				FMixerSourceMessages& Messages = GetProvider<FMixerSourceTraceProvider>().TraceMessages;
+				switch (RouteId)
+				{
+					case RouteId_Start:
+					{
+						Messages.StartMessages.Enqueue(FMixerSourceStartMessage { Context });
+						break;
+					}
+
+					case RouteId_Stop:
+					{
+						Messages.StopMessages.Enqueue(FMixerSourceStopMessage { Context });
+						break;
+					}
+
+					case RouteId_Volume:
+					{
+						Messages.VolumeMessages.Enqueue(FMixerSourceVolumeMessage { Context });
+						break;
+					}
+
+					case RouteId_Pitch:
+					{
+						Messages.PitchMessages.Enqueue(FMixerSourcePitchMessage { Context });
+						break;
+					}
+
+					case RouteId_Envelope:
+					{
+						Messages.EnvelopeMessages.Enqueue(FMixerSourceEnvelopeMessage { Context });
+						break;
+					}
+
+					case RouteId_Filters:
+					{
+						Messages.LPFFreqMessages.Enqueue(FMixerSourceLPFFreqMessage { Context });
+						Messages.HPFFreqMessages.Enqueue(FMixerSourceHPFFreqMessage { Context });
+						break;
+					}
+
+					case RouteId_DistanceAttenuation:
+					{
+						Messages.DistanceAttenuationMessages.Enqueue(FMixerSourceDistanceAttenuationMessage { Context });
+						break;
+					}
+
+					default:
+					{
+						return OnEventFailure(RouteId, Style, Context);
+					}
+				}
+
+				return OnEventSuccess(RouteId, Style, Context);
+			}
+
+		private:
+			enum : uint16
+			{
+				RouteId_DistanceAttenuation,
+				RouteId_Envelope,
+				RouteId_Filters,
+				RouteId_Pitch,
+				RouteId_Start,
+				RouteId_Stop,
+				RouteId_Volume,
+			};
+		};
+
+		return new FMixerSourceTraceAnalyzer(AsShared());
+	}
+} // namespace UE::Audio::Insights
diff --git a/Engine/Plugins/AudioInsights/Source/AudioInsights/Private/Providers/MixerSourceTraceProvider.h b/Engine/Plugins/AudioInsights/Source/AudioInsights/Private/Providers/MixerSourceTraceProvider.h
new file mode 100644
--- /dev/null
+++ b/Engine/Plugins/AudioInsights/Source/AudioInsights/Private/Providers/MixerSourceTraceProvider.h
@@ -0,0 +1,37 @@
+// Copyright Epic Games, Inc. All Rights Reserved.
+#pragma once
+
+#include "AudioInsightsTraceChannelHandle.h"
+#include "AudioInsightsTraceProviderBase.h"
+#include "Messages/MixerSourceTraceMessages.h"
+#include "UObject/NameTypes.h"
+
+
+namespace UE::Audio::Insights
+{
+	class FMixerSourceTraceProvider
+		: public TDeviceDataMapTraceProvider<uint32, TSharedPtr<FMixerSourceDashboardEntry>>
+		, public TSharedFromThis<FMixerSourceTraceProvider>
+	{
+	public:
+		FMixerSourceTraceProvider(TSharedRef<FTraceChannelManager> InManager)
+			: TDeviceDataMapTraceProvider<uint32, TSharedPtr<FMixerSourceDashboardEntry>>(GetName_Static())
+		{
+			Channels.Add(InManager->CreateHandle({ TEXT("AudioMixerChannel") }));
+		}
+
+		virtual ~FMixerSourceTraceProvider() = default;
+		virtual UE::Trace::IAnalyzer* ConstructAnalyzer() override;
+
+		virtual bool ProcessMessages() override;
+
+		static FName GetName_Static();
+
+	private:
+		TMap<int32, float> MaxEnvsMap;
+
+		FMixerSourceMessages TraceMessages;
+
+		TSet<FTraceChannelHandle> Channels;
+	};
+} // namespace UE::Audio::Insights
diff --git a/Engine/Plugins/AudioInsights/Source/AudioInsights/Private/Providers/SoundSubmixProvider.cpp b/Engine/Plugins/AudioInsights/Source/AudioInsights/Private/Providers/SoundSubmixProvider.cpp
new file mode 100644
--- /dev/null
+++ b/Engine/Plugins/AudioInsights/Source/AudioInsights/Private/Providers/SoundSubmixProvider.cpp
@@ -0,0 +1,166 @@
+// Copyright Epic Games, Inc. All Rights Reserved.
+
+#include "Providers/SoundSubmixProvider.h"
+
+#include "AssetRegistry/AssetData.h"
+#include "AssetRegistry/AssetRegistryModule.h"
+#include "AudioInsightsDashboardFactory.h"
+#include "Sound/SoundSubmix.h"
+
+namespace UE::Audio::Insights
+{
+	FSoundSubmixProvider::FSoundSubmixProvider()
+		: TDeviceDataMapTraceProvider<uint32, TSharedPtr<FSoundSubmixAssetDashboardEntry>>(GetName_Static())
+	{
+		FAssetRegistryModule& AssetRegistryModule = FModuleManager::LoadModuleChecked<FAssetRegistryModule>("AssetRegistry");
+
+		AssetRegistryModule.Get().OnAssetAdded().AddRaw(this, &FSoundSubmixProvider::OnAssetAdded);
+		AssetRegistryModule.Get().OnAssetRemoved().AddRaw(this, &FSoundSubmixProvider::OnAssetRemoved);
+		AssetRegistryModule.Get().OnFilesLoaded().AddRaw(this, &FSoundSubmixProvider::OnFilesLoaded);
+
+		FDashboardFactory::OnActiveAudioDeviceChanged.AddRaw(this, &FSoundSubmixProvider::OnActiveAudioDeviceChanged);
+	}
+	
+	FSoundSubmixProvider::~FSoundSubmixProvider()
+	{
+		if (FAssetRegistryModule* AssetRegistryModule = FModuleManager::GetModulePtr<FAssetRegistryModule>("AssetRegistry"))
+		{
+			AssetRegistryModule->Get().OnAssetAdded().RemoveAll(this);
+			AssetRegistryModule->Get().OnAssetRemoved().RemoveAll(this);
+			AssetRegistryModule->Get().OnFilesLoaded().RemoveAll(this);
+		}
+
+		FDashboardFactory::OnActiveAudioDeviceChanged.RemoveAll(this);
+	}
+
+	FName FSoundSubmixProvider::GetName_Static()
+	{
+		return "SubmixesProvider";
+	}
+
+	void FSoundSubmixProvider::OnAssetAdded(const FAssetData& InAssetData)
+	{
+		if (bAreFilesLoaded && InAssetData.AssetClassPath == FTopLevelAssetPath(USoundSubmix::StaticClass()))
+		{
+			AddSubmixAsset(InAssetData);
+		}
+	}
+
+	void FSoundSubmixProvider::OnAssetRemoved(const FAssetData& InAssetData)
+	{
+		if (InAssetData.AssetClassPath == FTopLevelAssetPath(USoundSubmix::StaticClass()))
+		{
+			RemoveSubmixAsset(InAssetData);
+		}
+	}
+
+	void FSoundSubmixProvider::OnFilesLoaded()
+	{
+		bAreFilesLoaded = true;
+		UpdateSubmixAssetNames();
+	}
+
+	void FSoundSubmixProvider::OnActiveAudioDeviceChanged()
+	{
+		UpdateSubmixAssetNames();
+	}
+
+	void FSoundSubmixProvider::AddSubmixAsset(const FAssetData& InAssetData)
+	{
+		const bool bIsSubmixAssetAlreadAdded = SubmixDataViewEntries.ContainsByPredicate(
+			[AssetName = InAssetData.GetObjectPathString()](const TSharedPtr<FSoundSubmixAssetDashboardEntry>& SoundSubmixAssetDashboardEntry)
+			{
+				return SoundSubmixAssetDashboardEntry->Name == AssetName;
+			});
+
+		if (!bIsSubmixAssetAlreadAdded)
+		{
+			const IAudioInsightsModule& InsightsModule = FModuleManager::GetModuleChecked<IAudioInsightsModule>(IAudioInsightsModule::GetName());
+			const ::Audio::FDeviceId DeviceId = InsightsModule.GetDeviceId();
+
+			TSharedPtr<FSoundSubmixAssetDashboardEntry> SoundSubmixAssetDashboardEntry = MakeShared<FSoundSubmixAssetDashboardEntry>();
+			SoundSubmixAssetDashboardEntry->DeviceId    = DeviceId;
+			SoundSubmixAssetDashboardEntry->Name        = InAssetData.GetObjectPathString();
+			SoundSubmixAssetDashboardEntry->SoundSubmix = Cast<USoundSubmix>(InAssetData.GetAsset());
+
+			SubmixDataViewEntries.Add(MoveTemp(SoundSubmixAssetDashboardEntry));
+
+			OnSubmixAssetAdded.Broadcast(InAssetData.GetAsset());
+			++LastUpdateId;
+		}
+	}
+
+	void FSoundSubmixProvider::RemoveSubmixAsset(const FAssetData& InAssetData)
+	{
+		const int32 FoundSubmixAssetNameIndex = SubmixDataViewEntries.IndexOfByPredicate(
+			[AssetName = InAssetData.GetObjectPathString()](const TSharedPtr<FSoundSubmixAssetDashboardEntry>& SoundSubmixAssetDashboardEntry)
+			{
+				return SoundSubmixAssetDashboardEntry->Name == AssetName;
+			});
+
+		if (FoundSubmixAssetNameIndex != INDEX_NONE)
+		{
+			const IAudioInsightsModule& InsightsModule = FModuleManager::GetModuleChecked<IAudioInsightsModule>(IAudioInsightsModule::GetName());
+			const ::Audio::FDeviceId DeviceId = InsightsModule.GetDeviceId();
+
+			RemoveDeviceEntry(DeviceId, SubmixDataViewEntries[FoundSubmixAssetNameIndex]->SoundSubmix->GetUniqueID());
+
+			SubmixDataViewEntries.RemoveAt(FoundSubmixAssetNameIndex);
+
+			OnSubmixAssetRemoved.Broadcast(InAssetData.GetAsset());
+			++LastUpdateId;
+		}
+	}
+
+	void FSoundSubmixProvider::UpdateSubmixAssetNames()
+	{
+		// Get all USoundSubmix assets
+		TArray<FAssetData> AssetDataArray;
+
+		FAssetRegistryModule& AssetRegistryModule = FModuleManager::LoadModuleChecked<FAssetRegistryModule>("AssetRegistry");
+		AssetRegistryModule.Get().GetAssetsByClass(FTopLevelAssetPath(USoundSubmix::StaticClass()), AssetDataArray);
+
+		// Build SubmixDataViewEntries
+		Reset();
+		SubmixDataViewEntries.Empty();
+
+		for (const FAssetData& AssetData : AssetDataArray)
+		{
+			AddSubmixAsset(AssetData);
+		}
+
+		SubmixDataViewEntries.Sort([](const TSharedPtr<FSoundSubmixAssetDashboardEntry>& A, const TSharedPtr<FSoundSubmixAssetDashboardEntry>& B)
+		{
+			return A->GetDisplayName().CompareToCaseIgnored(B->GetDisplayName()) < 0;
+		});
+
+		OnSubmixAssetListUpdated.Broadcast();
+	}
+
+	bool FSoundSubmixProvider::ProcessMessages()
+	{
+		const IAudioInsightsModule& InsightsModule = FModuleManager::GetModuleChecked<IAudioInsightsModule>(IAudioInsightsModule::GetName());
+		const ::Audio::FDeviceId DeviceId = InsightsModule.GetDeviceId();
+
+		for (uint32 Index = 0;
+			 const TSharedPtr<FSoundSubmixAssetDashboardEntry>& SubmixDataViewEntry : SubmixDataViewEntries)
+		{
+			UpdateDeviceEntry(DeviceId, SubmixDataViewEntry->SoundSubmix->GetUniqueID(), [&SubmixDataViewEntry](TSharedPtr<FSoundSubmixAssetDashboardEntry>& Entry)
+			{
+				if (!Entry.IsValid())
+				{
+					Entry = SubmixDataViewEntry;
+				}
+			});
+
+			++Index;
+		}
+
+		return true;
+	}
+
+	UE::Trace::IAnalyzer* FSoundSubmixProvider::ConstructAnalyzer()
+	{
+		return nullptr;
+	}
+} // namespace UE::Audio::Insights
diff --git a/Engine/Plugins/AudioInsights/Source/AudioInsights/Private/Providers/SoundSubmixProvider.h b/Engine/Plugins/AudioInsights/Source/AudioInsights/Private/Providers/SoundSubmixProvider.h
new file mode 100644
--- /dev/null
+++ b/Engine/Plugins/AudioInsights/Source/AudioInsights/Private/Providers/SoundSubmixProvider.h
@@ -0,0 +1,46 @@
+// Copyright Epic Games, Inc. All Rights Reserved.
+#pragma once
+
+#include "Delegates/Delegate.h"
+#include "UObject/NameTypes.h"
+#include "Views/SoundSubmixAssetDashboardEntry.h"
+
+namespace UE::Audio::Insights
+{
+	class FSoundSubmixProvider : public TDeviceDataMapTraceProvider<uint32, TSharedPtr<FSoundSubmixAssetDashboardEntry>>, public TSharedFromThis<FSoundSubmixProvider>
+	{
+	public:
+		FSoundSubmixProvider();
+		virtual ~FSoundSubmixProvider();
+
+		static FName GetName_Static();
+
+		virtual UE::Trace::IAnalyzer* ConstructAnalyzer() override;
+		
+		DECLARE_MULTICAST_DELEGATE_OneParam(FOnSubmixAssetAdded, const TWeakObjectPtr<UObject> /*Asset*/);
+		inline static FOnSubmixAssetAdded OnSubmixAssetAdded;
+
+		DECLARE_MULTICAST_DELEGATE_OneParam(FOnSubmixAssetRemoved, const TWeakObjectPtr<UObject> /*Asset*/);
+		inline static FOnSubmixAssetRemoved OnSubmixAssetRemoved;
+
+		DECLARE_MULTICAST_DELEGATE(FOnSubmixAssetListUpdated);
+		inline static FOnSubmixAssetListUpdated OnSubmixAssetListUpdated;
+
+	private:
+		void OnAssetAdded(const FAssetData& InAssetData);
+		void OnAssetRemoved(const FAssetData& InAssetData);
+		void OnFilesLoaded();
+		void OnActiveAudioDeviceChanged();
+
+		void AddSubmixAsset(const FAssetData& InAssetData);
+		void RemoveSubmixAsset(const FAssetData& InAssetData);
+
+		void UpdateSubmixAssetNames();
+
+		virtual bool ProcessMessages() override;
+
+		bool bAreFilesLoaded = false;
+
+		TArray<TSharedPtr<FSoundSubmixAssetDashboardEntry>> SubmixDataViewEntries;
+	};
+} // namespace UE::Audio::Insights
diff --git a/Engine/Plugins/AudioInsights/Source/AudioInsights/Private/Providers/VirtualLoopTraceProvider.cpp b/Engine/Plugins/AudioInsights/Source/AudioInsights/Private/Providers/VirtualLoopTraceProvider.cpp
new file mode 100644
--- /dev/null
+++ b/Engine/Plugins/AudioInsights/Source/AudioInsights/Private/Providers/VirtualLoopTraceProvider.cpp
@@ -0,0 +1,149 @@
+// Copyright Epic Games, Inc. All Rights Reserved.
+#include "Providers/VirtualLoopTraceProvider.h"
+
+#include "Trace/Analyzer.h"
+#include "TraceServices/Model/AnalysisSession.h"
+#include "TraceServices/ModuleService.h"
+
+
+namespace UE::Audio::Insights
+{
+	FName FVirtualLoopTraceProvider::GetName_Static()
+	{
+		return "AudioVirtualLoopProvider";
+	}
+
+	bool FVirtualLoopTraceProvider::ProcessMessages()
+	{
+		ProcessMessageQueue<FVirtualLoopVirtualizeMessage>(TraceMessages.VirtualizeMessages,
+		[this](const FVirtualLoopMessageBase& Msg)
+		{
+			TSharedPtr<FVirtualLoopDashboardEntry>* ToReturn = nullptr;
+			UpdateDeviceEntry(Msg.DeviceId, Msg.PlayOrder, [&ToReturn, &Msg](TSharedPtr<FVirtualLoopDashboardEntry>& Entry)
+			{
+				if (!Entry.IsValid())
+				{
+					Entry = MakeShared<FVirtualLoopDashboardEntry>();
+				}
+				Entry->DeviceId = Msg.DeviceId;
+				Entry->PlayOrder = Msg.PlayOrder;
+				Entry->Timestamp = Msg.Timestamp;
+
+				ToReturn = &Entry;
+			});
+			return ToReturn;
+		},
+		[](const FVirtualLoopVirtualizeMessage& Msg, TSharedPtr<FVirtualLoopDashboardEntry>* OutEntry)
+		{
+			FVirtualLoopDashboardEntry& EntryRef = *OutEntry->Get();
+			EntryRef.Name = Msg.Name;
+			EntryRef.ComponentId = Msg.ComponentId;
+		});
+
+		auto GetEntryFunc = [this](const FVirtualLoopMessageBase& Msg)
+		{
+			return FindDeviceEntry(Msg.DeviceId, Msg.PlayOrder);
+		};
+
+		ProcessMessageQueue<FVirtualLoopUpdateMessage>(TraceMessages.UpdateMessages, GetEntryFunc,
+		[](const FVirtualLoopUpdateMessage& Msg, TSharedPtr<FVirtualLoopDashboardEntry>* OutEntry)
+		{
+			if (OutEntry)
+			{
+				FVirtualLoopDashboardEntry& EntryRef = *OutEntry->Get();
+				EntryRef.PlaybackTime = Msg.PlaybackTime;
+				EntryRef.TimeVirtualized = Msg.TimeVirtualized;
+				EntryRef.UpdateInterval = Msg.UpdateInterval;
+				EntryRef.Location = FVector{ Msg.LocationX, Msg.LocationY, Msg.LocationZ };
+				EntryRef.Rotator = FRotator{ Msg.RotatorPitch, Msg.RotatorYaw, Msg.RotatorRoll };
+			}
+		});
+
+		auto RemoveEntryFunc = [this](const FVirtualLoopMessageBase& Msg, TSharedPtr<FVirtualLoopDashboardEntry>* OutEntry)
+		{
+			if (OutEntry && (*OutEntry)->Timestamp < Msg.Timestamp)
+			{
+				RemoveDeviceEntry(Msg.DeviceId, Msg.PlayOrder);
+			}
+		};
+
+		ProcessMessageQueue<FVirtualLoopRealizeMessage>(TraceMessages.RealizeMessages, GetEntryFunc, RemoveEntryFunc);
+		ProcessMessageQueue<FVirtualLoopStopMessage>(TraceMessages.StopMessages, GetEntryFunc, RemoveEntryFunc);
+
+		return true;
+	}
+
+	UE::Trace::IAnalyzer* FVirtualLoopTraceProvider::ConstructAnalyzer()
+	{
+		class FVirtualLoopTraceAnalyzer : public FTraceAnalyzerBase
+		{
+		public:
+			FVirtualLoopTraceAnalyzer(TSharedRef<FVirtualLoopTraceProvider> InProvider)
+				: FTraceAnalyzerBase(InProvider)
+			{
+			}
+
+			virtual void OnAnalysisBegin(const FOnAnalysisContext& Context) override
+			{
+				FTraceAnalyzerBase::OnAnalysisBegin(Context);
+
+				UE::Trace::IAnalyzer::FInterfaceBuilder& Builder = Context.InterfaceBuilder;
+				Builder.RouteEvent(RouteId_Realize, "Audio", "VirtualLoopRealize");
+				Builder.RouteEvent(RouteId_Stop, "Audio", "VirtualLoopStop");
+				Builder.RouteEvent(RouteId_Update, "Audio", "VirtualLoopUpdate");
+				Builder.RouteEvent(RouteId_Virtualize, "Audio", "VirtualLoopVirtualize");
+			}
+
+			virtual bool OnEvent(uint16 RouteId, EStyle Style, const FOnEventContext& Context) override
+			{
+				LLM_SCOPE_BYNAME(TEXT("Insights/FVirtualLoopTraceAnalyzer"));
+
+				FVirtualLoopMessages& Messages = GetProvider<FVirtualLoopTraceProvider>().TraceMessages;
+				switch (RouteId)
+				{
+					case RouteId_Realize:
+					{
+						Messages.RealizeMessages.Enqueue(FVirtualLoopRealizeMessage { Context });
+						break;
+					}
+
+					case RouteId_Stop:
+					{
+						Messages.StopMessages.Enqueue(FVirtualLoopStopMessage { Context });
+						break;
+					}
+
+					case RouteId_Update:
+					{
+						Messages.UpdateMessages.Enqueue(FVirtualLoopUpdateMessage { Context });
+						break;
+					}
+
+					case RouteId_Virtualize:
+					{
+						Messages.VirtualizeMessages.Enqueue(FVirtualLoopVirtualizeMessage { Context });
+						break;
+					}
+
+					default:
+					{
+						return OnEventFailure(RouteId, Style, Context);
+					}
+				}
+
+				return OnEventSuccess(RouteId, Style, Context);
+			}
+
+		private:
+			enum : uint16
+			{
+				RouteId_Virtualize,
+				RouteId_Update,
+				RouteId_Realize,
+				RouteId_Stop
+			};
+		};
+
+		return new FVirtualLoopTraceAnalyzer(AsShared());
+	}
+} // namespace UE::Audio::Insights
diff --git a/Engine/Plugins/AudioInsights/Source/AudioInsights/Private/Providers/VirtualLoopTraceProvider.h b/Engine/Plugins/AudioInsights/Source/AudioInsights/Private/Providers/VirtualLoopTraceProvider.h
new file mode 100644
--- /dev/null
+++ b/Engine/Plugins/AudioInsights/Source/AudioInsights/Private/Providers/VirtualLoopTraceProvider.h
@@ -0,0 +1,33 @@
+// Copyright Epic Games, Inc. All Rights Reserved.
+#pragma once
+
+#include "AudioInsightsTraceChannelHandle.h"
+#include "AudioInsightsTraceProviderBase.h"
+#include "Messages/VirtualLoopTraceMessages.h"
+
+
+namespace UE::Audio::Insights
+{
+	class FVirtualLoopTraceProvider
+		: public TDeviceDataMapTraceProvider<uint32, TSharedPtr<FVirtualLoopDashboardEntry>>
+		, public TSharedFromThis<FVirtualLoopTraceProvider>
+	{
+	public:
+		FVirtualLoopTraceProvider(TSharedRef<FTraceChannelManager> InChannelManager)
+			: TDeviceDataMapTraceProvider<uint32, TSharedPtr<FVirtualLoopDashboardEntry>>(GetName_Static())
+		{
+			Channels.Add(InChannelManager->CreateHandle(TEXT("AudioChannel")));
+		}
+
+		virtual ~FVirtualLoopTraceProvider() = default;
+		virtual UE::Trace::IAnalyzer* ConstructAnalyzer() override;
+
+		static FName GetName_Static();
+
+	private:
+		virtual bool ProcessMessages() override;
+
+		FVirtualLoopMessages TraceMessages;
+		TSet<FTraceChannelHandle> Channels;
+	};
+} // namespace UE::Audio::Insights
diff --git a/Engine/Plugins/AudioInsights/Source/AudioInsights/Private/Views/AudioBusAssetDashboardEntry.cpp b/Engine/Plugins/AudioInsights/Source/AudioInsights/Private/Views/AudioBusAssetDashboardEntry.cpp
new file mode 100644
--- /dev/null
+++ b/Engine/Plugins/AudioInsights/Source/AudioInsights/Private/Views/AudioBusAssetDashboardEntry.cpp
@@ -0,0 +1,26 @@
+// Copyright Epic Games, Inc. All Rights Reserved.
+
+#include "Views/AudioBusAssetDashboardEntry.h"
+
+namespace UE::Audio::Insights
+{
+	FText FAudioBusAssetDashboardEntry::GetDisplayName() const
+	{
+		return FText::FromString(FSoftObjectPath(Name).GetAssetName());
+	}
+
+	const UObject* FAudioBusAssetDashboardEntry::GetObject() const
+	{
+		return FSoftObjectPath(Name).ResolveObject();
+	}
+
+	UObject* FAudioBusAssetDashboardEntry::GetObject()
+	{
+		return FSoftObjectPath(Name).ResolveObject();
+	}
+
+	bool FAudioBusAssetDashboardEntry::IsValid() const
+	{
+		return true;
+	}
+} // namespace UE::Audio::Insights
diff --git a/Engine/Plugins/AudioInsights/Source/AudioInsights/Private/Views/AudioBusAssetDashboardEntry.h b/Engine/Plugins/AudioInsights/Source/AudioInsights/Private/Views/AudioBusAssetDashboardEntry.h
new file mode 100644
--- /dev/null
+++ b/Engine/Plugins/AudioInsights/Source/AudioInsights/Private/Views/AudioBusAssetDashboardEntry.h
@@ -0,0 +1,24 @@
+// Copyright Epic Games, Inc. All Rights Reserved.
+#pragma once
+
+#include "Views/TableDashboardViewFactory.h"
+
+class UAudioBus;
+
+namespace UE::Audio::Insights
+{
+	struct AUDIOINSIGHTS_API FAudioBusAssetDashboardEntry : public IObjectDashboardEntry
+	{
+		virtual ~FAudioBusAssetDashboardEntry() = default;
+
+		virtual FText GetDisplayName() const override;
+		virtual const UObject* GetObject() const override;
+		virtual UObject* GetObject() override;
+		virtual bool IsValid() const override;
+
+		::Audio::FDeviceId DeviceId = INDEX_NONE;
+		uint64 ComponentId = TNumericLimits<uint64>::Max();
+		FString Name;
+		TWeakObjectPtr<UAudioBus> AudioBus;
+	};
+} // namespace UE::Audio::Insights
diff --git a/Engine/Plugins/AudioInsights/Source/AudioInsights/Private/Views/AudioBusesDashboardViewFactory.cpp b/Engine/Plugins/AudioInsights/Source/AudioInsights/Private/Views/AudioBusesDashboardViewFactory.cpp
new file mode 100644
--- /dev/null
+++ b/Engine/Plugins/AudioInsights/Source/AudioInsights/Private/Views/AudioBusesDashboardViewFactory.cpp
@@ -0,0 +1,222 @@
+// Copyright Epic Games, Inc. All Rights Reserved.
+#include "Views/AudioBusesDashboardViewFactory.h"
+
+#include "AssetRegistry/AssetData.h"
+#include "AssetRegistry/AssetRegistryModule.h"
+#include "AudioDefines.h"
+#include "AudioDeviceManager.h"
+#include "AudioInsightsModule.h"
+#include "AudioInsightsStyle.h"
+#include "Editor.h"
+#include "Internationalization/Text.h"
+#include "Providers/AudioBusProvider.h"
+#include "Sound/AudioBus.h"
+#include "Subsystems/AssetEditorSubsystem.h"
+#include "Templates/SharedPointer.h"
+
+#define LOCTEXT_NAMESPACE "AudioInsights"
+
+namespace UE::Audio::Insights
+{
+	namespace AudioBusesPrivate
+	{
+		const FAudioBusAssetDashboardEntry& CastEntry(const IDashboardDataViewEntry& InData)
+		{
+			return static_cast<const FAudioBusAssetDashboardEntry&>(InData);
+		};
+	}
+
+	FAudioBusesDashboardViewFactory::FAudioBusesDashboardViewFactory()
+	{
+		FAudioBusProvider::OnAudioBusAssetAdded.AddRaw(this, &FAudioBusesDashboardViewFactory::HandleOnAudioBusAssetListUpdated);
+		FAudioBusProvider::OnAudioBusAssetRemoved.AddRaw(this, &FAudioBusesDashboardViewFactory::HandleOnAudioBusAssetListUpdated);
+		FAudioBusProvider::OnAudioBusAssetListUpdated.AddRaw(this, &FAudioBusesDashboardViewFactory::RequestListRefresh);
+
+		Providers = TArray<TSharedPtr<FTraceProviderBase>>
+		{
+			MakeShared<FAudioBusProvider>()
+		};
+
+		SortByColumn = "Name";
+		SortMode     = EColumnSortMode::Ascending;
+	}
+
+	FAudioBusesDashboardViewFactory::~FAudioBusesDashboardViewFactory()
+	{
+		FAudioBusProvider::OnAudioBusAssetAdded.RemoveAll(this);
+		FAudioBusProvider::OnAudioBusAssetRemoved.RemoveAll(this);
+		FAudioBusProvider::OnAudioBusAssetListUpdated.RemoveAll(this);
+	}
+
+	FName FAudioBusesDashboardViewFactory::GetName() const
+	{
+		return "AudioBuses";
+	}
+
+	FText FAudioBusesDashboardViewFactory::GetDisplayName() const
+	{
+		return LOCTEXT("AudioDashboard_AudioBuses_DisplayName", "Audio Buses");
+	}
+
+	TSharedRef<SWidget> FAudioBusesDashboardViewFactory::GenerateWidgetForColumn(TSharedRef<IDashboardDataViewEntry> InRowData, const FName& InColumnName)
+	{
+		if (InColumnName == "Name")
+		{
+			const TWeakObjectPtr<UAudioBus> AudioBus = AudioBusesPrivate::CastEntry(InRowData.Get()).AudioBus;
+
+			AudioBusCheckboxCheckedStates.FindOrAdd(AudioBus, false);
+
+			const FColumnData& ColumnData = GetColumns()[InColumnName];
+			const FText Label = ColumnData.GetDisplayValue(InRowData.Get());
+
+			return (Label.IsEmpty()) ? 
+				SNullWidget::NullWidget :
+				SNew(SHorizontalBox)
+				.Clipping(EWidgetClipping::ClipToBounds)
+				+ SHorizontalBox::Slot()
+				.FillWidth(0.1f)
+				[
+					SNew(SCheckBox)
+					.IsChecked(AudioBusCheckboxCheckedStates.Contains(AudioBus) && AudioBusCheckboxCheckedStates[AudioBus] == true)
+					.OnCheckStateChanged_Lambda([AudioBus, this](ECheckBoxState NewState)
+					{
+						if (AudioBusCheckboxCheckedStates.Contains(AudioBus))
+						{
+							AudioBusCheckboxCheckedStates[AudioBus] = NewState == ECheckBoxState::Checked;
+						}
+
+						OnAudioBusAssetChecked.Broadcast(NewState == ECheckBoxState::Checked, AudioBus);
+					})
+				]
+				+ SHorizontalBox::Slot()
+				.FillWidth(0.8f)
+				[
+					SNew(STextBlock)
+					.Text(Label)
+					.MinDesiredWidth(300)
+					.OnDoubleClicked_Lambda([this, InRowData](const FGeometry& MyGeometry, const FPointerEvent& PointerEvent)
+					{
+						if (GEditor)
+						{
+							if (TSharedPtr<IObjectDashboardEntry> ObjectData = StaticCastSharedPtr<IObjectDashboardEntry>(InRowData.ToSharedPtr()); 
+								ObjectData.IsValid())
+							{
+								if (const UObject* Object = ObjectData->GetObject();
+									Object && Object->IsAsset())
+								{
+									GEditor->GetEditorSubsystem<UAssetEditorSubsystem>()->OpenEditorForAsset(Object);
+									return FReply::Handled();
+								}
+							}
+						}
+
+						return FReply::Unhandled();
+					})
+				];
+		}
+
+		return SNullWidget::NullWidget;
+	}
+
+	void FAudioBusesDashboardViewFactory::ProcessEntries(FTraceTableDashboardViewFactory::EProcessReason Reason)
+	{
+		const FString FilterString = GetSearchFilterText().ToString();
+		
+		FTraceTableDashboardViewFactory::FilterEntries<FAudioBusProvider>([&FilterString](const IDashboardDataViewEntry& Entry)
+		{
+			const FAudioBusAssetDashboardEntry& AudioBusEntry = static_cast<const FAudioBusAssetDashboardEntry&>(Entry);
+			
+			return !AudioBusEntry.GetDisplayName().ToString().Contains(FilterString);
+		});
+	}
+
+	FSlateIcon FAudioBusesDashboardViewFactory::GetIcon() const
+	{
+		return FSlateStyle::Get().CreateIcon("AudioInsights.Icon");
+	}
+
+	EDefaultDashboardTabStack FAudioBusesDashboardViewFactory::GetDefaultTabStack() const
+	{
+		return EDefaultDashboardTabStack::Analysis;
+	}
+
+	TSharedRef<SWidget> FAudioBusesDashboardViewFactory::MakeWidget()
+	{
+		if (!DashboardWidget.IsValid())
+		{
+			DashboardWidget = FTraceTableDashboardViewFactory::MakeWidget();
+
+			if (FilteredEntriesListView.IsValid())
+			{
+				FilteredEntriesListView->SetSelectionMode(ESelectionMode::Single);
+			}
+		}
+
+		return DashboardWidget->AsShared();
+	}
+
+	const TMap<FName, FTraceTableDashboardViewFactory::FColumnData>& FAudioBusesDashboardViewFactory::GetColumns() const
+	{
+		auto CreateColumnData = []()
+		{
+			return TMap<FName, FTraceTableDashboardViewFactory::FColumnData>
+			{
+				{
+					"Name",
+					{
+						LOCTEXT("AudioBuses_NameColumnDisplayName", "Name"),
+						[](const IDashboardDataViewEntry& InData) { return AudioBusesPrivate::CastEntry(InData).GetDisplayName(); },
+						false,		/* bDefaultHidden */
+						0.88f,		/* FillWidth */
+						EHorizontalAlignment::HAlign_Left
+					}
+				}
+			};
+		};
+		
+		static const TMap<FName, FTraceTableDashboardViewFactory::FColumnData> ColumnData = CreateColumnData();
+		
+		return ColumnData;
+	}
+
+	void FAudioBusesDashboardViewFactory::SortTable()
+	{
+		if (SortByColumn == "Name")
+		{
+			if (SortMode == EColumnSortMode::Ascending)
+			{
+				DataViewEntries.Sort([](const TSharedPtr<IDashboardDataViewEntry>& A, const TSharedPtr<IDashboardDataViewEntry>& B)
+				{
+					const FAudioBusAssetDashboardEntry& AData = AudioBusesPrivate::CastEntry(*A.Get());
+					const FAudioBusAssetDashboardEntry& BData = AudioBusesPrivate::CastEntry(*B.Get());
+
+					return AData.GetDisplayName().CompareToCaseIgnored(BData.GetDisplayName()) < 0;
+				});
+			}
+			else if (SortMode == EColumnSortMode::Descending)
+			{
+				DataViewEntries.Sort([](const TSharedPtr<IDashboardDataViewEntry>& A, const TSharedPtr<IDashboardDataViewEntry>& B)
+				{
+					const FAudioBusAssetDashboardEntry& AData = AudioBusesPrivate::CastEntry(*A.Get());
+					const FAudioBusAssetDashboardEntry& BData = AudioBusesPrivate::CastEntry(*B.Get());
+
+					return BData.GetDisplayName().CompareToCaseIgnored(AData.GetDisplayName()) < 0;
+				});
+			}
+		}
+	}
+
+	void FAudioBusesDashboardViewFactory::RequestListRefresh()
+	{
+		if (FilteredEntriesListView.IsValid())
+		{
+			FilteredEntriesListView->RequestListRefresh();
+		}
+	}
+
+	void FAudioBusesDashboardViewFactory::HandleOnAudioBusAssetListUpdated(const TWeakObjectPtr<UObject> InAsset)
+	{
+		RequestListRefresh();
+	}
+} // namespace UE::Audio::Insights
+#undef LOCTEXT_NAMESPACE
diff --git a/Engine/Plugins/AudioInsights/Source/AudioInsights/Private/Views/AudioBusesDashboardViewFactory.h b/Engine/Plugins/AudioInsights/Source/AudioInsights/Private/Views/AudioBusesDashboardViewFactory.h
new file mode 100644
--- /dev/null
+++ b/Engine/Plugins/AudioInsights/Source/AudioInsights/Private/Views/AudioBusesDashboardViewFactory.h
@@ -0,0 +1,43 @@
+// Copyright Epic Games, Inc. All Rights Reserved.
+#pragma once
+
+#include "Delegates/Delegate.h"
+#include "Templates/SharedPointer.h"
+#include "Views/AudioBusAssetDashboardEntry.h"
+#include "Views/AudioBusAssetDashboardEntry.h"
+#include "Views/TableDashboardViewFactory.h"
+#include "Widgets/Input/SCheckBox.h"
+
+class UAudioBus;
+
+namespace UE::Audio::Insights
+{
+	class FAudioBusesDashboardViewFactory : public FTraceObjectTableDashboardViewFactory
+	{
+	public:
+		FAudioBusesDashboardViewFactory();
+		virtual ~FAudioBusesDashboardViewFactory();
+
+		virtual FName GetName() const override;
+		virtual FText GetDisplayName() const override;
+		virtual FSlateIcon GetIcon() const override;
+		virtual EDefaultDashboardTabStack GetDefaultTabStack() const override;
+		virtual TSharedRef<SWidget> MakeWidget() override;
+
+		DECLARE_MULTICAST_DELEGATE_TwoParams(FOnAudioBusAssetChecked, const bool /*bIsChecked*/, const TWeakObjectPtr<UAudioBus> /*AudioBus*/);
+		inline static FOnAudioBusAssetChecked OnAudioBusAssetChecked;
+
+	protected:
+		virtual TSharedRef<SWidget> GenerateWidgetForColumn(TSharedRef<IDashboardDataViewEntry> InRowData, const FName& InColumnName) override;
+		virtual void ProcessEntries(FTraceTableDashboardViewFactory::EProcessReason Reason) override;
+		virtual const TMap<FName, FTraceTableDashboardViewFactory::FColumnData>& GetColumns() const override;
+
+		virtual void SortTable() override;
+
+		void RequestListRefresh();
+
+		void HandleOnAudioBusAssetListUpdated(const TWeakObjectPtr<UObject> InAsset);
+
+		TMap<const TWeakObjectPtr<UAudioBus>, bool> AudioBusCheckboxCheckedStates;
+	};
+} // namespace UE::Audio::Insights
diff --git a/Engine/Plugins/AudioInsights/Source/AudioInsights/Private/Views/AudioMeterView.cpp b/Engine/Plugins/AudioInsights/Source/AudioInsights/Private/Views/AudioMeterView.cpp
new file mode 100644
--- /dev/null
+++ b/Engine/Plugins/AudioInsights/Source/AudioInsights/Private/Views/AudioMeterView.cpp
@@ -0,0 +1,177 @@
+// Copyright Epic Games, Inc. All Rights Reserved.
+#include "AudioMeterView.h"
+
+#include "Analyzers/AudioMeterSubmixAnalyzer.h"
+#include "AudioInsightsDashboardFactory.h"
+#include "AudioMeter.h"
+#include "Sound/SoundSubmix.h"
+#include "Views/SubmixesDashboardViewFactory.h"
+#include "Widgets/SNullWidget.h"
+
+namespace UE::Audio::Insights
+{
+	FAudioMeterView::FAudioMeterView(FAudioAssetVariant InAudioAssetVariant)
+		: AudioAssetVariant(InAudioAssetVariant)
+		, AudioMeterAnalyzerVariant(MakeAudioMeterAnalyzerVariant(InAudioAssetVariant))
+		, SubmixNameTextBlock(MakeSubmixNameTextBlock(InAudioAssetVariant))
+		, SubmixViewWidget(MakeWidget())
+		, OnActiveAudioDeviceChangedHandle(FDashboardFactory::OnActiveAudioDeviceChanged.AddRaw(this, &FAudioMeterView::HandleOnActiveAudioDeviceChanged))
+	{
+		
+	}
+
+	FAudioMeterView::~FAudioMeterView()
+	{
+		FDashboardFactory::OnActiveAudioDeviceChanged.Remove(OnActiveAudioDeviceChangedHandle);
+	}
+
+	FAudioMeterView::FAudioMeterVariant FAudioMeterView::MakeAudioMeterAnalyzerVariant(const FAudioAssetVariant InAudioAssetVariant)
+	{
+		switch (AudioAssetVariant.GetIndex())
+		{
+			case FAudioAssetVariant::IndexOfType<TWeakObjectPtr<UAudioBus>>():
+			{
+				return FAudioMeterVariant(TInPlaceType<TSharedPtr<FAudioMeterAnalyzer>>(), MakeShared<FAudioMeterAnalyzer>(InAudioAssetVariant.Get<TWeakObjectPtr<UAudioBus>>()));
+			}
+
+			case FAudioAssetVariant::IndexOfType<TWeakObjectPtr<USoundSubmix>>():
+			{
+				return FAudioMeterVariant(TInPlaceType<TSharedPtr<FAudioMeterSubmixAnalyzer>>(), MakeShared<FAudioMeterSubmixAnalyzer>(InAudioAssetVariant.Get<TWeakObjectPtr<USoundSubmix>>()));
+			}
+
+			default:
+				return FAudioMeterView::FAudioMeterVariant();
+		}
+	}
+
+	TSharedRef<STextBlock> FAudioMeterView::MakeSubmixNameTextBlock(const FAudioAssetVariant InAudioAssetVariant)
+	{
+		switch (InAudioAssetVariant.GetIndex())
+		{
+			case FAudioAssetVariant::IndexOfType<TWeakObjectPtr<UAudioBus>>():
+			{
+				if (TWeakObjectPtr<UAudioBus> AudioBusAsset = InAudioAssetVariant.Get<TWeakObjectPtr<UAudioBus>>();
+					AudioBusAsset.IsValid())
+				{
+					return SNew(STextBlock)
+						.Text_Lambda([AudioBusAsset] { return AudioBusAsset.IsValid() ? FText::FromString(AudioBusAsset->GetName()) : FText::GetEmpty(); })
+						.Justification(ETextJustify::Center)
+						.ColorAndOpacity(FSlateColor(FColor(80, 200, 255)));
+				}
+			}
+			break;
+
+			case FAudioAssetVariant::IndexOfType<TWeakObjectPtr<USoundSubmix>>():
+			{
+				if (TWeakObjectPtr<USoundSubmix> SoundSubmixAsset = AudioAssetVariant.Get<TWeakObjectPtr<USoundSubmix>>();
+					SoundSubmixAsset.IsValid())
+				{
+					return SNew(STextBlock)
+						.Text_Lambda([SoundSubmixAsset] { return SoundSubmixAsset.IsValid() ? FText::FromString(SoundSubmixAsset->GetName()) : FText::GetEmpty(); })
+						.Justification(ETextJustify::Center)
+						.ColorAndOpacity(FSlateColor(FColor(30, 240, 90)));
+				}
+			}
+			break;
+		}
+
+		return SNew(STextBlock);
+	}
+
+	TSharedRef<SWidget> FAudioMeterView::GetAudioMeterAnalyzerWidget() const
+	{
+		switch (AudioMeterAnalyzerVariant.GetIndex())
+		{
+			case FAudioMeterVariant::IndexOfType<TSharedPtr<FAudioMeterAnalyzer>>():
+			{
+				TSharedPtr<FAudioMeterAnalyzer> AudioMeterAnalyzer = AudioMeterAnalyzerVariant.Get<TSharedPtr<FAudioMeterAnalyzer>>();
+				return AudioMeterAnalyzer.IsValid() ? AudioMeterAnalyzer->GetAudioMeter()->GetWidget() : SNullWidget::NullWidget;
+			}
+
+			case FAudioMeterVariant::IndexOfType<TSharedPtr<FAudioMeterSubmixAnalyzer>>():
+			{
+				TSharedPtr<FAudioMeterSubmixAnalyzer> AudioMeterSubmixAnalyzer = AudioMeterAnalyzerVariant.Get<TSharedPtr<FAudioMeterSubmixAnalyzer>>();
+				return AudioMeterSubmixAnalyzer.IsValid() ? AudioMeterSubmixAnalyzer->GetWidget() : SNullWidget::NullWidget;
+			}
+
+			default:
+				return SNullWidget::NullWidget;
+		}
+	}
+
+	TSharedRef<SWidget> FAudioMeterView::MakeWidget()
+	{
+		return SNew(SVerticalBox)
+			+ SVerticalBox::Slot()
+			[
+				SNew(SVerticalBox)
+				.Clipping(EWidgetClipping::ClipToBounds)
+				// Height padding
+				+ SVerticalBox::Slot()
+				.FillHeight(0.05)
+				[
+					SNew(SBox)
+				]
+				// Audio Meter container
+				+ SVerticalBox::Slot()
+				.FillHeight(0.8)
+				[
+					// Audio Meter
+					SNew(SHorizontalBox)
+					+ SHorizontalBox::Slot()
+					.HAlign(HAlign_Center)
+					.VAlign(VAlign_Fill)
+					[
+						GetAudioMeterAnalyzerWidget()
+					]
+				]
+				// Height padding
+				+ SVerticalBox::Slot()
+				.FillHeight(0.025)
+				[
+					SNew(SBox)
+				]
+				// Submix name label
+				+ SVerticalBox::Slot()
+				.FillHeight(0.1)
+				[
+					SubmixNameTextBlock
+				]
+				// Height padding
+				+ SVerticalBox::Slot()
+				.FillHeight(0.025)
+				[
+					SNew(SBox)
+				]
+			];
+	}
+
+	void FAudioMeterView::HandleOnActiveAudioDeviceChanged()
+	{
+		switch (AudioAssetVariant.GetIndex())
+		{
+			case FAudioAssetVariant::IndexOfType<TWeakObjectPtr<UAudioBus>>():
+			{
+				if (TSharedPtr<FAudioMeterAnalyzer> AudioMeterAnalyzer = AudioMeterAnalyzerVariant.Get<TSharedPtr<FAudioMeterAnalyzer>>();
+					AudioMeterAnalyzer.IsValid())
+				{
+					AudioMeterAnalyzer->RebuildAudioMeter(AudioAssetVariant.Get<TWeakObjectPtr<UAudioBus>>());
+				}
+			}
+			break;
+
+			case FAudioAssetVariant::IndexOfType<TWeakObjectPtr<USoundSubmix>>():
+			{
+				if (TSharedPtr<FAudioMeterSubmixAnalyzer> AudioMeterSubmixAnalyzer = AudioMeterAnalyzerVariant.Get<TSharedPtr<FAudioMeterSubmixAnalyzer>>();
+					AudioMeterSubmixAnalyzer.IsValid())
+				{
+					AudioMeterSubmixAnalyzer->SetSubmix(AudioAssetVariant.Get<TWeakObjectPtr<USoundSubmix>>());
+				}
+			}
+			break;
+
+			default:
+				break;
+		}
+	}
+} // namespace UE::Audio::Insights
diff --git a/Engine/Plugins/AudioInsights/Source/AudioInsights/Private/Views/AudioMeterView.h b/Engine/Plugins/AudioInsights/Source/AudioInsights/Private/Views/AudioMeterView.h
new file mode 100644
--- /dev/null
+++ b/Engine/Plugins/AudioInsights/Source/AudioInsights/Private/Views/AudioMeterView.h
@@ -0,0 +1,46 @@
+// Copyright Epic Games, Inc. All Rights Reserved.
+#pragma once
+
+#include "Templates/SharedPointer.h"
+#include "UObject/WeakObjectPtr.h"
+#include "Widgets/SWidget.h"
+#include "Widgets/Text/STextBlock.h"
+
+class UAudioBus;
+class USoundSubmix;
+
+namespace UE::Audio::Insights
+{
+	class FAudioMeterAnalyzer;
+	class FAudioMeterSubmixAnalyzer;
+
+	class FAudioMeterView : public TSharedFromThis<FAudioMeterView>
+	{
+	public:
+		using FAudioAssetVariant = TVariant<TWeakObjectPtr<UAudioBus>, TWeakObjectPtr<USoundSubmix>>;
+		using FAudioMeterVariant = TVariant<TSharedPtr<FAudioMeterAnalyzer>, TSharedPtr<FAudioMeterSubmixAnalyzer>>;
+
+		FAudioMeterView(FAudioAssetVariant InAudioAssetVariant);
+		virtual ~FAudioMeterView();
+
+		TSharedRef<SWidget> GetWidget() const { return SubmixViewWidget; };
+
+	private:
+		FAudioMeterVariant MakeAudioMeterAnalyzerVariant(const FAudioAssetVariant InAudioAssetVariant);
+		TSharedRef<STextBlock> MakeSubmixNameTextBlock(const FAudioAssetVariant InAudioAssetVariant);
+		TSharedRef<SWidget> MakeWidget();
+
+		TSharedRef<SWidget> GetAudioMeterAnalyzerWidget() const;
+
+		void HandleOnActiveAudioDeviceChanged();
+
+		FAudioAssetVariant AudioAssetVariant;
+		FAudioMeterVariant AudioMeterAnalyzerVariant;
+
+		TSharedRef<STextBlock> SubmixNameTextBlock;
+
+		TSharedRef<SWidget> SubmixViewWidget;
+
+		FDelegateHandle OnActiveAudioDeviceChangedHandle;
+	};
+} // namespace UE::Audio::Insights
diff --git a/Engine/Plugins/AudioInsights/Source/AudioInsights/Private/Views/AudioMetersDashboardViewFactory.cpp b/Engine/Plugins/AudioInsights/Source/AudioInsights/Private/Views/AudioMetersDashboardViewFactory.cpp
new file mode 100644
--- /dev/null
+++ b/Engine/Plugins/AudioInsights/Source/AudioInsights/Private/Views/AudioMetersDashboardViewFactory.cpp
@@ -0,0 +1,163 @@
+// Copyright Epic Games, Inc. All Rights Reserved.
+#include "AudioMetersDashboardViewFactory.h"
+
+#include "AudioInsightsDashboardFactory.h"
+#include "AudioInsightsStyle.h"
+#include "Providers/AudioBusProvider.h"
+#include "Providers/SoundSubmixProvider.h"
+#include "Sound/AudioBus.h"
+#include "Sound/AudioSettings.h"
+#include "Sound/SoundSubmix.h"
+#include "Views/AudioBusesDashboardViewFactory.h"
+#include "Views/SubmixesDashboardViewFactory.h"
+
+#define LOCTEXT_NAMESPACE "AudioInsights"
+
+namespace UE::Audio::Insights
+{
+	FName FAudioMetersDashboardViewFactory::GetName() const
+	{
+		return "AudioMeters";
+	}
+
+	FText FAudioMetersDashboardViewFactory::GetDisplayName() const
+	{
+		return LOCTEXT("AudioDashboard_AudioMetersTab_DisplayName", "Audio Meters");
+	}
+
+	EDefaultDashboardTabStack FAudioMetersDashboardViewFactory::GetDefaultTabStack() const
+	{
+		return EDefaultDashboardTabStack::AudioMeters;
+	}
+
+	FSlateIcon FAudioMetersDashboardViewFactory::GetIcon() const
+	{
+		return FSlateStyle::Get().CreateIcon("AudioInsights.Icon.Submix");
+	}
+
+	TSharedRef<SWidget> FAudioMetersDashboardViewFactory::MakeWidget()
+	{
+		if (!OnSubmixAssetCheckedHandle.IsValid())
+		{
+			OnSubmixAssetCheckedHandle = FSubmixesDashboardViewFactory::OnSubmixAssetChecked.AddSP(this, &FAudioMetersDashboardViewFactory::HandleOnSubmixAssetChecked);
+		}
+
+		if (!OnAudioBusAssetCheckedHandle.IsValid())
+		{
+			OnAudioBusAssetCheckedHandle = FAudioBusesDashboardViewFactory::OnAudioBusAssetChecked.AddSP(this, &FAudioMetersDashboardViewFactory::HandleOnAudioBusAssetChecked);
+		}
+
+		if (!OnSubmixAssetRemovedHandle.IsValid())
+		{
+			OnSubmixAssetRemovedHandle = FSoundSubmixProvider::OnSubmixAssetRemoved.AddSP(this, &FAudioMetersDashboardViewFactory::HandleOnAssetRemoved);
+		}
+
+		if (!OnAudioBusAssetRemovedHandle.IsValid())
+		{
+			OnAudioBusAssetRemovedHandle = FAudioBusProvider::OnAudioBusAssetRemoved.AddSP(this, &FAudioMetersDashboardViewFactory::HandleOnAssetRemoved);
+		}
+		
+		if (!MeterViewsScrollBox.IsValid())
+		{
+			SAssignNew(MeterViewsScrollBox, SScrollBox)
+			.Orientation(Orient_Horizontal)
+			+ SScrollBox::Slot()
+			[
+				SAssignNew(AudioMeterViewsContainer, SHorizontalBox)
+			];
+		}
+
+		return MeterViewsScrollBox.ToSharedRef();
+	}
+
+	void FAudioMetersDashboardViewFactory::HandleOnSubmixAssetChecked(const bool bInIsChecked, const TWeakObjectPtr<USoundSubmix> InSoundSubmix)
+	{
+		if (!InSoundSubmix.IsValid())
+		{
+			return;
+		}
+
+		const uint32 SubmixUniqueId = InSoundSubmix->GetUniqueID();
+
+		if (bInIsChecked)
+		{
+			AudioMeterViews.Add(SubmixUniqueId, MakeShared<FAudioMeterView>(FAudioMeterView::FAudioAssetVariant(TInPlaceType<TWeakObjectPtr<USoundSubmix>>(), InSoundSubmix)));
+
+			AudioMeterViewsContainer->AddSlot()
+			.AutoWidth()
+			.Padding(10.0f, 0.0f, 10.0f, 0.0f)
+			[
+				AudioMeterViews[SubmixUniqueId]->GetWidget()
+			];
+		}
+		else
+		{
+			if (AudioMeterViews.Contains(SubmixUniqueId))
+			{
+				AudioMeterViewsContainer->RemoveSlot(AudioMeterViews[SubmixUniqueId]->GetWidget());
+				AudioMeterViews.Remove(SubmixUniqueId);
+			}
+		}
+
+		if (MeterViewsScrollBox.IsValid())
+		{
+			MeterViewsScrollBox->Invalidate(EInvalidateWidget::Layout);
+		}
+	}
+
+	void FAudioMetersDashboardViewFactory::HandleOnAudioBusAssetChecked(const bool bInIsChecked, const TWeakObjectPtr<UAudioBus> InAudioBus)
+	{
+		if (!InAudioBus.IsValid())
+		{
+			return;
+		}
+
+		const uint32 AudioBusUniqueId = InAudioBus->GetUniqueID();
+
+		if (bInIsChecked)
+		{
+			AudioMeterViews.Add(AudioBusUniqueId, MakeShared<FAudioMeterView>(FAudioMeterView::FAudioAssetVariant(TInPlaceType<TWeakObjectPtr<UAudioBus>>(), InAudioBus)));
+
+			AudioMeterViewsContainer->AddSlot()
+			.AutoWidth()
+			.Padding(10.0f, 0.0f, 10.0f, 0.0f)
+			[
+				AudioMeterViews[AudioBusUniqueId]->GetWidget()
+			];
+		}
+		else
+		{
+			if (AudioMeterViews.Contains(AudioBusUniqueId))
+			{
+				AudioMeterViewsContainer->RemoveSlot(AudioMeterViews[AudioBusUniqueId]->GetWidget());
+				AudioMeterViews.Remove(AudioBusUniqueId);
+			}
+		}
+
+		if (MeterViewsScrollBox.IsValid())
+		{
+			MeterViewsScrollBox->Invalidate(EInvalidateWidget::Layout);
+		}
+	}
+
+	void FAudioMetersDashboardViewFactory::HandleOnAssetRemoved(const TWeakObjectPtr<UObject> InAsset)
+	{
+		if (!InAsset.IsValid())
+		{
+			return;
+		}
+
+		if (const uint32 AssetUniqueId = InAsset->GetUniqueID();
+			AudioMeterViews.Contains(AssetUniqueId))
+		{
+			AudioMeterViewsContainer->RemoveSlot(AudioMeterViews[AssetUniqueId]->GetWidget());
+			AudioMeterViews.Remove(AssetUniqueId);
+		}
+
+		if (MeterViewsScrollBox.IsValid())
+		{
+			MeterViewsScrollBox->Invalidate(EInvalidateWidget::Layout);
+		}
+	}
+} // namespace UE::Audio::Insights
+#undef LOCTEXT_NAMESPACE
diff --git a/Engine/Plugins/AudioInsights/Source/AudioInsights/Private/Views/AudioMetersDashboardViewFactory.h b/Engine/Plugins/AudioInsights/Source/AudioInsights/Private/Views/AudioMetersDashboardViewFactory.h
new file mode 100644
--- /dev/null
+++ b/Engine/Plugins/AudioInsights/Source/AudioInsights/Private/Views/AudioMetersDashboardViewFactory.h
@@ -0,0 +1,43 @@
+// Copyright Epic Games, Inc. All Rights Reserved.
+#pragma once
+
+#include "AudioMeterView.h"
+#include "Templates/SharedPointer.h"
+#include "Views/DashboardViewFactory.h"
+#include "Widgets/Layout/SScrollBox.h"
+
+class SAudioVolumeSlider;
+class UAudioBus;
+class USoundSubmix;
+
+namespace UE::Audio::Insights
+{
+	class FAudioMeterAnalyzer;
+
+	class FAudioMetersDashboardViewFactory : public IDashboardViewFactory, public TSharedFromThis<FAudioMetersDashboardViewFactory>
+	{
+	public:
+		virtual ~FAudioMetersDashboardViewFactory() = default;
+
+		virtual FName GetName() const override;
+		virtual FText GetDisplayName() const override;
+		virtual EDefaultDashboardTabStack GetDefaultTabStack() const override;
+		virtual FSlateIcon GetIcon() const override;
+		virtual TSharedRef<SWidget> MakeWidget() override;
+
+	private:
+		void HandleOnSubmixAssetChecked(const bool bInIsChecked, const TWeakObjectPtr<USoundSubmix> InSoundSubmix);
+		void HandleOnAudioBusAssetChecked(const bool bInIsChecked, const TWeakObjectPtr<UAudioBus> InAudioBus);
+
+		void HandleOnAssetRemoved(const TWeakObjectPtr<UObject> InAsset);
+
+		TSharedPtr<SScrollBox> MeterViewsScrollBox;
+		TSharedPtr<SHorizontalBox> AudioMeterViewsContainer;
+		TMap<uint32, TSharedRef<FAudioMeterView>> AudioMeterViews;
+
+		FDelegateHandle OnSubmixAssetCheckedHandle;
+		FDelegateHandle OnAudioBusAssetCheckedHandle;
+		FDelegateHandle OnSubmixAssetRemovedHandle;
+		FDelegateHandle OnAudioBusAssetRemovedHandle;
+	};
+} // namespace UE::Audio::Insights
diff --git a/Engine/Plugins/AudioInsights/Source/AudioInsights/Private/Views/LogDashboardViewFactory.cpp b/Engine/Plugins/AudioInsights/Source/AudioInsights/Private/Views/LogDashboardViewFactory.cpp
new file mode 100644
--- /dev/null
+++ b/Engine/Plugins/AudioInsights/Source/AudioInsights/Private/Views/LogDashboardViewFactory.cpp
@@ -0,0 +1,121 @@
+// Copyright Epic Games, Inc. All Rights Reserved.
+
+#include "Views/LogDashboardViewFactory.h"
+
+#include "AudioInsightsStyle.h"
+#include "OutputLogModule.h"
+
+#define LOCTEXT_NAMESPACE "AudioInsights"
+
+
+namespace UE::Audio::Insights
+{
+	namespace LogDashboardViewFactoryPrivate
+	{
+		static const TSet<FString> DefaultCategoryKeywords = { TEXT("Audio"), TEXT("Sound") };
+
+		bool IsDefaultDashboardCategory(const FName InCategoryName)
+		{
+			const FString CategoryNameStr = InCategoryName.ToString();
+			for (const FString& Keyword : DefaultCategoryKeywords)
+			{
+				if (CategoryNameStr.Contains(Keyword))
+				{
+					return true;
+				}
+			}
+
+			return false;
+		}
+	} // namespace LogDashboardViewFactoryPrivate
+
+	FLogDashboardViewFactory::FLogDashboardViewFactory()
+	{
+		CategoryCollector.RunAsync();
+	}
+
+	FName FLogDashboardViewFactory::GetName() const
+	{
+		return "Log";
+	}
+
+	FText FLogDashboardViewFactory::GetDisplayName() const
+	{
+		return LOCTEXT("AudioDashboard_LogTab_DisplayName", "Log");
+	}
+
+	EDefaultDashboardTabStack FLogDashboardViewFactory::GetDefaultTabStack() const
+	{
+		return EDefaultDashboardTabStack::Log;
+	}
+
+	FSlateIcon FLogDashboardViewFactory::GetIcon() const
+	{
+		return FSlateStyle::Get().CreateIcon("AudioInsights.Icon.Log");
+	}
+
+	TSharedRef<SWidget> FLogDashboardViewFactory::MakeWidget()
+	{
+		using namespace LogDashboardViewFactoryPrivate;
+
+		FOutputLogCreationParams Params;
+		Params.SettingsMenuCreationFlags = EOutputLogSettingsMenuFlags::SkipClearOnPie
+			| EOutputLogSettingsMenuFlags::SkipOpenSourceButton
+			| EOutputLogSettingsMenuFlags::SkipEnableWordWrapping;
+
+		Params.DefaultCategorySelection = CategoryCollector.GetCollectedCategories();
+		Params.AllowAsInitialLogCategory = FAllowLogCategoryCallback::CreateStatic(&IsDefaultDashboardCategory);
+
+		FModuleManager::LoadModuleChecked<FOutputLogModule>("OutputLog");
+		return FOutputLogModule::Get().MakeOutputLogWidget(Params);
+	}
+
+	FLogDashboardViewFactory::FLogCategoryCollector::FLogCategoryCollector()
+	{
+		if (GLog)
+		{
+			GLog->AddOutputDevice(this);
+		}
+	}
+
+	FLogDashboardViewFactory::FLogCategoryCollector::~FLogCategoryCollector()
+	{
+		if (GLog)
+		{
+			GLog->RemoveOutputDevice(this);
+		}
+	}
+
+	FDefaultCategorySelectionMap FLogDashboardViewFactory::FLogCategoryCollector::GetCollectedCategories() const
+	{
+		FScopeLock CollectionLock(&CollectionCritSec);
+		return CollectedCategories;
+	}
+
+	void FLogDashboardViewFactory::FLogCategoryCollector::RunAsync()
+	{
+		FScopeLock CollectionLock(&CollectionCritSec);
+		CollectedCategories.Empty();
+
+		if (GLog)
+		{
+			GLog->SerializeBacklog(this);
+		}
+	}
+
+	void FLogDashboardViewFactory::FLogCategoryCollector::Serialize(const TCHAR* InMsg, ELogVerbosity::Type Verbosity, const FName& InCategory)
+	{
+		using namespace LogDashboardViewFactoryPrivate;
+
+		FScopeLock CollectionLock(&CollectionCritSec);
+
+		const bool bIsDefaultCategory = IsDefaultDashboardCategory(InCategory);
+		CollectedCategories.FindOrAdd(InCategory) = bIsDefaultCategory;
+	}
+
+	bool FLogDashboardViewFactory::FLogCategoryCollector::IsMemoryOnly() const
+	{
+		return true;
+	}
+} // namespace UE::Audio::Insights
+#undef LOCTEXT_NAMESPACE
diff --git a/Engine/Plugins/AudioInsights/Source/AudioInsights/Private/Views/LogDashboardViewFactory.h b/Engine/Plugins/AudioInsights/Source/AudioInsights/Private/Views/LogDashboardViewFactory.h
new file mode 100644
--- /dev/null
+++ b/Engine/Plugins/AudioInsights/Source/AudioInsights/Private/Views/LogDashboardViewFactory.h
@@ -0,0 +1,49 @@
+// Copyright Epic Games, Inc. All Rights Reserved.
+#pragma once
+
+#include "Internationalization/Text.h"
+#include "Misc/OutputDevice.h"
+#include "OutputLogCreationParams.h"
+#include "Templates/SharedPointer.h"
+#include "Textures/SlateIcon.h"
+#include "UObject/NameTypes.h"
+#include "Views/DashboardViewFactory.h"
+#include "Widgets/SWidget.h"
+
+#define LOCTEXT_NAMESPACE "AudioInsights"
+
+
+namespace UE::Audio::Insights
+{
+	class FLogDashboardViewFactory : public IDashboardViewFactory
+	{
+	public:
+		FLogDashboardViewFactory();
+
+		virtual FName GetName() const override;
+		virtual FText GetDisplayName() const override;
+		virtual EDefaultDashboardTabStack GetDefaultTabStack() const override;
+		virtual FSlateIcon GetIcon() const override;
+		virtual TSharedRef<SWidget> MakeWidget() override;
+
+	private:
+		struct FLogCategoryCollector : public FOutputDevice
+		{
+			FLogCategoryCollector();
+			~FLogCategoryCollector();
+
+			FDefaultCategorySelectionMap GetCollectedCategories() const;
+			void RunAsync();
+
+			virtual bool IsMemoryOnly() const override;
+			virtual void Serialize(const TCHAR* InMsg, ELogVerbosity::Type Verbosity, const FName& InCategory) override;
+
+		private:
+			mutable FCriticalSection CollectionCritSec;
+			FDefaultCategorySelectionMap CollectedCategories;
+		};
+
+		FLogCategoryCollector CategoryCollector;
+	};
+} // namespace UE::Audio::Insights
+#undef LOCTEXT_NAMESPACE
diff --git a/Engine/Plugins/AudioInsights/Source/AudioInsights/Private/Views/MixerSourceDashboardViewFactory.cpp b/Engine/Plugins/AudioInsights/Source/AudioInsights/Private/Views/MixerSourceDashboardViewFactory.cpp
new file mode 100644
--- /dev/null
+++ b/Engine/Plugins/AudioInsights/Source/AudioInsights/Private/Views/MixerSourceDashboardViewFactory.cpp
@@ -0,0 +1,999 @@
+// Copyright Epic Games, Inc. All Rights Reserved.
+#include "Views/MixerSourceDashboardViewFactory.h"
+
+#include "Audio/AudioDebug.h"
+#include "AudioDeviceManager.h"
+#include "AudioInsightsDashboardAssetCommands.h"
+#include "AudioInsightsModule.h"
+#include "AudioInsightsStyle.h"
+#include "Internationalization/Text.h"
+#include "IPropertyTypeCustomization.h"
+#include "Providers/MixerSourceTraceProvider.h"
+#include "SSimpleTimeSlider.h"
+#include "Templates/SharedPointer.h"
+#include "UObject/SoftObjectPath.h"
+#include "Widgets/Layout/SSplitter.h"
+
+#define LOCTEXT_NAMESPACE "AudioInsights"
+
+
+namespace UE::Audio::Insights
+{
+	namespace MixerSourcePrivate
+	{
+		const FMixerSourceDashboardEntry& CastEntry(const IDashboardDataViewEntry& InData)
+		{
+			return static_cast<const FMixerSourceDashboardEntry&>(InData);
+		};
+
+		static const FText PlotColumnSelectDescription = LOCTEXT("AudioDashboard_MixerSources_SelectPlotColumnDescription", "Select a column from the table to plot.");
+	} // namespace MixerSourcePrivate
+
+	const double FMixerSourceDashboardViewFactory::MaxPlotHistorySeconds = 5.0;
+	const int32 FMixerSourceDashboardViewFactory::MaxPlotSources = 32;
+
+	FMixerSourceDashboardViewFactory::FMixerSourceDashboardViewFactory()
+	{
+		const FTraceModule& TraceModule = FAudioInsightsModule::GetChecked().GetTraceModule();
+		Providers = TArray<TSharedPtr<FTraceProviderBase>>
+		{
+			TraceModule.FindAudioTraceProvider<FMixerSourceTraceProvider>()
+		};
+	}
+
+	FName FMixerSourceDashboardViewFactory::GetName() const
+	{
+		return "MixerSources";
+	}
+
+	FText FMixerSourceDashboardViewFactory::GetDisplayName() const
+	{
+		return LOCTEXT("AudioDashboard_MixerSources_DisplayName", "Sources");
+	}
+
+	FSlateIcon FMixerSourceDashboardViewFactory::GetIcon() const
+	{
+		return FSlateStyle::Get().CreateIcon("AudioInsights.Icon.Sources");
+	}
+
+	EDefaultDashboardTabStack FMixerSourceDashboardViewFactory::GetDefaultTabStack() const
+	{
+		return EDefaultDashboardTabStack::Analysis;
+	}
+
+	const TMap<FName, FTraceTableDashboardViewFactory::FColumnData>& FMixerSourceDashboardViewFactory::GetColumns() const
+	{
+		auto CreateColumnData = []()
+		{
+
+			return TMap<FName, FTraceTableDashboardViewFactory::FColumnData>
+			{
+				{
+					"PlayOrder",
+					{
+						LOCTEXT("PlayOrder_PlayOrderColumnDisplayName", "Play Order"),
+						[](const IDashboardDataViewEntry& InData) { return FText::AsNumber(MixerSourcePrivate::CastEntry(InData).PlayOrder); },
+						true /* bDefaultHidden */,
+						0.08f /* FillWidth */
+					}
+				},
+				{
+					"Name",
+					{
+						LOCTEXT("Source_NameColumnDisplayName", "Name"),
+						[](const IDashboardDataViewEntry& InData) { return FText::FromString(*FSoftObjectPath(MixerSourcePrivate::CastEntry(InData).Name).GetAssetName()); },
+						false /* bDefaultHidden */,
+						0.75f /* FillWidth */
+					}
+				},
+				{
+					"Amplitude",
+					{
+						LOCTEXT("Source_EnvColumnDisplayName", "Amp (Peak)"),
+						[](const IDashboardDataViewEntry& InData) { return FText::AsNumber(MixerSourcePrivate::CastEntry(InData).Envelope, FSlateStyle::Get().GetAmpFloatFormat()); },
+						false /* bDefaultHidden */,
+						0.12f /* FillWidth */
+					}
+				},
+				{
+					"Volume",
+					{
+						LOCTEXT("Source_VolumeColumnDisplayName", "Volume"),
+						[](const IDashboardDataViewEntry& InData) { return FText::AsNumber(MixerSourcePrivate::CastEntry(InData).Volume, FSlateStyle::Get().GetAmpFloatFormat()); },
+						false /* bDefaultHidden */,
+						0.07f /* FillWidth */
+					}
+				},
+				{
+					"DistanceAttenuation",
+					{
+						LOCTEXT("Source_AttenuationColumnDisplayName", "Distance Attenuation"),
+						[](const IDashboardDataViewEntry& InData) { return FText::AsNumber(MixerSourcePrivate::CastEntry(InData).DistanceAttenuation, FSlateStyle::Get().GetAmpFloatFormat()); },
+						true  /* bDefaultHidden */,
+						0.15f /* FillWidth */
+					}
+				},
+				{
+					"Pitch",
+					{
+						LOCTEXT("Source_PitchColumnDisplayName", "Pitch"),
+						[](const IDashboardDataViewEntry& InData) { return FText::AsNumber(MixerSourcePrivate::CastEntry(InData).Pitch, FSlateStyle::Get().GetPitchFloatFormat()); },
+						false /* bDefaultHidden */,
+						0.06f /* FillWidth */
+					}
+				},
+				{
+					"LPF",
+					{
+						LOCTEXT("Source_LPFColumnDisplayName", "LPF Freq (Hz)"),
+						[](const IDashboardDataViewEntry& InData) { return FText::AsNumber(MixerSourcePrivate::CastEntry(InData).LPFFreq, FSlateStyle::Get().GetFreqFloatFormat()); },
+						true  /* bDefaultHidden */,
+						0.1f /* FillWidth */
+					}
+				},
+				{
+					"HPF",
+					{
+						LOCTEXT("Source_HPFColumnDisplayName", "HPF Freq (Hz)"),
+						[](const IDashboardDataViewEntry& InData) { return FText::AsNumber(MixerSourcePrivate::CastEntry(InData).HPFFreq, FSlateStyle::Get().GetFreqFloatFormat()); },
+						true  /* bDefaultHidden */,
+						0.1f /* FillWidth */
+					}
+				}
+			};
+		};
+		static const TMap<FName, FTraceTableDashboardViewFactory::FColumnData> ColumnData = CreateColumnData();
+		return ColumnData;
+	}
+
+	void FMixerSourceDashboardViewFactory::SortTable()
+	{
+		if (SortByColumn == "PlayOrder")
+		{
+			if (SortMode == EColumnSortMode::Ascending)
+			{
+				DataViewEntries.Sort([](const TSharedPtr<IDashboardDataViewEntry>& A, const TSharedPtr<IDashboardDataViewEntry>& B)
+				{
+					const FMixerSourceDashboardEntry& AData = MixerSourcePrivate::CastEntry(*A.Get());
+					const FMixerSourceDashboardEntry& BData = MixerSourcePrivate::CastEntry(*B.Get());
+
+					return AData.PlayOrder < BData.PlayOrder;
+				});
+			}
+			else if (SortMode == EColumnSortMode::Descending)
+			{
+				DataViewEntries.Sort([](const TSharedPtr<IDashboardDataViewEntry>& A, const TSharedPtr<IDashboardDataViewEntry>& B)
+				{
+					const FMixerSourceDashboardEntry& AData = MixerSourcePrivate::CastEntry(*A.Get());
+					const FMixerSourceDashboardEntry& BData = MixerSourcePrivate::CastEntry(*B.Get());
+
+					return BData.PlayOrder < AData.PlayOrder;
+				});
+			}
+		}
+		else if (SortByColumn == "Name")
+		{
+			if (SortMode == EColumnSortMode::Ascending)
+			{
+				DataViewEntries.Sort([](const TSharedPtr<IDashboardDataViewEntry>& A, const TSharedPtr<IDashboardDataViewEntry>& B)
+				{
+					const FMixerSourceDashboardEntry& AData = MixerSourcePrivate::CastEntry(*A.Get());
+					const FMixerSourceDashboardEntry& BData = MixerSourcePrivate::CastEntry(*B.Get());
+
+					return AData.GetDisplayName().CompareToCaseIgnored(BData.GetDisplayName()) < 0;
+				});
+			}
+			else if (SortMode == EColumnSortMode::Descending)
+			{
+				DataViewEntries.Sort([](const TSharedPtr<IDashboardDataViewEntry>& A, const TSharedPtr<IDashboardDataViewEntry>& B)
+				{
+					const FMixerSourceDashboardEntry& AData = MixerSourcePrivate::CastEntry(*A.Get());
+					const FMixerSourceDashboardEntry& BData = MixerSourcePrivate::CastEntry(*B.Get());
+
+					return BData.GetDisplayName().CompareToCaseIgnored(AData.GetDisplayName()) < 0;
+				});
+			}
+		}
+		else if (SortByColumn == "Amplitude")
+		{
+			if (SortMode == EColumnSortMode::Ascending)
+			{
+				DataViewEntries.Sort([](const TSharedPtr<IDashboardDataViewEntry>& A, const TSharedPtr<IDashboardDataViewEntry>& B)
+				{
+					const FMixerSourceDashboardEntry& AData = MixerSourcePrivate::CastEntry(*A.Get());
+					const FMixerSourceDashboardEntry& BData = MixerSourcePrivate::CastEntry(*B.Get());
+
+					return AData.Envelope < BData.Envelope;
+				});
+			}
+			else if (SortMode == EColumnSortMode::Descending)
+			{
+				DataViewEntries.Sort([](const TSharedPtr<IDashboardDataViewEntry>& A, const TSharedPtr<IDashboardDataViewEntry>& B)
+				{
+					const FMixerSourceDashboardEntry& AData = MixerSourcePrivate::CastEntry(*A.Get());
+					const FMixerSourceDashboardEntry& BData = MixerSourcePrivate::CastEntry(*B.Get());
+
+					return BData.Envelope < AData.Envelope;
+				});
+			}
+		}
+		else if (SortByColumn == "Volume")
+		{
+			if (SortMode == EColumnSortMode::Ascending)
+			{
+				DataViewEntries.Sort([](const TSharedPtr<IDashboardDataViewEntry>& A, const TSharedPtr<IDashboardDataViewEntry>& B)
+				{
+					const FMixerSourceDashboardEntry& AData = MixerSourcePrivate::CastEntry(*A.Get());
+					const FMixerSourceDashboardEntry& BData = MixerSourcePrivate::CastEntry(*B.Get());
+
+					return AData.Volume < BData.Volume;
+				});
+			}
+			else if (SortMode == EColumnSortMode::Descending)
+			{
+				DataViewEntries.Sort([](const TSharedPtr<IDashboardDataViewEntry>& A, const TSharedPtr<IDashboardDataViewEntry>& B)
+				{
+					const FMixerSourceDashboardEntry& AData = MixerSourcePrivate::CastEntry(*A.Get());
+					const FMixerSourceDashboardEntry& BData = MixerSourcePrivate::CastEntry(*B.Get());
+
+					return BData.Volume < AData.Volume;
+				});
+			}
+		}
+		else if (SortByColumn == "DistanceAttenuation")
+		{
+			if (SortMode == EColumnSortMode::Ascending)
+			{
+				DataViewEntries.Sort([](const TSharedPtr<IDashboardDataViewEntry>& A, const TSharedPtr<IDashboardDataViewEntry>& B)
+				{
+					const FMixerSourceDashboardEntry& AData = MixerSourcePrivate::CastEntry(*A.Get());
+					const FMixerSourceDashboardEntry& BData = MixerSourcePrivate::CastEntry(*B.Get());
+
+					return AData.DistanceAttenuation < BData.DistanceAttenuation;
+				});
+			}
+			else if (SortMode == EColumnSortMode::Descending)
+			{
+				DataViewEntries.Sort([](const TSharedPtr<IDashboardDataViewEntry>& A, const TSharedPtr<IDashboardDataViewEntry>& B)
+				{
+					const FMixerSourceDashboardEntry& AData = MixerSourcePrivate::CastEntry(*A.Get());
+					const FMixerSourceDashboardEntry& BData = MixerSourcePrivate::CastEntry(*B.Get());
+
+					return BData.DistanceAttenuation < AData.DistanceAttenuation;
+				});
+			}
+		}
+		else if (SortByColumn == "Pitch")
+		{
+			if (SortMode == EColumnSortMode::Ascending)
+			{
+				DataViewEntries.Sort([](const TSharedPtr<IDashboardDataViewEntry>& A, const TSharedPtr<IDashboardDataViewEntry>& B)
+				{
+					const FMixerSourceDashboardEntry& AData = MixerSourcePrivate::CastEntry(*A.Get());
+					const FMixerSourceDashboardEntry& BData = MixerSourcePrivate::CastEntry(*B.Get());
+
+					return AData.Pitch < BData.Pitch;
+				});
+			}
+			else if (SortMode == EColumnSortMode::Descending)
+			{
+				DataViewEntries.Sort([](const TSharedPtr<IDashboardDataViewEntry>& A, const TSharedPtr<IDashboardDataViewEntry>& B)
+				{
+					const FMixerSourceDashboardEntry& AData = MixerSourcePrivate::CastEntry(*A.Get());
+					const FMixerSourceDashboardEntry& BData = MixerSourcePrivate::CastEntry(*B.Get());
+
+					return BData.Pitch < AData.Pitch;
+				});
+			}
+		}
+		else if (SortByColumn == "LPF")
+		{
+			if (SortMode == EColumnSortMode::Ascending)
+			{
+				DataViewEntries.Sort([](const TSharedPtr<IDashboardDataViewEntry>& A, const TSharedPtr<IDashboardDataViewEntry>& B)
+				{
+					const FMixerSourceDashboardEntry& AData = MixerSourcePrivate::CastEntry(*A.Get());
+					const FMixerSourceDashboardEntry& BData = MixerSourcePrivate::CastEntry(*B.Get());
+
+					return AData.LPFFreq < BData.LPFFreq;
+				});
+			}
+			else if (SortMode == EColumnSortMode::Descending)
+			{
+				DataViewEntries.Sort([](const TSharedPtr<IDashboardDataViewEntry>& A, const TSharedPtr<IDashboardDataViewEntry>& B)
+				{
+					const FMixerSourceDashboardEntry& AData = MixerSourcePrivate::CastEntry(*A.Get());
+					const FMixerSourceDashboardEntry& BData = MixerSourcePrivate::CastEntry(*B.Get());
+
+					return BData.LPFFreq < AData.LPFFreq;
+				});
+			}
+		}
+		else if (SortByColumn == "HPF")
+		{
+			if (SortMode == EColumnSortMode::Ascending)
+			{
+				DataViewEntries.Sort([](const TSharedPtr<IDashboardDataViewEntry>& A, const TSharedPtr<IDashboardDataViewEntry>& B)
+				{
+					const FMixerSourceDashboardEntry& AData = MixerSourcePrivate::CastEntry(*A.Get());
+					const FMixerSourceDashboardEntry& BData = MixerSourcePrivate::CastEntry(*B.Get());
+
+					return AData.HPFFreq < BData.HPFFreq;
+				});
+			}
+			else if (SortMode == EColumnSortMode::Descending)
+			{
+				DataViewEntries.Sort([](const TSharedPtr<IDashboardDataViewEntry>& A, const TSharedPtr<IDashboardDataViewEntry>& B)
+				{
+					const FMixerSourceDashboardEntry& AData = MixerSourcePrivate::CastEntry(*A.Get());
+					const FMixerSourceDashboardEntry& BData = MixerSourcePrivate::CastEntry(*B.Get());
+
+					return BData.HPFFreq < AData.HPFFreq;
+				});
+			}
+		}
+	}
+
+	TSharedPtr<SWidget> FMixerSourceDashboardViewFactory::OnConstructContextMenu()
+	{
+		const FDashboardAssetCommands& Commands = FDashboardAssetCommands::Get();
+
+		TSharedPtr<FUICommandList> CommandList = MakeShared<FUICommandList>();
+		CommandList->MapAction(Commands.GetMuteCommand(), FExecuteAction::CreateRaw(this, &FMixerSourceDashboardViewFactory::MuteSound));
+		CommandList->MapAction(Commands.GetSoloCommand(), FExecuteAction::CreateRaw(this, &FMixerSourceDashboardViewFactory::SoloSound));
+		CommandList->MapAction(Commands.GetClearMuteSoloCommand(), FExecuteAction::CreateRaw(this, &FMixerSourceDashboardViewFactory::ClearMutesAndSolos));
+
+		constexpr bool bShouldCloseWindowAfterMenuSelection = true;
+		FMenuBuilder MenuBuilder(bShouldCloseWindowAfterMenuSelection, CommandList);
+
+		MenuBuilder.BeginSection("SoundActions", LOCTEXT("SoundActions_Header", "Sound Actions"));
+		{
+			MenuBuilder.AddMenuEntry(Commands.GetMuteCommand());
+			MenuBuilder.AddMenuEntry(Commands.GetSoloCommand());
+			MenuBuilder.AddMenuEntry(Commands.GetClearMuteSoloCommand());
+		}
+		MenuBuilder.EndSection();
+
+		return MenuBuilder.MakeWidget();
+	}
+
+	FSlateColor FMixerSourceDashboardViewFactory::GetRowColor(const TSharedPtr<IDashboardDataViewEntry>& InRowDataPtr)
+	{
+		bool bIsSolo = false;
+		bool bIsMute = false;
+
+		FColor RowTextColor(255, 255, 255);
+
+		if (FAudioDeviceManager* AudioDeviceManager = FAudioDeviceManager::Get())
+		{
+			const FSoundAssetDashboardEntry& SoundAssetDashboardEntry = *StaticCastSharedPtr<FSoundAssetDashboardEntry>(InRowDataPtr).Get();
+
+			const FName SoundAssetName(SoundAssetDashboardEntry.GetDisplayName().ToString());
+
+			bIsSolo = AudioDeviceManager->GetDebugger().IsSoloSoundWave(SoundAssetName);
+			bIsMute = AudioDeviceManager->GetDebugger().IsMuteSoundWave(SoundAssetName);
+
+			if (bIsSolo)
+			{
+				RowTextColor = FColor(255, 255, 0);
+			}
+			else if (bIsMute)
+			{
+				RowTextColor = FColor(255, 0, 0);
+			}
+		}
+
+		return FSlateColor(RowTextColor);
+	}
+
+	void FMixerSourceDashboardViewFactory::ToggleMuteForAllItems(ECheckBoxState NewState)
+	{
+#if ENABLE_AUDIO_DEBUG
+		if (FilteredEntriesListView.IsValid())
+		{
+			if (FAudioDeviceManager* AudioDeviceManager = FAudioDeviceManager::Get())
+			{
+				const TArrayView<const TSharedPtr<IDashboardDataViewEntry>> TableItems = FilteredEntriesListView->GetItems();
+
+				for (const TSharedPtr<IDashboardDataViewEntry>& Item : TableItems)
+				{
+					if (Item.IsValid())
+					{
+						const FSoundAssetDashboardEntry& SoundAssetDashboardEntry = *StaticCastSharedPtr<FSoundAssetDashboardEntry>(Item).Get();
+
+						AudioDeviceManager->GetDebugger().SetMuteSoundWave(FName(SoundAssetDashboardEntry.GetDisplayName().ToString()), NewState == ECheckBoxState::Checked);
+					}
+				}
+			}
+		}
+#endif
+	}
+
+	void FMixerSourceDashboardViewFactory::ToggleSoloForAllItems(ECheckBoxState NewState)
+	{
+#if ENABLE_AUDIO_DEBUG
+		if (FilteredEntriesListView.IsValid())
+		{
+			if (FAudioDeviceManager* AudioDeviceManager = FAudioDeviceManager::Get())
+			{
+				const TArrayView<const TSharedPtr<IDashboardDataViewEntry>> TableItems = FilteredEntriesListView->GetItems();
+
+				for (const TSharedPtr<IDashboardDataViewEntry>& Item : TableItems)
+				{
+					if (Item.IsValid())
+					{
+						const FSoundAssetDashboardEntry& SoundAssetDashboardEntry = *StaticCastSharedPtr<FSoundAssetDashboardEntry>(Item).Get();
+
+						AudioDeviceManager->GetDebugger().SetSoloSoundWave(FName(SoundAssetDashboardEntry.GetDisplayName().ToString()), NewState == ECheckBoxState::Checked);
+					}
+				}
+			}
+		}
+#endif
+	}
+
+	void FMixerSourceDashboardViewFactory::MuteSound()
+	{
+#if ENABLE_AUDIO_DEBUG
+		if (FilteredEntriesListView.IsValid())
+		{
+			if (FAudioDeviceManager* AudioDeviceManager = FAudioDeviceManager::Get())
+			{
+				const TArray<TSharedPtr<IDashboardDataViewEntry>> SelectedItems = FilteredEntriesListView->GetSelectedItems();
+
+				for (const TSharedPtr<IDashboardDataViewEntry>& SelectedItem : SelectedItems)
+				{
+					if (SelectedItem.IsValid())
+					{
+						const FSoundAssetDashboardEntry& SoundAssetDashboardEntry = *StaticCastSharedPtr<FSoundAssetDashboardEntry>(SelectedItem).Get();
+						const FName SoundAssetDisplayName = FName(SoundAssetDashboardEntry.GetDisplayName().ToString());
+
+						AudioDeviceManager->GetDebugger().ToggleMuteSoundWave(SoundAssetDisplayName);
+					}
+				}
+
+				// Handle general Mute button state
+				bool bIsAnySoundMuted = false;
+
+				const TArrayView<const TSharedPtr<IDashboardDataViewEntry>> TableItems = FilteredEntriesListView->GetItems();
+
+				for (const TSharedPtr<IDashboardDataViewEntry>& Item : TableItems)
+				{
+					if (!Item.IsValid())
+					{
+						continue;
+					}
+
+					const FSoundAssetDashboardEntry& SoundAssetDashboardEntry = *StaticCastSharedPtr<FSoundAssetDashboardEntry>(Item).Get();
+					const FName SoundAssetDisplayName = FName(SoundAssetDashboardEntry.GetDisplayName().ToString());
+
+					if (AudioDeviceManager->GetDebugger().IsMuteSoundWave(SoundAssetDisplayName))
+					{
+						bIsAnySoundMuted = true;
+						break;
+					}
+				}
+
+				MuteToggleButton->SetIsChecked(bIsAnySoundMuted ? ECheckBoxState::Checked : ECheckBoxState::Unchecked);
+			}
+		}
+#endif
+	}
+
+	void FMixerSourceDashboardViewFactory::SoloSound()
+	{
+#if ENABLE_AUDIO_DEBUG
+		if (FilteredEntriesListView.IsValid())
+		{
+			if (FAudioDeviceManager* AudioDeviceManager = FAudioDeviceManager::Get())
+			{
+				const TArray<TSharedPtr<IDashboardDataViewEntry>> SelectedItems = FilteredEntriesListView->GetSelectedItems();
+
+				for (const TSharedPtr<IDashboardDataViewEntry>& SelectedItem : SelectedItems)
+				{
+					if (SelectedItem.IsValid())
+					{
+						const FSoundAssetDashboardEntry& SoundAssetDashboardEntry = *StaticCastSharedPtr<FSoundAssetDashboardEntry>(SelectedItem).Get();
+						const FName SoundAssetDisplayName = FName(SoundAssetDashboardEntry.GetDisplayName().ToString());
+
+						AudioDeviceManager->GetDebugger().ToggleSoloSoundWave(SoundAssetDisplayName);
+					}
+				}
+
+				// Handle general Solo button state
+				bool bIsAnySoundSoloed = false;
+
+				const TArrayView<const TSharedPtr<IDashboardDataViewEntry>> TableItems = FilteredEntriesListView->GetItems();
+
+				for (const TSharedPtr<IDashboardDataViewEntry>& Item : TableItems)
+				{
+					if (!Item.IsValid())
+					{
+						continue;
+					}
+
+					const FSoundAssetDashboardEntry& SoundAssetDashboardEntry = *StaticCastSharedPtr<FSoundAssetDashboardEntry>(Item).Get();
+					const FName SoundAssetDisplayName = FName(SoundAssetDashboardEntry.GetDisplayName().ToString());
+
+					if (AudioDeviceManager->GetDebugger().IsSoloSoundWave(SoundAssetDisplayName))
+					{
+						bIsAnySoundSoloed = true;
+						break;
+					}
+				}
+
+				SoloToggleButton->SetIsChecked(bIsAnySoundSoloed ? ECheckBoxState::Checked : ECheckBoxState::Unchecked);
+			}
+		}
+#endif
+	}
+
+	void FMixerSourceDashboardViewFactory::ClearMutesAndSolos()
+	{
+#if ENABLE_AUDIO_DEBUG
+		if (FAudioDeviceManager* AudioDeviceManager = FAudioDeviceManager::Get())
+		{
+			AudioDeviceManager->GetDebugger().ClearMutesAndSolos();
+
+			MuteToggleButton->SetIsChecked(ECheckBoxState::Unchecked);
+			SoloToggleButton->SetIsChecked(ECheckBoxState::Unchecked);
+		}
+#endif
+	}
+
+	void FMixerSourceDashboardViewFactory::UpdatePlotsWidgetsData()
+	{
+		if (DataViewEntries.Num() <= 0)
+		{
+			return;
+		}
+
+		// Process new data
+		bool bHasNewMetadata = false;
+		for (const TSharedPtr<IDashboardDataViewEntry>& DataEntry : DataViewEntries)
+		{
+			const FMixerSourceDashboardEntry& SourceDataPoint = MixerSourcePrivate::CastEntry(*DataEntry);
+			const uint32 SourceId = SourceDataPoint.SourceId; 
+
+			if (SourceId == INDEX_NONE)
+			{
+				break;
+			}
+			
+			// Only add new sources if there are less than the max 
+			const bool bCanAddNewSources = PlotWidgetMetadataPerCurve->Num() < MaxPlotSources;
+
+			const double PointTime = SourceDataPoint.Timestamp;
+			BeginTimestamp = FMath::Min(BeginTimestamp, PointTime);
+			CurrentTimestamp = FMath::Max(CurrentTimestamp, PointTime);
+			const double DataPointTime = SourceDataPoint.Timestamp - BeginTimestamp;
+
+			// For each column, get the array for this data point's source id and add the value to that data array
+			for (auto Iter = PlotWidgetCurveIdToPointDataMapPerColumn.CreateIterator(); Iter; ++Iter)
+			{
+				const FName& ColumnName = Iter.Key();
+				auto DataFunc = GetPlotColumnDataFunc(ColumnName);
+				TSharedPtr<FPointDataPerCurveMap>& DataMap = Iter.Value();
+
+				// Add new data point array
+				if (bCanAddNewSources && !DataMap->Contains(SourceId))
+				{
+					DataMap->Add(SourceId);
+				}
+
+				// Get the data point array for this source id, add new point
+				TArray<FPlotCurvePoint>* DataPoints = DataMap->Find(SourceId);
+				if (DataPoints)
+				{
+					const float Value = (DataFunc)(SourceDataPoint);
+					DataPoints->Emplace(DataPointTime, Value);
+				}
+			}
+
+			// Create metadata for this curve if necessary 
+			if (bCanAddNewSources && !PlotWidgetMetadataPerCurve->Contains(SourceId))
+			{
+				FPlotCurveMetadata& NewMetadata = PlotWidgetMetadataPerCurve->Add(SourceId);
+				NewMetadata.CurveId = SourceId;
+				NewMetadata.CurveColor = FLinearColor(FColor::MakeRandomColor()); 
+				NewMetadata.DisplayName = FText::FromString(*FSoftObjectPath(SourceDataPoint.Name).GetAssetName());
+				bHasNewMetadata = true;
+			}
+		}
+
+		// Set metadata for each widget if updated
+		if (bHasNewMetadata)
+		{
+			for (TSharedPtr<SAudioCurveView> PlotWidget : PlotWidgets)
+			{
+				PlotWidget->SetCurvesMetadata(PlotWidgetMetadataPerCurve);
+			}
+		}
+
+		// Remove points that are older than max history limit from the most recent timestamp
+		const static auto RemoveOldCurvePoints = [this](TSharedPtr<FPointDataPerCurveMap> PlotWidgetPointDataPerCurve)
+		{
+			for (auto Iter = PlotWidgetPointDataPerCurve->CreateIterator(); Iter; ++Iter)
+			{
+				TArray<FPlotCurvePoint>& CurvePoints = Iter->Value;
+				for (int32 i = CurvePoints.Num() - 1; i >= 0; --i)
+				{
+					const FPlotCurvePoint& Point = CurvePoints[i];
+					if (Point.Key + BeginTimestamp < CurrentTimestamp - MaxPlotHistorySeconds)
+					{
+						CurvePoints.RemoveAt(i);
+					}
+				}
+			}
+		};
+
+		// Remove old points and set curve data for each widget 
+		for (int32 WidgetIndex = 0; WidgetIndex < NumPlotWidgets; ++WidgetIndex)
+		{
+			TSharedPtr<SAudioCurveView> PlotWidget = PlotWidgets[WidgetIndex];
+			const FName& SelectedPlotColumn = SelectedPlotColumnNames[WidgetIndex];
+
+			if (TSharedPtr<FPointDataPerCurveMap> CurveData = *PlotWidgetCurveIdToPointDataMapPerColumn.Find(SelectedPlotColumn))
+			{
+				RemoveOldCurvePoints(CurveData);
+				PlotWidget->SetCurvesPointData(CurveData);
+			}
+		}
+	}
+
+	const TMap<FName, FMixerSourceDashboardViewFactory::FPlotColumnInfo>& FMixerSourceDashboardViewFactory::GetPlotColumnInfo()
+	{
+		auto CreatePlotColumnInfo = []()
+		{
+			return TMap<FName, FMixerSourceDashboardViewFactory::FPlotColumnInfo>
+			{
+				{
+					"Amplitude",
+					{
+						[](const IDashboardDataViewEntry& InData) { return MixerSourcePrivate::CastEntry(InData).Envelope; },
+						FSlateStyle::Get().GetAmpFloatFormat()
+					}
+				},
+				{
+					"Volume",
+					{
+
+						[](const IDashboardDataViewEntry& InData) { return MixerSourcePrivate::CastEntry(InData).Volume; },
+						FSlateStyle::Get().GetAmpFloatFormat()
+					}
+				},
+				{
+					"DistanceAttenuation",
+					{
+						[](const IDashboardDataViewEntry& InData) { return MixerSourcePrivate::CastEntry(InData).DistanceAttenuation; }, 
+						FSlateStyle::Get().GetAmpFloatFormat()
+					}
+				},
+				{
+					"Pitch",
+					{
+						[](const IDashboardDataViewEntry& InData) { return MixerSourcePrivate::CastEntry(InData).Pitch; },
+						FSlateStyle::Get().GetPitchFloatFormat()
+					}
+				},
+				{
+					"LPF",
+					{
+						[](const IDashboardDataViewEntry& InData) { return MixerSourcePrivate::CastEntry(InData).LPFFreq; },
+							FSlateStyle::Get().GetFreqFloatFormat() 
+					}
+				},
+				{
+					"HPF",
+					{
+						[](const IDashboardDataViewEntry& InData) { return MixerSourcePrivate::CastEntry(InData).HPFFreq; },
+						FSlateStyle::Get().GetFreqFloatFormat()
+					}
+				}
+			};
+		};
+		static const TMap<FName, FMixerSourceDashboardViewFactory::FPlotColumnInfo> ColumnInfo = CreatePlotColumnInfo();
+		return ColumnInfo;
+	}
+
+	const FNumberFormattingOptions* FMixerSourceDashboardViewFactory::GetPlotColumnNumberFormat(const FName& ColumnName)
+	{
+		if (const FMixerSourceDashboardViewFactory::FPlotColumnInfo* PlotColumnInfo = GetPlotColumnInfo().Find(ColumnName))
+		{
+			return PlotColumnInfo->FormatOptions;
+		}
+		return nullptr;
+	}
+
+	const TFunctionRef<float(const IDashboardDataViewEntry& InData)> FMixerSourceDashboardViewFactory::GetPlotColumnDataFunc(const FName& ColumnName)
+	{
+		return GetPlotColumnInfo().Find(ColumnName)->DataFunc;
+	}
+
+	const FText FMixerSourceDashboardViewFactory::GetPlotColumnDisplayName(const FName& ColumnName)
+	{
+		if (const FTraceTableDashboardViewFactory::FColumnData* ColumnInfo = GetColumns().Find(ColumnName))
+		{
+			return ColumnInfo->DisplayName;
+		}
+		return FText::GetEmpty();
+	}
+
+	TSharedRef<SWidget> FMixerSourceDashboardViewFactory::MakePlotsWidget()
+	{
+		// Initialize column options and initially selected columns 
+		GetPlotColumnInfo().GenerateKeyArray(ColumnNames);
+		if (SelectedPlotColumnNames.IsEmpty() && ColumnNames.Num() > 3)
+		{
+			SelectedPlotColumnNames.Add(ColumnNames[0]); // Amplitude
+			SelectedPlotColumnNames.Add(ColumnNames[3]); // Pitch
+		}
+
+		// Initialize curve data and metadata
+		if (!PlotWidgetMetadataPerCurve.IsValid())
+		{
+			PlotWidgetMetadataPerCurve = MakeShared<TMap<int32, SAudioCurveView::FCurveMetadata>>();
+			for (const FName& ColumnName : ColumnNames)
+			{
+				TSharedPtr<FPointDataPerCurveMap> PointDataPerCurveMap = MakeShared<FPointDataPerCurveMap>();
+				PlotWidgetCurveIdToPointDataMapPerColumn.Emplace(ColumnName, MoveTemp(PointDataPerCurveMap));
+			}
+		}
+
+		// Create plot widgets
+		auto GetViewRange = [this]()
+		{
+			return TRange<double>(FMath::Max(0, CurrentTimestamp - MaxPlotHistorySeconds - BeginTimestamp), CurrentTimestamp - BeginTimestamp);
+		};
+
+		if (PlotWidgets.IsEmpty())
+		{
+			PlotWidgets.AddDefaulted(NumPlotWidgets);
+			for (int32 WidgetNum = 0; WidgetNum < NumPlotWidgets; ++WidgetNum)
+			{
+				SAssignNew(PlotWidgets[WidgetNum], SAudioCurveView)
+					.ViewRange_Lambda(GetViewRange);
+			}
+		}
+
+		// Create plot column combo box widgets  
+		auto CreatePlotColumnComboBoxWidget = [this](int32 PlotWidgetIndex)
+		{
+			return SNew(SComboBox<FName>)
+				.ToolTipText(MixerSourcePrivate::PlotColumnSelectDescription)
+				.OptionsSource(&ColumnNames)
+				.OnGenerateWidget_Lambda([this](const FName& ColumnName)
+				{
+					return SNew(STextBlock)
+					.Font(IPropertyTypeCustomizationUtils::GetRegularFont())
+					.Text(GetPlotColumnDisplayName(ColumnName));
+				})
+				.OnSelectionChanged_Lambda([this, PlotWidgetIndex](FName NewColumnName, ESelectInfo::Type)
+				{
+					SelectedPlotColumnNames[PlotWidgetIndex] = NewColumnName;
+					if (TSharedPtr<FPointDataPerCurveMap>* DataMap = PlotWidgetCurveIdToPointDataMapPerColumn.Find(NewColumnName))
+					{
+						PlotWidgets[PlotWidgetIndex]->SetCurvesPointData(*DataMap);
+						PlotWidgets[PlotWidgetIndex]->SetYValueFormattingOptions(*GetPlotColumnNumberFormat(NewColumnName));
+					}
+				})
+				[
+					SNew(STextBlock)
+					.Font(IPropertyTypeCustomizationUtils::GetRegularFont())
+					.Text_Lambda([this, PlotWidgetIndex]()
+					{
+						return GetPlotColumnDisplayName(SelectedPlotColumnNames[PlotWidgetIndex]);
+					})
+				];
+		};
+
+		return SNew(SVerticalBox)
+			.Clipping(EWidgetClipping::ClipToBounds)
+			+ SVerticalBox::Slot()
+			.AutoHeight()
+			.HAlign(HAlign_Fill)
+			[
+				SNew(SSimpleTimeSlider)
+				.ViewRange_Lambda(GetViewRange)
+				.ClampRangeHighlightSize(0.0f) // Hide clamp range
+			]
+			+ SVerticalBox::Slot()
+			.AutoHeight()
+			.HAlign(HAlign_Fill)
+			[
+				CreatePlotColumnComboBoxWidget(0)
+			]
+			+ SVerticalBox::Slot()
+			.AutoHeight()
+			.HAlign(HAlign_Fill)
+			[
+				PlotWidgets[0].ToSharedRef()
+			]
+			+ SVerticalBox::Slot()
+			.AutoHeight()
+			.HAlign(HAlign_Fill)
+			[
+				CreatePlotColumnComboBoxWidget(1)
+			]
+			+ SVerticalBox::Slot()
+			.AutoHeight()
+			.HAlign(HAlign_Fill)
+			[
+				PlotWidgets[1].ToSharedRef()
+			];
+	}
+
+	TSharedRef<SWidget> FMixerSourceDashboardViewFactory::MakeMuteSoloWidget()
+	{
+		// Mute/Solo labels generation
+		auto GenerateToggleButtonLabelWidget = [](const FText& InLabel = FText::GetEmpty(), const FName& InTextStyle = TEXT("ButtonText")) -> TSharedRef<SWidget>
+		{
+       		TSharedPtr<SHorizontalBox> HBox = SNew(SHorizontalBox);
+			
+       		if (!InLabel.IsEmpty())
+       		{
+       			HBox->AddSlot()	
+				.Padding(0.0f, 0.5f, 0.0f, 0.0f)
+				.AutoWidth()
+				[
+					SNew(STextBlock)
+					.TextStyle( &FAppStyle::Get().GetWidgetStyle<FTextBlockStyle>(InTextStyle))
+					.Justification(ETextJustify::Center)
+					.Text(InLabel)
+				];
+       		}
+
+       		return SNew(SBox)
+				.HeightOverride(16.0f)
+				[
+					HBox.ToSharedRef()
+				];
+		};
+		
+		const FSlateColor WhiteColor(FColor::White);
+
+		// Mute button style
+		MuteToggleButtonStyle = FAppStyle::Get().GetWidgetStyle<FCheckBoxStyle>("ToggleButtonCheckBox");
+		MuteToggleButtonStyle.BorderBackgroundColor = FSlateColor(FColor(200, 0, 0));
+
+		MuteToggleButtonStyle.CheckedHoveredImage.TintColor = WhiteColor;
+		MuteToggleButtonStyle.CheckedImage.TintColor        = WhiteColor;
+		MuteToggleButtonStyle.CheckedPressedImage.TintColor = WhiteColor;
+
+		// Solo button style
+		SoloToggleButtonStyle = FAppStyle::Get().GetWidgetStyle<FCheckBoxStyle>("ToggleButtonCheckBox");
+		SoloToggleButtonStyle.BorderBackgroundColor = FSlateColor(FColor(255, 200, 0));
+
+		SoloToggleButtonStyle.CheckedHoveredImage.TintColor = WhiteColor;
+		SoloToggleButtonStyle.CheckedImage.TintColor        = WhiteColor;
+		SoloToggleButtonStyle.CheckedPressedImage.TintColor = WhiteColor;
+
+		return SNew(SHorizontalBox)
+			+ SHorizontalBox::Slot()
+			.AutoWidth()
+			.Padding(2.0f, 10.0f, 0.0f, 0.0f)
+			[
+				SNew(STextBlock)
+				.Margin(FMargin(0.0, 2.0, 0.0, 0.0))
+				.Text(LOCTEXT("TableDashboardView_GlobalMuteSoloText", "Global Mute/Solo:"))
+			]
+			+ SHorizontalBox::Slot()
+			.MaxWidth(2.0f)
+			.Padding(0.0f, 10.0f, 0.0f, 0.0f)
+			[
+				SNew(SBox)
+			]
+			+ SHorizontalBox::Slot()
+			.AutoWidth()
+			.HAlign(EHorizontalAlignment::HAlign_Center)
+			.Padding(0.0f, 10.0f, 0.0f, 0.0f)
+			[
+				SAssignNew(MuteToggleButton, SCheckBox)
+				.HAlign(EHorizontalAlignment::HAlign_Center)
+				.Style(&MuteToggleButtonStyle)
+				.ToolTip(FSlateApplicationBase::Get().MakeToolTip(LOCTEXT("TableDashboardView_MuteButtonTooltipText", "Mute/Unmute all the items in the list.")))
+				.IsChecked(ECheckBoxState::Unchecked)
+				.OnCheckStateChanged(this, &FMixerSourceDashboardViewFactory::ToggleMuteForAllItems)
+				[
+					GenerateToggleButtonLabelWidget(LOCTEXT("TableDashboardView_MuteButtonText", "M"), "SmallButtonText")
+				]
+			]
+			+ SHorizontalBox::Slot()
+			.MaxWidth(2.0f)
+			.Padding(0.0f, 10.0f, 0.0f, 0.0f)
+			[
+				SNew(SBox)
+			]
+			+ SHorizontalBox::Slot()
+			.AutoWidth()
+			.HAlign(EHorizontalAlignment::HAlign_Center)
+			.Padding(0.0f, 10.0f, 0.0f, 0.0f)
+			[
+				SAssignNew(SoloToggleButton, SCheckBox)
+				.HAlign(EHorizontalAlignment::HAlign_Center)
+				.Style(&SoloToggleButtonStyle)
+				.ToolTip(FSlateApplicationBase::Get().MakeToolTip(LOCTEXT("TableDashboardView_MuteButtonTooltipText", "Enabled/Disable Solo on all the items in the list.")))
+				.IsChecked(ECheckBoxState::Unchecked)
+				.OnCheckStateChanged(this, &FMixerSourceDashboardViewFactory::ToggleSoloForAllItems)
+				[
+					GenerateToggleButtonLabelWidget(LOCTEXT("TableDashboardView_SoloButtonText", "S"), "SmallButtonText")
+				]
+			];
+	}
+
+	TSharedRef<SWidget> FMixerSourceDashboardViewFactory::MakeWidget()
+	{
+		FDashboardFactory::OnActiveAudioDeviceChanged.AddSP(this, &FMixerSourceDashboardViewFactory::ClearMutesAndSolos);
+
+		TSharedRef<SWidget> MuteSoloWidget = MakeMuteSoloWidget();
+		TSharedRef<SWidget> TableDashboardWidget = FTraceTableDashboardViewFactory::MakeWidget();
+		TSharedRef<SWidget> PlotsWidget = MakePlotsWidget();
+
+		return SNew(SVerticalBox)
+			+ SVerticalBox::Slot()
+			.AutoHeight()
+			.HAlign(HAlign_Fill)
+			[
+				MuteSoloWidget
+			]
+			+ SVerticalBox::Slot()
+			.HAlign(HAlign_Fill)
+			.Padding(0.0f, 6.0f, 0.0f, 0.0f)
+			[
+				// Dashboard and plots area
+				SNew(SSplitter)
+				.Orientation(Orient_Vertical)
+				+ SSplitter::Slot()
+				.Value(0.68f)
+				[
+					TableDashboardWidget
+				]
+				+ SSplitter::Slot()
+				.Value(0.32f)
+				[
+					PlotsWidget
+				]
+			];
+	}	
+
+	void FMixerSourceDashboardViewFactory::ProcessEntries(FTraceTableDashboardViewFactory::EProcessReason InReason)
+	{
+		const FString FilterString = GetSearchFilterText().ToString();
+		FTraceTableDashboardViewFactory::FilterEntries<FMixerSourceTraceProvider>([&FilterString](const IDashboardDataViewEntry& Entry)
+		{
+			const FMixerSourceDashboardEntry& MixerSourceEntry = MixerSourcePrivate::CastEntry(Entry);
+			if (MixerSourceEntry.GetDisplayName().ToString().Contains(FilterString))
+			{
+				return false;
+			}
+
+			return true;
+		});
+
+		UpdatePlotsWidgetsData();
+	}
+
+#if WITH_EDITOR
+	bool FMixerSourceDashboardViewFactory::IsDebugDrawEnabled() const
+	{
+		return false;
+	}
+
+	void FMixerSourceDashboardViewFactory::DebugDraw(float InElapsed, const IDashboardDataViewEntry& InEntry, ::Audio::FDeviceId DeviceId) const
+	{
+		// TODO: Get source position if 3d so debug draw works
+// 		const FMixerSourceDashboardEntry& LoopData = static_cast<const FMixerSourceDashboardEntry&>(InEntry);
+// 		const FRotator& Rotator = LoopData.Rotator;
+// 		const FVector& Location = LoopData.Location;
+// 		const FString Description = FString::Printf(TEXT("%s [Virt: %.2fs]"), *LoopData.Name, LoopData.TimeVirtualized);
+// 
+// 		const TArray<UWorld*> Worlds = FAudioDeviceManager::Get()->GetWorldsUsingAudioDevice(DeviceId);
+// 		for (UWorld* World : Worlds)
+// 		{
+// 			DrawDebugSphere(World, Location, 30.0f, 8, FColor::Magenta, false, InElapsed, SDPG_Foreground);
+// 			DrawDebugString(World, Location + FVector(0, 0, 32), *Description, nullptr, FColor::Magenta, InElapsed, false, 1.0f);
+// 		}
+	}
+#endif // WITH_EDITOR
+} // namespace UE::Audio::Insights
+#undef LOCTEXT_NAMESPACE
diff --git a/Engine/Plugins/AudioInsights/Source/AudioInsights/Private/Views/MixerSourceDashboardViewFactory.h b/Engine/Plugins/AudioInsights/Source/AudioInsights/Private/Views/MixerSourceDashboardViewFactory.h
new file mode 100644
--- /dev/null
+++ b/Engine/Plugins/AudioInsights/Source/AudioInsights/Private/Views/MixerSourceDashboardViewFactory.h
@@ -0,0 +1,90 @@
+// Copyright Epic Games, Inc. All Rights Reserved.
+#pragma once
+
+#include "SAudioCurveView.h"
+#include "Views/TableDashboardViewFactory.h"
+#include "Widgets/Input/SCheckBox.h"
+
+namespace UE::Audio::Insights
+{
+	class FMixerSourceDashboardViewFactory : public FTraceObjectTableDashboardViewFactory
+	{
+	public:
+		FMixerSourceDashboardViewFactory();
+		struct FPlotColumnInfo
+		{
+			const TFunctionRef<float(const IDashboardDataViewEntry& InData)> DataFunc;
+			const FNumberFormattingOptions* FormatOptions;
+		};
+
+		virtual ~FMixerSourceDashboardViewFactory() = default;
+
+		virtual FName GetName() const override;
+		virtual FText GetDisplayName() const override;
+		virtual FSlateIcon GetIcon() const override;
+		virtual EDefaultDashboardTabStack GetDefaultTabStack() const override;
+		virtual void ProcessEntries(FTraceTableDashboardViewFactory::EProcessReason Reason) override;
+		virtual TSharedRef<SWidget> MakeWidget() override;
+
+		// Maximum amount of data history kept for plots (in seconds)
+		static const double MaxPlotHistorySeconds;
+		// Maximum number of sources to plot at once 
+		static const int32 MaxPlotSources; 
+
+	protected:
+		virtual const TMap<FName, FTraceTableDashboardViewFactory::FColumnData>& GetColumns() const override;
+		virtual void SortTable() override;
+
+#if WITH_EDITOR
+		virtual bool IsDebugDrawEnabled() const override;
+		virtual void DebugDraw(float InElapsed, const IDashboardDataViewEntry& InEntry, ::Audio::FDeviceId DeviceId) const override;
+#endif // WITH_EDITOR
+		TSharedRef<SWidget> MakePlotsWidget();
+
+	private:
+		void UpdatePlotsWidgetsData();
+		// Column information used by plot widgets, keyed by column name. These keys should be a subset of the keys in GetColumns(). 
+		const TMap<FName, FPlotColumnInfo>& GetPlotColumnInfo();
+		const TFunctionRef<float(const IDashboardDataViewEntry& InData)> GetPlotColumnDataFunc(const FName& ColumnName);
+		const FNumberFormattingOptions* GetPlotColumnNumberFormat(const FName& ColumnName);
+		const FText GetPlotColumnDisplayName(const FName& ColumnName);
+
+		virtual TSharedPtr<SWidget> OnConstructContextMenu() override;
+		virtual FSlateColor GetRowColor(const TSharedPtr<IDashboardDataViewEntry>& InRowDataPtr) override;
+
+		TSharedRef<SWidget> MakeMuteSoloWidget();
+
+		void ToggleMuteForAllItems(ECheckBoxState NewState);
+		void ToggleSoloForAllItems(ECheckBoxState NewState);
+
+		void MuteSound();
+		void SoloSound();
+		void ClearMutesAndSolos();
+
+		using FPlotCurvePoint = SAudioCurveView::FCurvePoint;
+		// Map of source id to data point array 
+		using FPointDataPerCurveMap = TMap<int32, TArray<FPlotCurvePoint>>;
+		using FPlotCurveMetadata = SAudioCurveView::FCurveMetadata;
+
+		FCheckBoxStyle MuteToggleButtonStyle;
+		FCheckBoxStyle SoloToggleButtonStyle;
+
+		TSharedPtr<SCheckBox> MuteToggleButton;
+		TSharedPtr<SCheckBox> SoloToggleButton;
+
+		// Curve points per timestamp per source id per column name 
+		TMap<FName, TSharedPtr<FPointDataPerCurveMap>> PlotWidgetCurveIdToPointDataMapPerColumn;
+		// SourceId to metadata for the corresponding curve
+		TSharedPtr<TMap<int32, FPlotCurveMetadata>> PlotWidgetMetadataPerCurve;
+		
+		// Column names for plot selector widget 
+		TArray<FName> ColumnNames;
+		
+		double BeginTimestamp = TNumericLimits<double>::Max();
+		double CurrentTimestamp = 0;
+
+		const static int32 NumPlotWidgets = 2;
+		TArray<FName> SelectedPlotColumnNames;
+		TArray<TSharedPtr<SAudioCurveView>> PlotWidgets;
+	};
+} // namespace UE::Audio::Insights
diff --git a/Engine/Plugins/AudioInsights/Source/AudioInsights/Private/Views/OutputMeterDashboardViewFactory.cpp b/Engine/Plugins/AudioInsights/Source/AudioInsights/Private/Views/OutputMeterDashboardViewFactory.cpp
new file mode 100644
--- /dev/null
+++ b/Engine/Plugins/AudioInsights/Source/AudioInsights/Private/Views/OutputMeterDashboardViewFactory.cpp
@@ -0,0 +1,91 @@
+// Copyright Epic Games, Inc. All Rights Reserved.
+
+#include "Views/OutputMeterDashboardViewFactory.h"
+
+#include "Analyzers/AudioMeterSubmixAnalyzer.h"
+#include "AudioInsightsDashboardFactory.h"
+#include "AudioInsightsStyle.h"
+#include "AudioMeter.h"
+#include "Sound/AudioSettings.h"
+#include "Sound/SoundSubmix.h"
+#include "SubmixesDashboardViewFactory.h"
+#include "UObject/WeakObjectPtr.h"
+#include "Widgets/Colors/SColorBlock.h"
+#include "Widgets/SBoxPanel.h"
+#include "Widgets/SOverlay.h"
+
+#define LOCTEXT_NAMESPACE "AudioInsights"
+
+namespace UE::Audio::Insights
+{
+	FName FOutputMeterDashboardViewFactory::GetName() const
+	{
+		return "OutputMeter";
+	}
+
+	FText FOutputMeterDashboardViewFactory::GetDisplayName() const
+	{
+		return LOCTEXT("AudioDashboard_DashboardsTab_DisplayName", "Output Meter");
+	}
+
+	EDefaultDashboardTabStack FOutputMeterDashboardViewFactory::GetDefaultTabStack() const
+	{
+		return EDefaultDashboardTabStack::AudioMeter;
+	}
+
+	FSlateIcon FOutputMeterDashboardViewFactory::GetIcon() const
+	{
+		return FSlateStyle::Get().CreateIcon("AudioInsights.Icon");
+	}
+
+	TSharedRef<SWidget> FOutputMeterDashboardViewFactory::MakeWidget()
+	{
+		// Create OutputMeterSubmixAnalyzer (use MainSubmix as default)
+		if (TWeakObjectPtr<UAudioSettings> AudioSettings = GetMutableDefault<UAudioSettings>();
+			AudioSettings.IsValid())
+		{
+			if (MainSubmix = Cast<USoundSubmix>(AudioSettings->MasterSubmix.ResolveObject()))
+			{
+				OutputMeterSubmixAnalyzer = MakeShared<FAudioMeterSubmixAnalyzer>(MainSubmix);
+
+				FDashboardFactory::OnActiveAudioDeviceChanged.AddSP(this, &FOutputMeterDashboardViewFactory::HandleOnActiveAudioDeviceChanged);
+				FSubmixesDashboardViewFactory::OnSubmixSelectionChanged.AddSP(this, &FOutputMeterDashboardViewFactory::HandleOnSubmixSelectionChanged);
+			}
+		}
+
+		return SNew(SOverlay)
+		+ SOverlay::Slot()
+		[
+			SNew(SColorBlock)
+			.Color(FSlateStyle::Get().GetColor("AudioInsights.Analyzers.BackgroundColor"))
+		]
+		+ SOverlay::Slot()
+		[
+			SNew(SVerticalBox)
+			+ SVerticalBox::Slot()
+			.FillHeight(1.0f)
+			.HAlign(HAlign_Center)
+			.VAlign(VAlign_Fill)
+			[
+				OutputMeterSubmixAnalyzer.IsValid() ? OutputMeterSubmixAnalyzer->GetWidget() : SNullWidget::NullWidget
+			]
+		];
+	}
+
+	void FOutputMeterDashboardViewFactory::HandleOnActiveAudioDeviceChanged()
+	{
+		if (OutputMeterSubmixAnalyzer.IsValid())
+		{
+			OutputMeterSubmixAnalyzer->SetSubmix(MainSubmix);
+		}
+	}
+
+	void FOutputMeterDashboardViewFactory::HandleOnSubmixSelectionChanged(const TWeakObjectPtr<USoundSubmix> InSoundSubmix)
+	{
+		if (OutputMeterSubmixAnalyzer.IsValid())
+		{
+			OutputMeterSubmixAnalyzer->SetSubmix(InSoundSubmix);
+		}
+	}
+} // namespace UE::Audio::Insights
+#undef LOCTEXT_NAMESPACE
diff --git a/Engine/Plugins/AudioInsights/Source/AudioInsights/Private/Views/OutputMeterDashboardViewFactory.h b/Engine/Plugins/AudioInsights/Source/AudioInsights/Private/Views/OutputMeterDashboardViewFactory.h
new file mode 100644
--- /dev/null
+++ b/Engine/Plugins/AudioInsights/Source/AudioInsights/Private/Views/OutputMeterDashboardViewFactory.h
@@ -0,0 +1,29 @@
+// Copyright Epic Games, Inc. All Rights Reserved.
+#pragma once
+
+#include "Templates/SharedPointer.h"
+#include "Views/DashboardViewFactory.h"
+
+class USoundSubmix;
+
+namespace UE::Audio::Insights
+{
+	class FAudioMeterSubmixAnalyzer;
+
+	class FOutputMeterDashboardViewFactory : public IDashboardViewFactory, public TSharedFromThis<FOutputMeterDashboardViewFactory>
+	{
+	public:
+		virtual FName GetName() const override;
+		virtual FText GetDisplayName() const override;
+		virtual EDefaultDashboardTabStack GetDefaultTabStack() const override;
+		virtual FSlateIcon GetIcon() const override;
+		virtual TSharedRef<SWidget> MakeWidget() override;
+
+	private:
+		void HandleOnActiveAudioDeviceChanged();
+		void HandleOnSubmixSelectionChanged(const TWeakObjectPtr<USoundSubmix> InSoundSubmix);
+
+		TSharedPtr<FAudioMeterSubmixAnalyzer> OutputMeterSubmixAnalyzer;
+		TObjectPtr<USoundSubmix> MainSubmix;
+	};
+} // namespace UE::Audio::Insights
diff --git a/Engine/Plugins/AudioInsights/Source/AudioInsights/Private/Views/OutputOscilloscopeDashboardViewFactory.cpp b/Engine/Plugins/AudioInsights/Source/AudioInsights/Private/Views/OutputOscilloscopeDashboardViewFactory.cpp
new file mode 100644
--- /dev/null
+++ b/Engine/Plugins/AudioInsights/Source/AudioInsights/Private/Views/OutputOscilloscopeDashboardViewFactory.cpp
@@ -0,0 +1,89 @@
+// Copyright Epic Games, Inc. All Rights Reserved.
+
+#include "Views/OutputOscilloscopeDashboardViewFactory.h"
+
+#include "Analyzers/AudioOscilloscopeAnalyzer.h"
+#include "AudioInsightsDashboardFactory.h"
+#include "AudioInsightsStyle.h"
+#include "AudioOscilloscope.h"
+#include "Sound/AudioSettings.h"
+#include "Sound/SoundSubmix.h"
+#include "SubmixesDashboardViewFactory.h"
+#include "UObject/WeakObjectPtr.h"
+#include "Widgets/Colors/SColorBlock.h"
+
+#define LOCTEXT_NAMESPACE "AudioInsights"
+
+namespace UE::Audio::Insights
+{
+	FName FOutputOscilloscopeDashboardViewFactory::GetName() const
+	{
+		return "OutputOscilloscope";
+	}
+
+	FText FOutputOscilloscopeDashboardViewFactory::GetDisplayName() const
+	{
+		return LOCTEXT("AudioDashboard_DashboardsTab_DisplayName", "Output Oscilloscope");
+	}
+
+	EDefaultDashboardTabStack FOutputOscilloscopeDashboardViewFactory::GetDefaultTabStack() const
+	{
+		return EDefaultDashboardTabStack::Oscilloscope;
+	}
+
+	FSlateIcon FOutputOscilloscopeDashboardViewFactory::GetIcon() const
+	{
+		return FSlateStyle::Get().CreateIcon("AudioInsights.Icon");
+	}
+
+	TSharedRef<SWidget> FOutputOscilloscopeDashboardViewFactory::MakeWidget()
+	{
+		// Create OutputOscilloscopeAnalyzer (use MainSubmix as default)
+		if (TWeakObjectPtr<UAudioSettings> AudioSettings = GetMutableDefault<UAudioSettings>();
+			AudioSettings.IsValid())
+		{
+			if (MainSubmix = Cast<USoundSubmix>(AudioSettings->MasterSubmix.ResolveObject()))
+			{
+				OutputOscilloscopeAnalyzer = MakeShared<FAudioOscilloscopeAnalyzer>(MainSubmix);
+
+				FDashboardFactory::OnActiveAudioDeviceChanged.AddSP(this, &FOutputOscilloscopeDashboardViewFactory::HandleOnActiveAudioDeviceChanged);
+				FSubmixesDashboardViewFactory::OnSubmixSelectionChanged.AddSP(this, &FOutputOscilloscopeDashboardViewFactory::HandleOnSubmixSelectionChanged);
+			}
+		}
+
+		return SNew(SOverlay)
+		+ SOverlay::Slot()
+		[
+			SNew(SColorBlock)
+			.Color(FSlateStyle::Get().GetColor("AudioInsights.Analyzers.BackgroundColor"))
+		]
+		+ SOverlay::Slot()
+		[
+			SNew(SVerticalBox)
+			+ SVerticalBox::Slot()
+			.FillHeight(1.0f)
+			.HAlign(HAlign_Center)
+			.VAlign(VAlign_Fill)
+			[
+				OutputOscilloscopeAnalyzer.IsValid() ? OutputOscilloscopeAnalyzer->GetAudioOscilloscope()->GetPanelWidget() : SNullWidget::NullWidget
+			]
+		];
+	}
+
+	void FOutputOscilloscopeDashboardViewFactory::HandleOnActiveAudioDeviceChanged()
+	{
+		if (OutputOscilloscopeAnalyzer.IsValid())
+		{
+			OutputOscilloscopeAnalyzer->RebuildAudioOscilloscope(MainSubmix);
+		}
+	}
+
+	void FOutputOscilloscopeDashboardViewFactory::HandleOnSubmixSelectionChanged(const TWeakObjectPtr<USoundSubmix> InSoundSubmix)
+	{
+		if (OutputOscilloscopeAnalyzer.IsValid())
+		{
+			OutputOscilloscopeAnalyzer->RebuildAudioOscilloscope(InSoundSubmix);
+		}
+	}
+} // namespace UE::Audio::Insights
+#undef LOCTEXT_NAMESPACE
diff --git a/Engine/Plugins/AudioInsights/Source/AudioInsights/Private/Views/OutputOscilloscopeDashboardViewFactory.h b/Engine/Plugins/AudioInsights/Source/AudioInsights/Private/Views/OutputOscilloscopeDashboardViewFactory.h
new file mode 100644
--- /dev/null
+++ b/Engine/Plugins/AudioInsights/Source/AudioInsights/Private/Views/OutputOscilloscopeDashboardViewFactory.h
@@ -0,0 +1,29 @@
+// Copyright Epic Games, Inc. All Rights Reserved.
+#pragma once
+
+#include "Templates/SharedPointer.h"
+#include "Views/DashboardViewFactory.h"
+
+class USoundSubmix;
+
+namespace UE::Audio::Insights
+{
+	class FAudioOscilloscopeAnalyzer;
+
+	class FOutputOscilloscopeDashboardViewFactory : public IDashboardViewFactory, public TSharedFromThis<FOutputOscilloscopeDashboardViewFactory>
+	{
+	public:
+		virtual FName GetName() const override;
+		virtual FText GetDisplayName() const override;
+		virtual EDefaultDashboardTabStack GetDefaultTabStack() const override;
+		virtual FSlateIcon GetIcon() const override;
+		virtual TSharedRef<SWidget> MakeWidget() override;
+
+	private:
+		void HandleOnActiveAudioDeviceChanged();
+		void HandleOnSubmixSelectionChanged(const TWeakObjectPtr<USoundSubmix> InSoundSubmix);
+
+		TSharedPtr<FAudioOscilloscopeAnalyzer> OutputOscilloscopeAnalyzer;
+		TObjectPtr<USoundSubmix> MainSubmix;
+	};
+} // namespace UE::Audio::Insights
diff --git a/Engine/Plugins/AudioInsights/Source/AudioInsights/Private/Views/SAudioCurveView.cpp b/Engine/Plugins/AudioInsights/Source/AudioInsights/Private/Views/SAudioCurveView.cpp
new file mode 100644
--- /dev/null
+++ b/Engine/Plugins/AudioInsights/Source/AudioInsights/Private/Views/SAudioCurveView.cpp
@@ -0,0 +1,382 @@
+// Copyright Epic Games, Inc. All Rights Reserved.
+
+#include "SAudioCurveView.h"
+#include "Algo/MaxElement.h"
+#include "Algo/MinElement.h"
+#include "Audio.h"
+#include "Fonts/FontMeasure.h"
+#include "Framework/Application/SlateApplication.h"
+#include "Rendering/DrawElements.h"
+#include "Widgets/Text/STextBlock.h"
+#include "Widgets/SToolTip.h"
+
+#define LOCTEXT_NAMESPACE "SAudioCurveView"
+
+void SAudioCurveView::Construct( const SAudioCurveView::FArguments& InArgs )
+{
+	// A lot of this should go into a style 
+    ViewRange = InArgs._ViewRange;
+    GridLineColor = InArgs._GridLineColor;
+    AxesLabelColor = InArgs._AxesLabelColor;
+	YMargin = FMath::Clamp(InArgs._YMargin.Get(), 0.0f, 0.5f);
+	HorizontalAxisIncrement = InArgs._HorizontalAxisIncrement;
+	DesiredSize = InArgs._DesiredSize;
+	YDataRange = FVector2f(0.0f, 1.0f);
+	NumHorizontalGridLines = 10;
+
+	// set clipping on by default, since the OnPaint function is drawing outside the bounds
+	Clipping = EWidgetClipping::ClipToBounds;
+	XValueFormattingOptions.MaximumFractionalDigits = 3;
+	LineDrawEffects = ESlateDrawEffect::NoPixelSnapping;
+	LabelFont = FCoreStyle::GetDefaultFontStyle("Regular", 6);
+
+	SetToolTip(CreateCurveTooltip());
+}
+
+int32 SAudioCurveView::OnPaint( const FPaintArgs& Args, const FGeometry& AllottedGeometry, const FSlateRect& MyCullingRect, FSlateWindowElementList& OutDrawElements, int32 LayerId, const FWidgetStyle& InWidgetStyle, bool bParentEnabled ) const
+{
+	int32 NewLayer = PaintCurves( AllottedGeometry, MyCullingRect, OutDrawElements, LayerId, InWidgetStyle, bParentEnabled );
+
+	return FMath::Max( NewLayer, SCompoundWidget::OnPaint(Args, AllottedGeometry, MyCullingRect, OutDrawElements, NewLayer, InWidgetStyle, ShouldBeEnabled( bParentEnabled ) ) );
+}
+
+FReply SAudioCurveView::OnMouseMove(const FGeometry& InMyGeometry, const FPointerEvent& InMouseEvent)
+{
+	UpdateCurveToolTip(InMyGeometry, InMouseEvent);
+	return FReply::Handled();
+}
+
+TSharedRef<SToolTip> SAudioCurveView::CreateCurveTooltip()
+{
+	return SNew(SToolTip)
+		.Visibility_Lambda([this]()
+		{ 
+			return MetadataPerCurve == nullptr ? EVisibility::Collapsed : EVisibility::Visible; 
+		})
+		.BorderImage(FCoreStyle::Get().GetBrush("ToolTip.BrightBackground"))
+		[
+			SNew(SVerticalBox)
+			+ SVerticalBox::Slot()
+			[
+				SNew(STextBlock)
+				.Text(this, &SAudioCurveView::GetCurveToolTipDisplayNameText)
+				.Font(FCoreStyle::Get().GetFontStyle("ToolTip.LargerFont"))
+				.ColorAndOpacity(FLinearColor::Black)
+			]
+			+ SVerticalBox::Slot()
+			[
+				SNew(STextBlock)
+				.Text(this, &SAudioCurveView::GetCurveToolTipXValueText)
+				.Font(FCoreStyle::Get().GetFontStyle("ToolTip.LargerFont"))
+				.ColorAndOpacity(FLinearColor::Black)
+			]
+			+ SVerticalBox::Slot()
+			[
+				SNew(STextBlock)
+				.Text(this, &SAudioCurveView::GetCurveToolTipYValueText)
+				.Font(FCoreStyle::Get().GetFontStyle("ToolTip.LargerFont"))
+				.ColorAndOpacity(FLinearColor::Black)
+			]
+		];
+}
+
+void SAudioCurveView::UpdateCurveToolTip(const FGeometry& InMyGeometry, const FPointerEvent& InMouseEvent)
+{
+	if (InMyGeometry.IsUnderLocation(InMouseEvent.GetScreenSpacePosition()) && MetadataPerCurve)
+	{
+		// Mouse position in widget space
+		const FVector2f HitPosition = InMyGeometry.AbsoluteToLocal(InMouseEvent.GetScreenSpacePosition());
+
+		// Range helper struct
+		const SSimpleTimeSlider::FScrubRangeToScreen RangeToScreen(ViewRange.Get(), InMyGeometry.GetLocalSize());
+
+		// Mouse position from widget space to curve input space
+		const double TargetX = RangeToScreen.LocalXToInput(HitPosition.X);
+
+		// Keep track of closest curve index, closest point, value 
+		int32 ClosestCurveId = INDEX_NONE;
+		FCurvePoint ClosestPoint;
+		float ClosestDistance = TNumericLimits<float>::Max();
+
+		for (auto Iter = MetadataPerCurve->CreateConstIterator(); Iter; ++Iter)
+		{
+			const FCurveMetadata& CurveMetadata = Iter->Value;
+			const int32 CurveId = CurveMetadata.CurveId;
+			const TArray<FCurvePoint>* CurvePoints = PointDataPerCurve->Find(CurveId);
+			const int32 NumPoints = CurvePoints->Num();
+			if (NumPoints > 0)
+			{
+				for (int32 i = 1; i < NumPoints; ++i)
+				{
+					const FCurvePoint& Point1 = (*CurvePoints)[i - 1];
+					const FCurvePoint& Point2 = (*CurvePoints)[i];
+
+					// Find points that contain mouse hit-point x
+					if (Point1.Key >= TargetX && TargetX <= Point2.Key)
+					{
+						// Choose point with the smallest x delta
+						const float Delta1 = abs(TargetX - Point1.Key);
+						const float Delta2 = abs(TargetX - Point2.Key);
+						const FCurvePoint& TargetPoint = Delta1 < Delta2 ? Point1 : Point2;
+
+						// Convert target point Y to widget space 
+						const FVector2f LocalSize = InMyGeometry.GetLocalSize();
+						const float WidgetSpaceY = ValueToLocalY(LocalSize, TargetPoint.Value);
+						
+						// Compare distance in widget space between HitPosition and closest point by x value on this curve
+						const float Distance = FVector2f::Distance(HitPosition, FVector2f(RangeToScreen.InputToLocalX(TargetPoint.Key), WidgetSpaceY));
+
+						if (Distance < ClosestDistance)
+						{
+							ClosestDistance = Distance;
+							ClosestCurveId = CurveId;
+							ClosestPoint = TargetPoint;
+						}
+						break;
+					}
+				}
+			}
+		}
+
+		// Set tooltip text values  
+		if (ClosestCurveId != INDEX_NONE)
+		{
+			CurveToolTipXValueText = FText::Format(LOCTEXT("CurveToolTipValueXFormat", "X: {0}"), FText::AsNumber(ClosestPoint.Key, &XValueFormattingOptions));
+			CurveToolTipYValueText = FText::Format(LOCTEXT("CurveToolTipValueYFormat", "Y: {0}"), FText::AsNumber(ClosestPoint.Value, &YValueFormattingOptions));
+			CurveToolTipDisplayNameText = MetadataPerCurve->Find(ClosestCurveId)->DisplayName;
+		}
+	}
+}
+
+FVector2D SAudioCurveView::ComputeDesiredSize( float ) const
+{
+	return DesiredSize.Get();
+}
+
+void SAudioCurveView::UpdateYDataRange()
+{
+	const static auto FCurvePointCompare = [](const FCurvePoint& A, const FCurvePoint& B) 
+	{
+		return A.Value < B.Value;
+	};
+
+	float MinValue = TNumericLimits<float>::Max();
+	float MaxValue = TNumericLimits<float>::Min();
+	for (auto Iter = PointDataPerCurve->CreateConstIterator(); Iter; ++Iter)
+	{
+		const TArray<FCurvePoint>& CurvePoints = Iter->Value;
+		if (CurvePoints.IsEmpty())
+		{
+			continue;
+		}
+		const FCurvePoint* MinPoint = Algo::MinElement(CurvePoints, FCurvePointCompare);
+		const FCurvePoint* MaxPoint = Algo::MaxElement(CurvePoints, FCurvePointCompare);
+		MinValue = MinPoint ? FMath::Min(MinValue, MinPoint->Value) : 0.0f;
+		MaxValue = MaxPoint ? FMath::Max(MaxValue, MaxPoint->Value) : 1.0f;
+	}
+	YDataRange.X = MinValue;
+	YDataRange.Y = MaxValue;
+}
+
+void SAudioCurveView::SetCurvesPointData(TSharedPtr<TMap<int32, TArray<FCurvePoint>>> InPointDataPerCurve)
+{
+	PointDataPerCurve = InPointDataPerCurve;
+	UpdateYDataRange();
+}
+
+void SAudioCurveView::SetCurvesMetadata(TSharedPtr<TMap<int32, FCurveMetadata>> InMetadataPerCurve)
+{
+	MetadataPerCurve = InMetadataPerCurve;
+}
+
+void SAudioCurveView::SetYValueFormattingOptions(const FNumberFormattingOptions InYValueFormattingOptions)
+{
+	YValueFormattingOptions = InYValueFormattingOptions;
+}
+
+float SAudioCurveView::ValueToLocalY(const FVector2f AllottedLocalSize, const float Value) const
+{
+	// Slate Y values increase going down the screen, so base < top but base is above top on the screen
+	const float MarginBase = YMargin.Get() * AllottedLocalSize.Y;
+	const float MarginTop = AllottedLocalSize.Y - MarginBase;
+	
+	// Special case to add padding based on YMargin if YDataRange min/max are the same 
+	if (FMath::IsNearlyEqual(YDataRange.X, YDataRange.Y))
+	{
+		const FVector2f PaddedYDataRange = FVector2f(YDataRange.Y * (1.0f - YMargin.Get()), YDataRange.X * (1.0f + YMargin.Get()));
+		return FMath::GetMappedRangeValueUnclamped(PaddedYDataRange, FVector2f(MarginTop, MarginBase), Value);
+	}
+	return FMath::GetMappedRangeValueUnclamped(YDataRange, FVector2f(MarginTop, MarginBase), Value);
+}
+
+float SAudioCurveView::LocalYToValue(const FVector2f AllottedLocalSize, const float LocalY) const
+{
+	// Slate Y values increase going down the screen, so base < top but base is above top on the screen
+	const float MarginBase = YMargin.Get() * AllottedLocalSize.Y;
+	const float MarginTop = AllottedLocalSize.Y - MarginBase;
+
+	// Special case to add padding based on YMargin if YDataRange min/max are the same 
+	if (FMath::IsNearlyEqual(YDataRange.X, YDataRange.Y))
+	{
+		const FVector2f PaddedYDataRange = FVector2f(YDataRange.Y * (1.0f - YMargin.Get()), YDataRange.X * (1.0f + YMargin.Get()));
+		return FMath::GetMappedRangeValueUnclamped(FVector2f(MarginTop, MarginBase), PaddedYDataRange, LocalY);
+	}
+	return FMath::GetMappedRangeValueUnclamped(FVector2f(MarginTop, MarginBase), YDataRange, LocalY);
+}
+
+int32 SAudioCurveView::PaintGridLines(const FGeometry& AllottedGeometry, const FSlateRect& MyCullingRect, FSlateWindowElementList& OutDrawElements, int32 LayerId, const FWidgetStyle& InWidgetStyle, bool bParentEnabled, const SSimpleTimeSlider::FScrubRangeToScreen RangeToScreen) const
+{
+	const FVector2f Size = AllottedGeometry.GetLocalSize();
+	const int32 GridLineLayer = LayerId++;
+	TArray<FVector2f> GridPoints;
+	GridPoints.AddDefaulted(2);
+
+	// Draw vertical grid lines on multiples of HorizontalAxisIncrement
+	if (HorizontalAxisIncrement.Get() > 0)
+	{
+		const double Factor = 1.0 / HorizontalAxisIncrement.Get();
+		// Start at rounded nearest HorizontalAxisIncrement from the lower bound and increment by HorizontalAxisIncrement (ex. 1.5, 2.0, 2.5... if increment is 0.5)
+		double VerticalLineValue = FMath::RoundToDouble(ViewRange.Get().GetLowerBoundValue() * Factor) / Factor;
+		while (VerticalLineValue < ViewRange.Get().GetUpperBoundValue())
+		{
+			const float WidgetX = RangeToScreen.InputToLocalX(VerticalLineValue);
+			GridPoints[0].X = WidgetX;
+			GridPoints[0].Y = 0;
+			GridPoints[1].X = WidgetX;
+			GridPoints[1].Y = Size.Y;
+
+			FSlateDrawElement::MakeLines(
+				OutDrawElements,
+				GridLineLayer,
+				AllottedGeometry.ToPaintGeometry(),
+				GridPoints,
+				LineDrawEffects,
+				GridLineColor.Get(),
+				false
+			);
+			VerticalLineValue += HorizontalAxisIncrement.Get();
+		}
+	}
+
+	// Draw and label NumHorizontalGridLines horizontal grid lines 
+	const int32 GridLineLabelLayer = LayerId++;
+	const float MarginBase = YMargin.Get() * Size.Y;
+	const float MarginTop = (1.0f - YMargin.Get()) * Size.Y;
+	const float GridLineYIncrement = (MarginTop - MarginBase) / (NumHorizontalGridLines - 1);
+	const TSharedRef<FSlateFontMeasure> FontMeasureService = FSlateApplication::Get().GetRenderer()->GetFontMeasureService();
+	for (uint32 HorizontalLineIndex = 0; HorizontalLineIndex < NumHorizontalGridLines; ++HorizontalLineIndex)
+	{
+		const float WidgetY = GridLineYIncrement * HorizontalLineIndex + MarginBase;
+		GridPoints[0].X = 0;
+		GridPoints[0].Y = WidgetY;
+		GridPoints[1].X = Size.X;
+		GridPoints[1].Y = WidgetY;
+
+		FSlateDrawElement::MakeLines(
+			OutDrawElements,
+			GridLineLayer,
+			AllottedGeometry.ToPaintGeometry(),
+			GridPoints,
+			LineDrawEffects,
+			GridLineColor.Get(),
+			false
+		);
+
+		// Draw y axis text label every other grid line
+		if (HorizontalLineIndex % 2 == 1)
+		{
+			const float LabelValue = LocalYToValue(Size, WidgetY);
+			FText LabelString = FText::AsNumber(LabelValue, &YValueFormattingOptions);
+
+			// Position text slightly above the corresponding horizontal line 
+			const FVector2f TextSize = FontMeasureService->Measure(LabelString, LabelFont);
+			FVector2f TextOffset(TextSize.X * 0.5f, WidgetY - TextSize.Y * 0.85f);
+
+			FSlateDrawElement::MakeText(
+				OutDrawElements,
+				GridLineLabelLayer,
+				AllottedGeometry.ToPaintGeometry(TextSize, FSlateLayoutTransform(TextOffset)),
+				LabelString,
+				LabelFont,
+				LineDrawEffects,
+				AxesLabelColor.Get()
+			);
+		}
+	}
+	return LayerId;
+}
+
+int32 SAudioCurveView::PaintCurves(const FGeometry& AllottedGeometry, const FSlateRect& MyCullingRect, FSlateWindowElementList& OutDrawElements, int32 LayerId, const FWidgetStyle& InWidgetStyle, bool bParentEnabled) const 
+{
+	static const float LargeFrameTime = 0.25;  // ideally, we could check the recorded frame data for actual frame length
+	
+	const SSimpleTimeSlider::FScrubRangeToScreen RangeToScreen(ViewRange.Get(), AllottedGeometry.GetLocalSize());
+	const FVector2f Size = AllottedGeometry.GetLocalSize();
+
+	// Skip drawing if curve data is not initialized yet
+	if (!PointDataPerCurve || !MetadataPerCurve)
+	{
+		return LayerId;
+	}
+
+	if (PointDataPerCurve->Num() != MetadataPerCurve->Num())
+	{
+		UE_LOG(LogAudio, Warning, TEXT("Invalid audio curve view data. Metadata and point curve data nums do not match."))
+		return LayerId;
+	}
+
+	LayerId = PaintGridLines(AllottedGeometry, MyCullingRect, OutDrawElements, LayerId, InWidgetStyle, bParentEnabled, RangeToScreen);
+
+	// Create and draw points per curve
+	for (auto Iter = MetadataPerCurve->CreateConstIterator(); Iter; ++Iter)
+	{
+		const FCurveMetadata& CurveMetadata = Iter->Value;
+		const int32 CurveId = CurveMetadata.CurveId;
+		const TArray<FCurvePoint>* CurvePoints = PointDataPerCurve->Find(CurveId);
+		if (!CurvePoints || CurvePoints->Num() <= 0)
+		{
+			continue;
+		}
+
+		TArray<FVector2f> Points;
+		Points.Reserve(CurvePoints->Num());
+
+		float PrevX = (*CurvePoints)[0].Key;
+		for (int32 i = 0; i < CurvePoints->Num(); i++)
+		{
+			const FCurvePoint& Point = (*CurvePoints)[i];
+			if (Point.Value - PrevX > LargeFrameTime && Points.Num()>1)
+			{
+				// break the line list - data has stopped and started again
+				FSlateDrawElement::MakeLines(
+					OutDrawElements,
+					LayerId++,
+					AllottedGeometry.ToPaintGeometry(),
+					Points,
+					LineDrawEffects,
+					CurveMetadata.CurveColor,
+					false
+				);
+				Points.SetNum(0, false);
+			}
+
+			const float X = RangeToScreen.InputToLocalX(Point.Key);
+			PrevX = Point.Key;
+			const float Y = ValueToLocalY(Size, Point.Value);
+			Points.Add(FVector2f(X, Y));
+		}
+
+		FSlateDrawElement::MakeLines(
+			OutDrawElements,
+			LayerId,
+			AllottedGeometry.ToPaintGeometry(),
+			Points,
+			LineDrawEffects,
+			CurveMetadata.CurveColor,
+			false
+		);
+	}
+	return LayerId;
+}
+
+#undef LOCTEXT_NAMESPACE
diff --git a/Engine/Plugins/AudioInsights/Source/AudioInsights/Private/Views/SAudioCurveView.h b/Engine/Plugins/AudioInsights/Source/AudioInsights/Private/Views/SAudioCurveView.h
new file mode 100644
--- /dev/null
+++ b/Engine/Plugins/AudioInsights/Source/AudioInsights/Private/Views/SAudioCurveView.h
@@ -0,0 +1,110 @@
+// Copyright Epic Games, Inc. All Rights Reserved.
+
+#pragma once
+
+#include "SSimpleTimeSlider.h"
+#include "Widgets/SCompoundWidget.h"
+#include "Widgets/DeclarativeSyntaxSupport.h"
+
+class FPaintArgs;
+class FSlateWindowElementList;
+class SToolTip;
+
+// Loosely based on SCurveTimelineView
+// TODO move this to a more general place like the Audio Widgets plugin once it's further developed
+class SAudioCurveView : public SCompoundWidget
+{
+public:
+	// A curve point is a (double XValue, float YValue) pair
+	using FCurvePoint = TPair<double, float>;
+
+	struct FCurveMetadata
+	{
+		int32 CurveId; 
+		FLinearColor CurveColor;
+		FText DisplayName;
+	};
+
+	SLATE_BEGIN_ARGS(SAudioCurveView)
+		: _ViewRange(TRange<double>(0,5))
+		, _YMargin(0.05f)
+		, _HorizontalAxisIncrement(0.5)
+		, _GridLineColor(FLinearColor(0.5f, 0.5f, 0.5f, 0.25f))
+		, _AxesLabelColor(FLinearColor(0.5f, 0.5f, 0.75f))
+		, _DesiredSize(FVector2D(100.f,100.f))
+	{}
+    	/** View X axis range (in value space) */
+    	SLATE_ATTRIBUTE(TRange<double>, ViewRange);
+
+		/** Margin for Y axis, as a 0 - 0.5f proportion, for the space each of above and below the data range. (ex. 0.05 means a 5% margin on the top and bottom, with 90% of the widget's vertical size corresponding to the data range). */
+		SLATE_ATTRIBUTE(float, YMargin);
+
+		/** X axis increment for grid lines. */
+		SLATE_ATTRIBUTE(double, HorizontalAxisIncrement);
+
+		SLATE_ATTRIBUTE(FLinearColor, GridLineColor);
+
+		SLATE_ATTRIBUTE(FLinearColor, AxesLabelColor);
+
+		/** Desired widget size */
+		SLATE_ATTRIBUTE(FVector2D, DesiredSize);
+	
+	SLATE_END_ARGS()
+
+	/**
+	 * Construct the widget
+	 * 
+	 * @param InArgs   A declaration from which to construct the widget
+	 */
+	void Construct( const FArguments& InArgs );
+
+	void SetYValueFormattingOptions(const FNumberFormattingOptions InValueFormattingOptions);
+	void SetCurvesPointData(TSharedPtr<TMap<int32, TArray<FCurvePoint>>> InCurvesPointData);
+	void SetCurvesMetadata(TSharedPtr<TMap<int32, FCurveMetadata>> InMetadataPerCurve);
+
+	FText GetCurveToolTipXValueText() const { return CurveToolTipXValueText; }
+	FText GetCurveToolTipYValueText() const { return CurveToolTipYValueText; }
+	FText GetCurveToolTipDisplayNameText() const { return CurveToolTipDisplayNameText; }
+
+	/** Helper functions for converting between widget local Y position and a given data value (within DataRange). */
+	float ValueToLocalY(const FVector2f AllottedLocalSize, const float Value) const;
+	float LocalYToValue(const FVector2f AllottedLocalSize, const float LocalY) const;
+	
+protected:
+	// SWidget interface
+	virtual FVector2D ComputeDesiredSize(float) const override;
+	virtual int32 OnPaint( const FPaintArgs& Args, const FGeometry& AllottedGeometry, const FSlateRect& MyCullingRect, FSlateWindowElementList& OutDrawElements, int32 LayerId, const FWidgetStyle& InWidgetStyle, bool bParentEnabled ) const override;
+	virtual FReply OnMouseMove(const FGeometry& MyGeometry, const FPointerEvent& MouseEvent) override;
+	
+	int32 PaintCurves(const FGeometry& AllottedGeometry, const FSlateRect& MyCullingRect, FSlateWindowElementList& OutDrawElements, int32 LayerId, const FWidgetStyle& InWidgetStyle, bool bParentEnabled) const;
+	int32 PaintGridLines(const FGeometry& AllottedGeometry, const FSlateRect& MyCullingRect, FSlateWindowElementList& OutDrawElements, int32 LayerId, const FWidgetStyle& InWidgetStyle, bool bParentEnabled, const SSimpleTimeSlider::FScrubRangeToScreen RangeToScreen) const;
+	void UpdateCurveToolTip(const FGeometry& MyGeometry, const FPointerEvent& MouseEvent);
+	void UpdateYDataRange();
+	TSharedRef<SToolTip> CreateCurveTooltip();
+
+	TAttribute<TRange<double>> ViewRange;
+	TAttribute<float> YMargin;
+	TAttribute<double> HorizontalAxisIncrement;
+	TAttribute<FLinearColor> GridLineColor;
+	TAttribute<FLinearColor> AxesLabelColor;
+	TAttribute<FVector2D> DesiredSize;
+
+	// Point data and metadata, keyed by curve id
+	TSharedPtr<TMap<int32, TArray<FCurvePoint>>> PointDataPerCurve;
+	TSharedPtr<TMap<int32, FCurveMetadata>> MetadataPerCurve; 
+	// Y axis data range in value space
+	FVector2f YDataRange;
+
+	// Tooltip text
+	FText CurveToolTipXValueText;		
+	FText CurveToolTipYValueText;
+	FText CurveToolTipDisplayNameText;
+	
+	// Tooltip and axis text formatting
+	FNumberFormattingOptions XValueFormattingOptions;
+	FNumberFormattingOptions YValueFormattingOptions;
+
+	ESlateDrawEffect LineDrawEffects;
+	uint32 NumHorizontalGridLines;
+	FSlateFontInfo LabelFont;
+};
diff --git a/Engine/Plugins/AudioInsights/Source/AudioInsights/Private/Views/SoundSubmixAssetDashboardEntry.cpp b/Engine/Plugins/AudioInsights/Source/AudioInsights/Private/Views/SoundSubmixAssetDashboardEntry.cpp
new file mode 100644
--- /dev/null
+++ b/Engine/Plugins/AudioInsights/Source/AudioInsights/Private/Views/SoundSubmixAssetDashboardEntry.cpp
@@ -0,0 +1,26 @@
+// Copyright Epic Games, Inc. All Rights Reserved.
+
+#include "Views/SoundSubmixAssetDashboardEntry.h"
+
+namespace UE::Audio::Insights
+{
+	FText FSoundSubmixAssetDashboardEntry::GetDisplayName() const
+	{
+		return FText::FromString(FSoftObjectPath(Name).GetAssetName());
+	}
+
+	const UObject* FSoundSubmixAssetDashboardEntry::GetObject() const
+	{
+		return FSoftObjectPath(Name).ResolveObject();
+	}
+
+	UObject* FSoundSubmixAssetDashboardEntry::GetObject()
+	{
+		return FSoftObjectPath(Name).ResolveObject();
+	}
+
+	bool FSoundSubmixAssetDashboardEntry::IsValid() const
+	{
+		return true;
+	}
+} // namespace UE::Audio::Insights
diff --git a/Engine/Plugins/AudioInsights/Source/AudioInsights/Private/Views/SoundSubmixAssetDashboardEntry.h b/Engine/Plugins/AudioInsights/Source/AudioInsights/Private/Views/SoundSubmixAssetDashboardEntry.h
new file mode 100644
--- /dev/null
+++ b/Engine/Plugins/AudioInsights/Source/AudioInsights/Private/Views/SoundSubmixAssetDashboardEntry.h
@@ -0,0 +1,24 @@
+// Copyright Epic Games, Inc. All Rights Reserved.
+#pragma once
+
+#include "Views/TableDashboardViewFactory.h"
+
+class USoundSubmix;
+
+namespace UE::Audio::Insights
+{
+	struct AUDIOINSIGHTS_API FSoundSubmixAssetDashboardEntry : public IObjectDashboardEntry
+	{
+		virtual ~FSoundSubmixAssetDashboardEntry() = default;
+
+		virtual FText GetDisplayName() const override;
+		virtual const UObject* GetObject() const override;
+		virtual UObject* GetObject() override;
+		virtual bool IsValid() const override;
+
+		::Audio::FDeviceId DeviceId = INDEX_NONE;
+		uint64 ComponentId = TNumericLimits<uint64>::Max();
+		FString Name;
+		TWeakObjectPtr<USoundSubmix> SoundSubmix;
+	};
+} // namespace UE::Audio::Insights
diff --git a/Engine/Plugins/AudioInsights/Source/AudioInsights/Private/Views/SubmixesDashboardViewFactory.cpp b/Engine/Plugins/AudioInsights/Source/AudioInsights/Private/Views/SubmixesDashboardViewFactory.cpp
new file mode 100644
--- /dev/null
+++ b/Engine/Plugins/AudioInsights/Source/AudioInsights/Private/Views/SubmixesDashboardViewFactory.cpp
@@ -0,0 +1,232 @@
+// Copyright Epic Games, Inc. All Rights Reserved.
+#include "Views/SubmixesDashboardViewFactory.h"
+
+#include "AssetRegistry/AssetData.h"
+#include "AssetRegistry/AssetRegistryModule.h"
+#include "AudioDefines.h"
+#include "AudioDeviceManager.h"
+#include "AudioInsightsModule.h"
+#include "AudioInsightsStyle.h"
+#include "Editor.h"
+#include "Internationalization/Text.h"
+#include "Providers/SoundSubmixProvider.h"
+#include "Sound/SoundSubmix.h"
+#include "Subsystems/AssetEditorSubsystem.h"
+#include "Templates/SharedPointer.h"
+
+#define LOCTEXT_NAMESPACE "AudioInsights"
+
+namespace UE::Audio::Insights
+{
+	namespace SubmixesPrivate
+	{
+		const FSoundSubmixAssetDashboardEntry& CastEntry(const IDashboardDataViewEntry& InData)
+		{
+			return static_cast<const FSoundSubmixAssetDashboardEntry&>(InData);
+		};
+	}
+
+	FSubmixesDashboardViewFactory::FSubmixesDashboardViewFactory()
+	{
+		FSoundSubmixProvider::OnSubmixAssetAdded.AddRaw(this, &FSubmixesDashboardViewFactory::HandleOnSubmixAssetListUpdated);
+		FSoundSubmixProvider::OnSubmixAssetRemoved.AddRaw(this, &FSubmixesDashboardViewFactory::HandleOnSubmixAssetListUpdated);
+		FSoundSubmixProvider::OnSubmixAssetListUpdated.AddRaw(this, &FSubmixesDashboardViewFactory::RequestListRefresh);
+
+		Providers = TArray<TSharedPtr<FTraceProviderBase>>
+		{
+			MakeShared<FSoundSubmixProvider>()
+		};
+
+		SortByColumn = "Name";
+		SortMode     = EColumnSortMode::Ascending;
+	}
+
+	FSubmixesDashboardViewFactory::~FSubmixesDashboardViewFactory()
+	{
+		FSoundSubmixProvider::OnSubmixAssetAdded.RemoveAll(this);
+		FSoundSubmixProvider::OnSubmixAssetRemoved.RemoveAll(this);
+		FSoundSubmixProvider::OnSubmixAssetListUpdated.RemoveAll(this);
+	}
+
+	FName FSubmixesDashboardViewFactory::GetName() const
+	{
+		return "Submixes";
+	}
+
+	FText FSubmixesDashboardViewFactory::GetDisplayName() const
+	{
+		return LOCTEXT("AudioDashboard_Submixes_DisplayName", "Submixes");
+	}
+
+	TSharedRef<SWidget> FSubmixesDashboardViewFactory::GenerateWidgetForColumn(TSharedRef<IDashboardDataViewEntry> InRowData, const FName& InColumnName)
+	{
+		if (InColumnName == "Name")
+		{
+			const TWeakObjectPtr<USoundSubmix> SoundSubmix = SubmixesPrivate::CastEntry(InRowData.Get()).SoundSubmix;
+
+			SubmixCheckboxCheckedStates.FindOrAdd(SoundSubmix, false);
+
+			const FColumnData& ColumnData = GetColumns()[InColumnName];
+			const FText Label = ColumnData.GetDisplayValue(InRowData.Get());
+
+			return (Label.IsEmpty()) ? 
+				SNullWidget::NullWidget :
+				SNew(SHorizontalBox)
+				.Clipping(EWidgetClipping::ClipToBounds)
+				+ SHorizontalBox::Slot()
+				.FillWidth(0.1f)
+				[
+					SNew(SCheckBox)
+					.IsChecked(SubmixCheckboxCheckedStates.Contains(SoundSubmix) && SubmixCheckboxCheckedStates[SoundSubmix] == true)
+					.OnCheckStateChanged_Lambda([SoundSubmix, this](ECheckBoxState NewState)
+					{
+						if (SubmixCheckboxCheckedStates.Contains(SoundSubmix))
+						{
+							SubmixCheckboxCheckedStates[SoundSubmix] = NewState == ECheckBoxState::Checked;
+						}
+
+						OnSubmixAssetChecked.Broadcast(NewState == ECheckBoxState::Checked, SoundSubmix);
+					})
+				]
+				+ SHorizontalBox::Slot()
+				.FillWidth(0.8f)
+				[
+					SNew(STextBlock)
+					.Text(Label)
+					.MinDesiredWidth(300)
+					.OnDoubleClicked_Lambda([this, InRowData](const FGeometry& MyGeometry, const FPointerEvent& PointerEvent)
+					{
+						if (GEditor)
+						{
+							if (TSharedPtr<IObjectDashboardEntry> ObjectData = StaticCastSharedPtr<IObjectDashboardEntry>(InRowData.ToSharedPtr()); 
+								ObjectData.IsValid())
+							{
+								if (const UObject* Object = ObjectData->GetObject();
+									Object && Object->IsAsset())
+								{
+									GEditor->GetEditorSubsystem<UAssetEditorSubsystem>()->OpenEditorForAsset(Object);
+									return FReply::Handled();
+								}
+							}
+						}
+
+						return FReply::Unhandled();
+					})
+				];
+		}
+
+		return SNullWidget::NullWidget;
+	}
+
+	void FSubmixesDashboardViewFactory::ProcessEntries(FTraceTableDashboardViewFactory::EProcessReason Reason)
+	{
+		const FString FilterString = GetSearchFilterText().ToString();
+		
+		FTraceTableDashboardViewFactory::FilterEntries<FSoundSubmixProvider>([&FilterString](const IDashboardDataViewEntry& Entry)
+		{
+			const FSoundSubmixAssetDashboardEntry& SubmixEntry = static_cast<const FSoundSubmixAssetDashboardEntry&>(Entry);
+			
+			return !SubmixEntry.GetDisplayName().ToString().Contains(FilterString);
+		});
+	}
+
+	FSlateIcon FSubmixesDashboardViewFactory::GetIcon() const
+	{
+		return FSlateStyle::Get().CreateIcon("AudioInsights.Icon.Submix");
+	}
+
+	EDefaultDashboardTabStack FSubmixesDashboardViewFactory::GetDefaultTabStack() const
+	{
+		return EDefaultDashboardTabStack::Analysis;
+	}
+
+	TSharedRef<SWidget> FSubmixesDashboardViewFactory::MakeWidget()
+	{
+		if (!DashboardWidget.IsValid())
+		{
+			DashboardWidget = FTraceTableDashboardViewFactory::MakeWidget();
+
+			if (FilteredEntriesListView.IsValid())
+			{
+				FilteredEntriesListView->SetSelectionMode(ESelectionMode::Single);
+			}
+		}
+
+		return DashboardWidget->AsShared();
+	}
+
+	const TMap<FName, FTraceTableDashboardViewFactory::FColumnData>& FSubmixesDashboardViewFactory::GetColumns() const
+	{
+		auto CreateColumnData = []()
+		{
+			return TMap<FName, FTraceTableDashboardViewFactory::FColumnData>
+			{
+				{
+					"Name",
+					{
+						LOCTEXT("Submixes_NameColumnDisplayName", "Name"),
+						[](const IDashboardDataViewEntry& InData) { return SubmixesPrivate::CastEntry(InData).GetDisplayName(); },
+						false,		/* bDefaultHidden */
+						0.88f,		/* FillWidth */
+						EHorizontalAlignment::HAlign_Left
+					}
+				}
+			};
+		};
+		
+		static const TMap<FName, FTraceTableDashboardViewFactory::FColumnData> ColumnData = CreateColumnData();
+		
+		return ColumnData;
+	}
+
+	void FSubmixesDashboardViewFactory::SortTable()
+	{
+		if (SortByColumn == "Name")
+		{
+			if (SortMode == EColumnSortMode::Ascending)
+			{
+				DataViewEntries.Sort([](const TSharedPtr<IDashboardDataViewEntry>& A, const TSharedPtr<IDashboardDataViewEntry>& B)
+				{
+					const FSoundSubmixAssetDashboardEntry& AData = SubmixesPrivate::CastEntry(*A.Get());
+					const FSoundSubmixAssetDashboardEntry& BData = SubmixesPrivate::CastEntry(*B.Get());
+
+					return AData.GetDisplayName().CompareToCaseIgnored(BData.GetDisplayName()) < 0;
+				});
+			}
+			else if (SortMode == EColumnSortMode::Descending)
+			{
+				DataViewEntries.Sort([](const TSharedPtr<IDashboardDataViewEntry>& A, const TSharedPtr<IDashboardDataViewEntry>& B)
+				{
+					const FSoundSubmixAssetDashboardEntry& AData = SubmixesPrivate::CastEntry(*A.Get());
+					const FSoundSubmixAssetDashboardEntry& BData = SubmixesPrivate::CastEntry(*B.Get());
+
+					return BData.GetDisplayName().CompareToCaseIgnored(AData.GetDisplayName()) < 0;
+				});
+			}
+		}
+	}
+
+	void FSubmixesDashboardViewFactory::OnSelectionChanged(TSharedPtr<IDashboardDataViewEntry> SelectedItem, ESelectInfo::Type SelectInfo)
+	{
+		if (SelectedItem.IsValid())
+		{
+			const FSoundSubmixAssetDashboardEntry& SoundSubmixAssetDashboardEntry = SubmixesPrivate::CastEntry(*SelectedItem.Get());
+
+			OnSubmixSelectionChanged.Broadcast(SoundSubmixAssetDashboardEntry.SoundSubmix);
+		}
+	}
+
+	void FSubmixesDashboardViewFactory::RequestListRefresh()
+	{
+		if (FilteredEntriesListView.IsValid())
+		{
+			FilteredEntriesListView->RequestListRefresh();
+		}
+	}
+
+	void FSubmixesDashboardViewFactory::HandleOnSubmixAssetListUpdated(const TWeakObjectPtr<UObject> InAsset)
+	{
+		RequestListRefresh();
+	}
+} // namespace UE::Audio::Insights
+#undef LOCTEXT_NAMESPACE
diff --git a/Engine/Plugins/AudioInsights/Source/AudioInsights/Private/Views/SubmixesDashboardViewFactory.h b/Engine/Plugins/AudioInsights/Source/AudioInsights/Private/Views/SubmixesDashboardViewFactory.h
new file mode 100644
--- /dev/null
+++ b/Engine/Plugins/AudioInsights/Source/AudioInsights/Private/Views/SubmixesDashboardViewFactory.h
@@ -0,0 +1,47 @@
+// Copyright Epic Games, Inc. All Rights Reserved.
+#pragma once
+
+#include "Delegates/Delegate.h"
+#include "Templates/SharedPointer.h"
+#include "Views/SoundSubmixAssetDashboardEntry.h"
+#include "Views/TableDashboardViewFactory.h"
+#include "Widgets/Input/SCheckBox.h"
+
+class USoundSubmix;
+
+namespace UE::Audio::Insights
+{
+	class FSubmixesDashboardViewFactory : public FTraceObjectTableDashboardViewFactory
+	{
+	public:
+		FSubmixesDashboardViewFactory();
+		virtual ~FSubmixesDashboardViewFactory();
+
+		virtual FName GetName() const override;
+		virtual FText GetDisplayName() const override;
+		virtual FSlateIcon GetIcon() const override;
+		virtual EDefaultDashboardTabStack GetDefaultTabStack() const override;
+		virtual TSharedRef<SWidget> MakeWidget() override;
+
+		DECLARE_MULTICAST_DELEGATE_TwoParams(FOnSubmixAssetChecked, const bool /*bIsChecked*/, TWeakObjectPtr<USoundSubmix> /*SoundSubmix*/);
+		inline static FOnSubmixAssetChecked OnSubmixAssetChecked;
+
+		DECLARE_MULTICAST_DELEGATE_OneParam(FOnSubmixSelectionChanged, const TWeakObjectPtr<USoundSubmix> /*SoundSubmix*/);
+		inline static FOnSubmixSelectionChanged OnSubmixSelectionChanged;
+
+	protected:
+		virtual TSharedRef<SWidget> GenerateWidgetForColumn(TSharedRef<IDashboardDataViewEntry> InRowData, const FName& InColumnName) override;
+		virtual void ProcessEntries(FTraceTableDashboardViewFactory::EProcessReason Reason) override;
+		virtual const TMap<FName, FTraceTableDashboardViewFactory::FColumnData>& GetColumns() const override;
+
+		virtual void SortTable() override;
+
+		virtual void OnSelectionChanged(TSharedPtr<IDashboardDataViewEntry> SelectedItem, ESelectInfo::Type SelectInfo) override;
+
+		void RequestListRefresh();
+
+		void HandleOnSubmixAssetListUpdated(const TWeakObjectPtr<UObject> InAsset);
+
+		TMap<const TWeakObjectPtr<USoundSubmix>, bool> SubmixCheckboxCheckedStates;
+	};
+} // namespace UE::Audio::Insights
diff --git a/Engine/Plugins/AudioInsights/Source/AudioInsights/Private/Views/TableDashboardViewFactory.cpp b/Engine/Plugins/AudioInsights/Source/AudioInsights/Private/Views/TableDashboardViewFactory.cpp
new file mode 100644
--- /dev/null
+++ b/Engine/Plugins/AudioInsights/Source/AudioInsights/Private/Views/TableDashboardViewFactory.cpp
@@ -0,0 +1,592 @@
+// Copyright Epic Games, Inc. All Rights Reserved.
+#include "Views/TableDashboardViewFactory.h"
+
+#include "Algo/Transform.h"
+#include "AudioDeviceManager.h"
+#include "AudioInsightsDashboardAssetCommands.h"
+#include "AudioInsightsTraceProviderBase.h"
+#include "Containers/Array.h"
+#include "DrawDebugHelpers.h"
+#include "Editor.h"
+#include "Framework/Commands/UICommandList.h"
+#include "Sound/SoundCue.h"
+#include "Sound/SoundNodeAttenuation.h"
+#include "Subsystems/AssetEditorSubsystem.h"
+#include "Templates/SharedPointer.h"
+#include "Widgets/SBoxPanel.h"
+
+#define LOCTEXT_NAMESPACE "AudioInsights"
+
+namespace UE::Audio::Insights
+{
+	void FTraceTableDashboardViewFactory::SRowWidget::Construct(const FArguments& InArgs, const TSharedRef<STableViewBase>& InOwnerTable, TSharedPtr<IDashboardDataViewEntry> InData, TSharedRef<FTraceTableDashboardViewFactory> InFactory)
+	{
+		Data = InData;
+		Factory = InFactory;
+		SMultiColumnTableRow<TSharedPtr<IDashboardDataViewEntry>>::Construct(FSuperRowType::FArguments(), InOwnerTable);
+	}
+
+	TSharedRef<SWidget> FTraceTableDashboardViewFactory::SRowWidget::GenerateWidgetForColumn(const FName& Column)
+	{
+		return Factory->GenerateWidgetForColumn(Data->AsShared(), Column);
+	}
+
+	TSharedRef<SWidget> FTraceTableDashboardViewFactory::GenerateWidgetForColumn(TSharedRef<IDashboardDataViewEntry> InRowData, const FName& Column)
+	{
+		const IDashboardDataViewEntry& RowData = InRowData.Get();
+		const FColumnData& ColumnData = GetColumns()[Column];
+		const FText Label = ColumnData.GetDisplayValue(RowData);
+		if (Label.IsEmpty())
+		{
+			return SNullWidget::NullWidget;
+		}
+
+		return SNew(SHorizontalBox)
+		+ SHorizontalBox::Slot().Padding(2)
+		[
+			SNew(STextBlock).Text(Label)
+		];
+	}
+
+	FReply FTraceTableDashboardViewFactory::OnDataRowKeyInput(const FGeometry& Geometry, const FKeyEvent& KeyEvent) const
+	{
+		return FReply::Unhandled();
+	}
+
+	EColumnSortMode::Type FTraceTableDashboardViewFactory::GetColumnSortMode(const FName InColumnId) const
+	{
+		return (SortByColumn == InColumnId) ? SortMode : EColumnSortMode::None;
+	}
+
+	void FTraceTableDashboardViewFactory::RequestSort()
+	{
+		SortTable();
+
+		if (FilteredEntriesListView.IsValid())
+		{
+			FilteredEntriesListView->RequestListRefresh();
+		}
+	}
+
+	void FTraceTableDashboardViewFactory::OnColumnSortModeChanged(const EColumnSortPriority::Type InSortPriority, const FName& InColumnId, const EColumnSortMode::Type InSortMode)
+	{
+		SortByColumn = InColumnId;
+		SortMode = InSortMode;
+		RequestSort();
+	}
+
+	TSharedPtr<SWidget> FTraceTableDashboardViewFactory::OnConstructContextMenu()
+	{
+		return SNullWidget::NullWidget;
+	}
+
+	void FTraceTableDashboardViewFactory::OnSelectionChanged(TSharedPtr<IDashboardDataViewEntry> SelectedItem, ESelectInfo::Type SelectInfo)
+	{
+		// To be optionally implemented by derived classes
+	}
+
+	FSlateColor FTraceTableDashboardViewFactory::GetRowColor(const TSharedPtr<IDashboardDataViewEntry>& InRowDataPtr)
+	{
+		return FSlateColor(FColor(255, 255, 255));
+	}
+
+	TSharedRef<SHeaderRow> FTraceTableDashboardViewFactory::MakeHeaderRowWidget()
+	{
+		TArray<FName> DefaultHiddenColumns;
+		Algo::TransformIf(GetColumns(), DefaultHiddenColumns,
+			[](const TPair<FName, FColumnData>& ColumnInfo) { return ColumnInfo.Value.bDefaultHidden; },
+			[](const TPair<FName, FColumnData>& ColumnInfo) { return ColumnInfo.Key; });
+
+		TSharedRef<SHeaderRow> HeaderRowWidget = SNew(SHeaderRow)
+		.CanSelectGeneratedColumn(true); // Allows for showing/hiding columns
+
+		// This only works if header row columns are added with slots and not programmatically
+		// check in SHeaderRow::Construct: for ( FColumn* const Column : InArgs.Slots ) for more info
+		// A potential alternative would be to delegate to the derived classes the SHeaderRow creation with slots
+		//.HiddenColumnsList(DefaultHiddenColumns);
+
+		for (const auto& [ColumnName, ColumnData] : GetColumns())
+		{
+			const SHeaderRow::FColumn::FArguments ColumnArgs = SHeaderRow::Column(ColumnName)
+				.DefaultLabel(ColumnData.DisplayName)
+				.HAlignCell(ColumnData.Alignment)
+				.FillWidth(ColumnData.FillWidth)
+				.SortMode(this, &FTraceTableDashboardViewFactory::GetColumnSortMode, ColumnName)
+				.OnSort(this, &FTraceTableDashboardViewFactory::OnColumnSortModeChanged);
+
+			// .HiddenColumnsList workaround:
+			// simulate what SHeaderRow::AddColumn( const FColumn::FArguments& NewColumnArgs ) does but allowing us to modify the bIsVisible property
+			// Memory handling (delete) is done by TIndirectArray<FColumn> Columns; defined in SHeaderRow
+			SHeaderRow::FColumn* NewColumn = new SHeaderRow::FColumn(ColumnArgs);
+			NewColumn->bIsVisible = !DefaultHiddenColumns.Contains(ColumnName);
+			HeaderRowWidget->AddColumn(*NewColumn);
+		}
+
+		return HeaderRowWidget;
+	}
+
+	TSharedRef<SWidget> FTraceTableDashboardViewFactory::MakeWidget()
+	{
+		if (!DashboardWidget.IsValid())
+		{
+			FilteredEntriesListView = SNew(SListView<TSharedPtr<IDashboardDataViewEntry>>)
+			.ListItemsSource(&DataViewEntries)
+			.OnContextMenuOpening(this, &FTraceTableDashboardViewFactory::OnConstructContextMenu)
+			.OnSelectionChanged(this, &FTraceTableDashboardViewFactory::OnSelectionChanged)
+			.OnGenerateRow_Lambda([this](TSharedPtr<IDashboardDataViewEntry> Item, const TSharedRef<STableViewBase>& OwnerTable)
+			{
+				return SNew(SRowWidget, OwnerTable, Item, AsShared());
+			})
+			.HeaderRow
+			(
+				MakeHeaderRowWidget()
+			)
+			.SelectionMode(ESelectionMode::Multi)
+			.OnKeyDownHandler_Lambda([this](const FGeometry& Geometry, const FKeyEvent& KeyEvent)
+			{
+				return OnDataRowKeyInput(Geometry, KeyEvent);
+			});
+
+			DashboardWidget = SNew(SVerticalBox)
+			.Clipping(EWidgetClipping::ClipToBounds)
+			+ SVerticalBox::Slot()
+			.AutoHeight()
+			.Padding(0, 2)
+			[
+				SNew(SHorizontalBox)
+				+ SHorizontalBox::Slot()
+				.AutoWidth()
+				[
+					SAssignNew(SearchBoxWidget, SSearchBox)
+					.SelectAllTextWhenFocused(true)
+					.HintText(LOCTEXT("TableDashboardView_SearchBoxHintText", "Filter"))
+					.MinDesiredWidth(100)
+					.OnTextChanged(this, &FTraceTableDashboardViewFactory::SetSearchBoxFilterText)
+				]
+			]
+			+ SVerticalBox::Slot()
+			.FillHeight(1.0f)
+			.Padding(0, 2)
+			[
+				FilteredEntriesListView->AsShared()
+			];
+		}
+
+		return DashboardWidget->AsShared();
+	}
+
+	void FTraceTableDashboardViewFactory::SetSearchBoxFilterText(const FText& NewText)
+	{
+		SearchBoxFilterText = NewText;
+		UpdateFilterReason = EProcessReason::FilterUpdated;
+	}
+
+
+	void FTraceTableDashboardViewFactory::RebuildFilteredEntriesListView()
+	{
+		if (FilteredEntriesListView.IsValid())
+		{
+			FilteredEntriesListView->RebuildList();
+		}
+	}
+
+	const FText& FTraceTableDashboardViewFactory::GetSearchFilterText() const
+	{
+		return SearchBoxFilterText;
+	}
+
+	void FTraceTableDashboardViewFactory::Tick(float InElapsed)
+	{
+		for (TSharedPtr<FTraceProviderBase> Provider : Providers)
+		{
+			const FName ProviderName = Provider->GetName();
+			if (const uint64* CurrentUpdateId = UpdateIds.Find(ProviderName))
+			{
+				const uint64 LastUpdateId = Provider->GetLastUpdateId();
+				if (*CurrentUpdateId != LastUpdateId)
+				{
+					UpdateFilterReason = EProcessReason::EntriesUpdated;
+				}
+			}
+			else
+			{
+				UpdateFilterReason = EProcessReason::EntriesUpdated;
+			}
+		}
+
+		if (UpdateFilterReason != EProcessReason::None)
+		{
+			ProcessEntries(UpdateFilterReason);
+			if (UpdateFilterReason == EProcessReason::EntriesUpdated)
+			{
+				for (TSharedPtr<FTraceProviderBase> Provider : Providers)
+				{
+					const FName ProviderName = Provider->GetName();
+					const uint64 LastUpdateId = Provider->GetLastUpdateId();
+					UpdateIds.FindOrAdd(ProviderName) = LastUpdateId;
+				}
+			}
+
+			RebuildFilteredEntriesListView();
+			UpdateFilterReason = EProcessReason::None;
+		}
+
+#if WITH_EDITOR
+		const bool bDrawDebug = IsDebugDrawEnabled();
+		if (bDrawDebug)
+		{
+			FAudioDeviceManager* Manager = FAudioDeviceManager::Get();
+			if (!Manager)
+			{
+				return;
+			}
+
+			if (!FilteredEntriesListView.IsValid())
+			{
+				return;
+			}
+
+			TArray<TSharedPtr<IDashboardDataViewEntry>> SelectedItems = FilteredEntriesListView->GetSelectedItems();
+			for (const TSharedPtr<IDashboardDataViewEntry>& SelectedEntry : SelectedItems)
+			{
+				if (SelectedEntry.IsValid())
+				{
+					Manager->IterateOverAllDevices([this, &SelectedEntry, InElapsed](::Audio::FDeviceId DeviceId, FAudioDevice* Device)
+					{
+						DebugDraw(InElapsed, *SelectedEntry.Get(), DeviceId);
+					});
+				}
+			}
+		}
+#endif // WITH_EDITOR
+	}
+
+	FTraceTableDashboardViewFactory::FTraceTableDashboardViewFactory()
+	{
+		TickerHandle = FTSTicker::GetCoreTicker().AddTicker(TEXT("TraceTableDashboardViewFactory"), 0.0f, [this](float DeltaTime)
+		{
+			Tick(DeltaTime);
+			return true;
+		});
+	}
+
+	FTraceTableDashboardViewFactory::~FTraceTableDashboardViewFactory()
+	{
+		FTSTicker::GetCoreTicker().RemoveTicker(TickerHandle);
+	}
+
+	FText FSoundAssetDashboardEntry::GetDisplayName() const
+	{
+		return FText::FromString(FSoftObjectPath(Name).GetAssetName());
+	}
+
+	const UObject* FSoundAssetDashboardEntry::GetObject() const
+	{
+		return FSoftObjectPath(Name).ResolveObject();
+	}
+
+	UObject* FSoundAssetDashboardEntry::GetObject()
+	{
+		return FSoftObjectPath(Name).ResolveObject();
+	}
+
+	bool FSoundAssetDashboardEntry::IsValid() const
+	{
+		return PlayOrder != INDEX_NONE;
+	}
+
+	TSharedRef<SWidget> FTraceObjectTableDashboardViewFactory::GenerateWidgetForColumn(TSharedRef<IDashboardDataViewEntry> InRowData, const FName& Column)
+	{
+		const IDashboardDataViewEntry& RowData = InRowData.Get();
+		const FColumnData& ColumnData = GetColumns()[Column];
+		const FText Label = ColumnData.GetDisplayValue(RowData);
+		if (Label.IsEmpty())
+		{
+			return SNullWidget::NullWidget;
+		}
+
+		TSharedPtr<IDashboardDataViewEntry> RowDataPtr = InRowData;
+		return SNew(SHorizontalBox)
+		+ SHorizontalBox::Slot().Padding(2)
+		[
+			SNew(STextBlock).Text(Label)
+			.ColorAndOpacity(GetRowColor(RowDataPtr))
+			.OnDoubleClicked_Lambda([this, RowDataPtr](const FGeometry& MyGeometry, const FPointerEvent& PointerEvent)
+			{
+				if (GEditor)
+				{
+					TSharedPtr<IObjectDashboardEntry> ObjectData = StaticCastSharedPtr<IObjectDashboardEntry>(RowDataPtr);
+					if (ObjectData.IsValid())
+					{
+						const UObject* Object = ObjectData->GetObject();
+						if (Object && Object->IsAsset())
+						{
+							GEditor->GetEditorSubsystem<UAssetEditorSubsystem>()->OpenEditorForAsset(Object);
+							return FReply::Handled();
+						}
+					}
+				}
+
+				return FReply::Unhandled();
+			})
+		];
+	}
+
+	FReply FTraceObjectTableDashboardViewFactory::OnDataRowKeyInput(const FGeometry& Geometry, const FKeyEvent& KeyEvent) const
+	{
+		if (FilteredEntriesListView.IsValid())
+		{
+			TArray<TSharedPtr<IDashboardDataViewEntry>> SelectedItems = FilteredEntriesListView->GetSelectedItems();
+			if (KeyEvent.GetKey() == EKeys::Enter)
+			{
+				for (const TSharedPtr<IDashboardDataViewEntry>& SelectedItem : SelectedItems)
+				{
+					if (SelectedItem.IsValid())
+					{
+						IObjectDashboardEntry& RowData = *StaticCastSharedPtr<IObjectDashboardEntry>(SelectedItem).Get();
+						if (UObject* Object = RowData.GetObject())
+						{
+							if (Object && Object->IsAsset())
+							{
+								GEditor->GetEditorSubsystem<UAssetEditorSubsystem>()->OpenEditorForAsset(Object);
+							}
+						}
+					}
+				}
+
+				return FReply::Handled();
+			}
+		}
+		return FReply::Unhandled();
+	}
+
+	TSharedRef<SWidget> FTraceObjectTableDashboardViewFactory::MakeWidget()
+	{
+		if (!DashboardWidget.IsValid())
+		{
+			DashboardWidget = SNew(SVerticalBox)
+			+ SVerticalBox::Slot()
+			.AutoHeight()
+			.HAlign(HAlign_Fill)
+			[
+				SNew(SHorizontalBox)
+				+ SHorizontalBox::Slot()
+				.AutoWidth()
+				[
+					MakeAssetMenuBar()
+				]
+			]
+			+ SVerticalBox::Slot()
+			.AutoHeight()
+			.HAlign(HAlign_Fill)
+			[
+				FTraceTableDashboardViewFactory::MakeWidget()
+			];
+
+			if (FilteredEntriesListView.IsValid())
+			{
+				FilteredEntriesListView->SetSelectionMode(ESelectionMode::Multi);
+			}
+		}
+
+		return DashboardWidget->AsShared();
+	}
+
+	TSharedRef<SWidget> FTraceObjectTableDashboardViewFactory::MakeAssetMenuBar() const
+	{
+		const FDashboardAssetCommands& Commands = FDashboardAssetCommands::Get();
+		TSharedPtr<FUICommandList> ToolkitCommands = MakeShared<FUICommandList>();
+		ToolkitCommands->MapAction(Commands.GetOpenCommand(), FExecuteAction::CreateLambda([this]() { OpenAsset(); }));
+		ToolkitCommands->MapAction(Commands.GetBrowserSyncCommand(), FExecuteAction::CreateLambda([this]() { BrowseToAsset(); }));
+
+		FToolBarBuilder ToolbarBuilder(ToolkitCommands, FMultiBoxCustomization::None);
+		Commands.AddAssetCommands(ToolbarBuilder);
+
+		return ToolbarBuilder.MakeWidget();
+	}
+
+	TArray<UObject*> FTraceObjectTableDashboardViewFactory::GetSelectedEditableAssets() const
+	{
+		TArray<UObject*> Objects;
+
+		if (!FilteredEntriesListView.IsValid())
+		{
+			return Objects;
+		}
+
+		const TArray<TSharedPtr<IDashboardDataViewEntry>> Items = FilteredEntriesListView->GetSelectedItems();
+		Algo::TransformIf(Items, Objects,
+			[](const TSharedPtr<IDashboardDataViewEntry>& Item)
+			{
+				if (Item.IsValid())
+				{
+					IObjectDashboardEntry& RowData = *StaticCastSharedPtr<IObjectDashboardEntry>(Item).Get();
+					if (UObject* Object = RowData.GetObject())
+					{
+						return Object && Object->IsAsset();
+					}
+				}
+
+				return false;
+			},
+			[](const TSharedPtr<IDashboardDataViewEntry>& Item)
+			{
+				IObjectDashboardEntry& RowData = *StaticCastSharedPtr<IObjectDashboardEntry>(Item).Get();
+				return RowData.GetObject();
+			}
+		);
+
+		return Objects;
+	}
+
+	bool FTraceObjectTableDashboardViewFactory::OpenAsset() const
+	{
+		if (GEditor && FilteredEntriesListView.IsValid())
+		{
+			TArray<UObject*> Objects = GetSelectedEditableAssets();
+			if (UAssetEditorSubsystem* AssetSubsystem = GEditor->GetEditorSubsystem<UAssetEditorSubsystem>())
+			{
+				return AssetSubsystem->OpenEditorForAssets(Objects);
+			}
+		}
+
+		return false;
+	}
+
+	bool FTraceObjectTableDashboardViewFactory::BrowseToAsset() const
+	{
+		if (GEditor)
+		{
+			TArray<UObject*> EditableAssets = GetSelectedEditableAssets();
+			GEditor->SyncBrowserToObjects(EditableAssets);
+			return true;
+		}
+
+		return false;
+	}
+
+	void FSoundAttenuationVisualizer::Draw(float InDeltaTime, const FTransform& InTransform, const UObject& InObject, const UWorld& InWorld) const
+	{
+		if (LastObjectId != InObject.GetUniqueID())
+		{
+			ShapeDetailsMap.Reset();
+			if (const USoundCue* SoundCue = Cast<const USoundCue>(&InObject))
+			{
+				TArray<const USoundNodeAttenuation*> AttenuationNodes;
+				SoundCue->RecursiveFindAttenuation(SoundCue->FirstNode, AttenuationNodes);
+				for (const USoundNodeAttenuation* Node : AttenuationNodes)
+				{
+					if (Node)
+					{
+						if (const FSoundAttenuationSettings* AttenuationSettingsToApply = Node->GetAttenuationSettingsToApply())
+						{
+							AttenuationSettingsToApply->CollectAttenuationShapesForVisualization(ShapeDetailsMap);
+						}
+					}
+				}
+			}
+			else if (const USoundBase* SoundBase = Cast<USoundBase>(&InObject))
+			{
+				if (const FSoundAttenuationSettings* Settings = SoundBase->GetAttenuationSettingsToApply())
+				{
+					Settings->CollectAttenuationShapesForVisualization(ShapeDetailsMap);
+				}
+			}
+			else
+			{
+				return;
+			}
+		}
+		LastObjectId = InObject.GetUniqueID();
+
+		for (const TPair<EAttenuationShape::Type, FBaseAttenuationSettings::AttenuationShapeDetails>& Pair : ShapeDetailsMap)
+		{
+			const FBaseAttenuationSettings::AttenuationShapeDetails& ShapeDetails = Pair.Value;
+			switch (Pair.Key)
+			{
+				case EAttenuationShape::Sphere:
+				{
+					if (ShapeDetails.Falloff > 0.f)
+					{
+						DrawDebugSphere(&InWorld, InTransform.GetTranslation(), ShapeDetails.Extents.X + ShapeDetails.Falloff, 10, Color);
+						DrawDebugSphere(&InWorld, InTransform.GetTranslation(), ShapeDetails.Extents.X, 10, Color);
+					}
+					else
+					{
+						DrawDebugSphere(&InWorld, InTransform.GetTranslation(), ShapeDetails.Extents.X, 10, Color);
+					}
+					break;
+				}
+
+				case EAttenuationShape::Box:
+				{
+					if (ShapeDetails.Falloff > 0.f)
+					{
+						DrawDebugBox(&InWorld, InTransform.GetTranslation(), ShapeDetails.Extents + FVector(ShapeDetails.Falloff), InTransform.GetRotation(), Color);
+						DrawDebugBox(&InWorld, InTransform.GetTranslation(), ShapeDetails.Extents, InTransform.GetRotation(), Color);
+					}
+					else
+					{
+						DrawDebugBox(&InWorld, InTransform.GetTranslation(), ShapeDetails.Extents, InTransform.GetRotation(), Color);
+					}
+					break;
+				}
+
+				case EAttenuationShape::Capsule:
+				{
+					if (ShapeDetails.Falloff > 0.f)
+					{
+						DrawDebugCapsule(&InWorld, InTransform.GetTranslation(), ShapeDetails.Extents.X + ShapeDetails.Falloff, ShapeDetails.Extents.Y + ShapeDetails.Falloff, InTransform.GetRotation(), Color);
+						DrawDebugCapsule(&InWorld, InTransform.GetTranslation(), ShapeDetails.Extents.X, ShapeDetails.Extents.Y, InTransform.GetRotation(), Color);
+					}
+					else
+					{
+						DrawDebugCapsule(&InWorld, InTransform.GetTranslation(), ShapeDetails.Extents.X, ShapeDetails.Extents.Y, InTransform.GetRotation(), Color);
+					}
+					break;
+				}
+
+				case EAttenuationShape::Cone:
+				{
+					const FVector Origin = InTransform.GetTranslation() - (InTransform.GetUnitAxis(EAxis::X) * ShapeDetails.ConeOffset);
+
+					if (ShapeDetails.Falloff > 0.f || ShapeDetails.Extents.Z > 0.f)
+					{
+						const float OuterAngle = FMath::DegreesToRadians(ShapeDetails.Extents.Y + ShapeDetails.Extents.Z);
+						const float InnerAngle = FMath::DegreesToRadians(ShapeDetails.Extents.Y);
+						DrawDebugCone(&InWorld, Origin, InTransform.GetUnitAxis(EAxis::X), ShapeDetails.Extents.X + ShapeDetails.Falloff + ShapeDetails.ConeOffset, OuterAngle, OuterAngle, 10, Color);
+						DrawDebugCone(&InWorld, Origin, InTransform.GetUnitAxis(EAxis::X), ShapeDetails.Extents.X + ShapeDetails.ConeOffset, InnerAngle, InnerAngle, 10, Color);
+					}
+					else
+					{
+						const float Angle = FMath::DegreesToRadians(ShapeDetails.Extents.Y);
+						DrawDebugCone(&InWorld, Origin, InTransform.GetUnitAxis(EAxis::X), ShapeDetails.Extents.X + ShapeDetails.ConeOffset, Angle, Angle, 10, Color);
+					}
+
+					if (!FMath::IsNearlyZero(ShapeDetails.ConeSphereRadius, UE_KINDA_SMALL_NUMBER))
+					{
+						if (ShapeDetails.ConeSphereFalloff > 0.f)
+						{
+
+							DrawDebugSphere(&InWorld, Origin, ShapeDetails.ConeSphereRadius + ShapeDetails.ConeSphereFalloff, 10, Color);
+							DrawDebugSphere(&InWorld, Origin, ShapeDetails.ConeSphereRadius, 10, Color);
+						}
+						else
+						{
+							DrawDebugSphere(&InWorld, Origin, ShapeDetails.ConeSphereRadius, 10, Color);
+						}
+					}
+
+					break;
+				}
+
+				default:
+				{
+					break;
+				}
+			}
+		}
+	}
+} // namespace UE::Audio::Insights
+
+#undef LOCTEXT_NAMESPACE
diff --git a/Engine/Plugins/AudioInsights/Source/AudioInsights/Private/Views/ViewportDashboardViewFactory.h b/Engine/Plugins/AudioInsights/Source/AudioInsights/Private/Views/ViewportDashboardViewFactory.h
new file mode 100644
--- /dev/null
+++ b/Engine/Plugins/AudioInsights/Source/AudioInsights/Private/Views/ViewportDashboardViewFactory.h
@@ -0,0 +1,49 @@
+// Copyright Epic Games, Inc. All Rights Reserved.
+#pragma once
+
+#include "Views/DashboardViewFactory.h"
+#include "Widgets/SNullWidget.h"
+// #include "LevelEditor.h"
+// #include "SLevelViewport.h"
+
+#define LOCTEXT_NAMESPACE "AudioInsights"
+
+
+namespace UE::Audio::Insights
+{
+	class FViewportDashboardViewFactory : public IDashboardViewFactory
+	{
+	public:
+		virtual FName GetName() const override
+		{
+			return "Viewport";
+		}
+
+		virtual FText GetDisplayName() const override
+		{
+			return LOCTEXT("AudioDashboard_ViewportTab_DisplayName", "Viewport");
+		}
+
+		virtual EDefaultDashboardTabStack GetDefaultTabStack() const override
+		{
+			return EDefaultDashboardTabStack::Viewport;
+		}
+
+		virtual FSlateIcon GetIcon() const override
+		{
+			return FSlateStyle::Get().CreateIcon("AudioInsights.Icon.Viewport");
+		}
+
+		virtual TSharedRef<SWidget> MakeWidget() override
+		{
+			// TODO: This test currently hijacks the main level viewport which is not ideal.
+			// Find way to make our own instance and bind to tab here.
+			// FLevelEditorModule& LevelEditor = FModuleManager::GetModuleChecked<FLevelEditorModule>("LevelEditor");
+			// TSharedPtr<SLevelViewport> LevelViewport = LevelEditor.GetFirstActiveLevelViewport();
+			// ViewportContent = LevelViewport->AsShared();
+
+			return SNullWidget::NullWidget;
+		}
+	};
+} // namespace UE::Audio::Insights
+#undef LOCTEXT_NAMESPACE
diff --git a/Engine/Plugins/AudioInsights/Source/AudioInsights/Private/Views/VirtualLoopDashboardViewFactory.cpp b/Engine/Plugins/AudioInsights/Source/AudioInsights/Private/Views/VirtualLoopDashboardViewFactory.cpp
new file mode 100644
--- /dev/null
+++ b/Engine/Plugins/AudioInsights/Source/AudioInsights/Private/Views/VirtualLoopDashboardViewFactory.cpp
@@ -0,0 +1,283 @@
+// Copyright Epic Games, Inc. All Rights Reserved.
+#include "Views/VirtualLoopDashboardViewFactory.h"
+
+#include "Audio/AudioDebug.h"
+#include "AudioDefines.h"
+#include "AudioDeviceManager.h"
+#include "AudioInsightsModule.h"
+#include "AudioInsightsStyle.h"
+#include "DrawDebugHelpers.h"
+#include "Internationalization/Text.h"
+#include "Providers/VirtualLoopTraceProvider.h"
+#include "Templates/SharedPointer.h"
+
+#define LOCTEXT_NAMESPACE "AudioInsights"
+
+
+namespace UE::Audio::Insights
+{
+	namespace VirtualLoopPrivate
+	{
+		const FVirtualLoopDashboardEntry& CastEntry(const IDashboardDataViewEntry& InData)
+		{
+			return static_cast<const FVirtualLoopDashboardEntry&>(InData);
+		};
+	} // namespace VirtualLoopPrivate
+
+	FVirtualLoopDashboardViewFactory::FVirtualLoopDashboardViewFactory()
+		: AttenuationVisualizer(FColor::Blue)
+	{
+		const FTraceModule& TraceModule = FAudioInsightsModule::GetChecked().GetTraceModule();
+		Providers = TArray<TSharedPtr<FTraceProviderBase>>
+		{
+			TraceModule.FindAudioTraceProvider<FVirtualLoopTraceProvider>()
+		};
+	}
+
+	FName FVirtualLoopDashboardViewFactory::GetName() const
+	{
+		return "VirtualLoops";
+	}
+
+	FText FVirtualLoopDashboardViewFactory::GetDisplayName() const
+	{
+		return LOCTEXT("AudioDashboard_VirtualLoops_DisplayName", "Virtual Loops");
+	}
+
+	void FVirtualLoopDashboardViewFactory::ProcessEntries(FTraceTableDashboardViewFactory::EProcessReason Reason)
+	{
+		const FString FilterString = GetSearchFilterText().ToString();
+		FTraceTableDashboardViewFactory::FilterEntries<FVirtualLoopTraceProvider>([&FilterString](const IDashboardDataViewEntry& Entry)
+		{
+			const FVirtualLoopDashboardEntry& VirtualLoopEntry = static_cast<const FVirtualLoopDashboardEntry&>(Entry);
+			if (VirtualLoopEntry.GetDisplayName().ToString().Contains(FilterString))
+			{
+				return false;
+			}
+
+			return true;
+		});
+	}
+
+	FSlateIcon FVirtualLoopDashboardViewFactory::GetIcon() const
+	{
+		return FSlateStyle::Get().CreateIcon("AudioInsights.Icon.VirtualLoop");
+	}
+
+	EDefaultDashboardTabStack FVirtualLoopDashboardViewFactory::GetDefaultTabStack() const
+	{
+		return EDefaultDashboardTabStack::Analysis;
+	}
+
+	const TMap<FName, FTraceTableDashboardViewFactory::FColumnData>& FVirtualLoopDashboardViewFactory::GetColumns() const
+	{
+		auto CreateColumnData = []()
+		{
+			return TMap<FName, FTraceTableDashboardViewFactory::FColumnData>
+			{
+				{
+					"PlayOrder",
+					{
+						LOCTEXT("VirtualLoop_PlayOrderColumnDisplayName", "Play Order"),
+						[](const IDashboardDataViewEntry& InData) { return FText::AsNumber(static_cast<const FVirtualLoopDashboardEntry&>(InData).PlayOrder); },
+						true /* bDefaultHidden */,
+						0.1f /* FillWidth */
+					}
+				},
+				{
+					"Name",
+					{
+						LOCTEXT("VirtualLoop_NameColumnDisplayName", "Name"),
+						[](const IDashboardDataViewEntry& InData) { return VirtualLoopPrivate::CastEntry(InData).GetDisplayName(); },
+						false /* bDefaultHidden */,
+						0.6f  /* FillWidth */
+					}
+				},
+				{
+					"TimeVirtualized",
+					{
+						LOCTEXT("VirtualLoop_VirtualizedTimeColumnDisplayName", "Time (Virtualized)"),
+						[](const IDashboardDataViewEntry& InData) { return FSlateStyle::Get().FormatSecondsAsTime(VirtualLoopPrivate::CastEntry(InData).TimeVirtualized); },
+						false /* bDefaultHidden */,
+						0.15f /* FillWidth */
+					}
+				},
+				{
+					"PlaybackTime",
+					{
+						LOCTEXT("VirtualLoop_TotalTimeColumnDisplayName", "Time (Total)"),
+						[](const IDashboardDataViewEntry& InData) { return FSlateStyle::Get().FormatSecondsAsTime(VirtualLoopPrivate::CastEntry(InData).PlaybackTime); },
+						false /* bDefaultHidden */,
+						0.12f /* FillWidth */
+					}
+				},
+				{
+					"UpdateInterval",
+					{
+						LOCTEXT("VirtualLoop_UpdateIntervalColumnDisplayName", "Update Interval"),
+						[](const IDashboardDataViewEntry& InData) { return FSlateStyle::Get().FormatSecondsAsTime(VirtualLoopPrivate::CastEntry(InData).UpdateInterval); },
+						false /* bDefaultHidden */,
+						0.13f /* FillWidth */
+					}
+				}
+			};
+		};
+		static const TMap<FName, FTraceTableDashboardViewFactory::FColumnData> ColumnData = CreateColumnData();
+		return ColumnData;
+	}
+
+	void FVirtualLoopDashboardViewFactory::SortTable()
+	{
+		if (SortByColumn == "PlayOrder")
+		{
+			if (SortMode == EColumnSortMode::Ascending)
+			{
+				DataViewEntries.Sort([](const TSharedPtr<IDashboardDataViewEntry>& A, const TSharedPtr<IDashboardDataViewEntry>& B)
+				{
+					const FVirtualLoopDashboardEntry& AData = VirtualLoopPrivate::CastEntry(*A.Get());
+					const FVirtualLoopDashboardEntry& BData = VirtualLoopPrivate::CastEntry(*B.Get());
+
+					return AData.PlayOrder < BData.PlayOrder;
+				});
+			}
+			else if (SortMode == EColumnSortMode::Descending)
+			{
+				DataViewEntries.Sort([](const TSharedPtr<IDashboardDataViewEntry>& A, const TSharedPtr<IDashboardDataViewEntry>& B)
+				{
+					const FVirtualLoopDashboardEntry& AData = VirtualLoopPrivate::CastEntry(*A.Get());
+					const FVirtualLoopDashboardEntry& BData = VirtualLoopPrivate::CastEntry(*B.Get());
+
+					return BData.PlayOrder < AData.PlayOrder;
+				});
+			}
+		}
+		else if (SortByColumn == "Name")
+		{
+			if (SortMode == EColumnSortMode::Ascending)
+			{
+				DataViewEntries.Sort([](const TSharedPtr<IDashboardDataViewEntry>& A, const TSharedPtr<IDashboardDataViewEntry>& B)
+				{
+					const FVirtualLoopDashboardEntry& AData = VirtualLoopPrivate::CastEntry(*A.Get());
+					const FVirtualLoopDashboardEntry& BData = VirtualLoopPrivate::CastEntry(*B.Get());
+
+					return AData.GetDisplayName().CompareToCaseIgnored(BData.GetDisplayName()) < 0;
+				});
+			}
+			else if (SortMode == EColumnSortMode::Descending)
+			{
+				DataViewEntries.Sort([](const TSharedPtr<IDashboardDataViewEntry>& A, const TSharedPtr<IDashboardDataViewEntry>& B)
+				{
+					const FVirtualLoopDashboardEntry& AData = VirtualLoopPrivate::CastEntry(*A.Get());
+					const FVirtualLoopDashboardEntry& BData = VirtualLoopPrivate::CastEntry(*B.Get());
+
+					return AData.GetDisplayName().CompareToCaseIgnored(BData.GetDisplayName()) < 0;
+				});
+			}
+		}
+		else if (SortByColumn == "TimeVirtualized")
+		{
+			if (SortMode == EColumnSortMode::Ascending)
+			{
+				DataViewEntries.Sort([](const TSharedPtr<IDashboardDataViewEntry>& A, const TSharedPtr<IDashboardDataViewEntry>& B)
+				{
+					const FVirtualLoopDashboardEntry& AData = VirtualLoopPrivate::CastEntry(*A.Get());
+					const FVirtualLoopDashboardEntry& BData = VirtualLoopPrivate::CastEntry(*B.Get());
+
+					return AData.TimeVirtualized < BData.TimeVirtualized;
+				});
+			}
+			else if (SortMode == EColumnSortMode::Descending)
+			{
+				DataViewEntries.Sort([](const TSharedPtr<IDashboardDataViewEntry>& A, const TSharedPtr<IDashboardDataViewEntry>& B)
+				{
+					const FVirtualLoopDashboardEntry& AData = VirtualLoopPrivate::CastEntry(*A.Get());
+					const FVirtualLoopDashboardEntry& BData = VirtualLoopPrivate::CastEntry(*B.Get());
+
+					return BData.TimeVirtualized < AData.TimeVirtualized;
+				});
+			}
+		}
+		else if (SortByColumn == "PlaybackTime")
+		{
+			if (SortMode == EColumnSortMode::Ascending)
+			{
+				DataViewEntries.Sort([](const TSharedPtr<IDashboardDataViewEntry>& A, const TSharedPtr<IDashboardDataViewEntry>& B)
+				{
+					const FVirtualLoopDashboardEntry& AData = VirtualLoopPrivate::CastEntry(*A.Get());
+					const FVirtualLoopDashboardEntry& BData = VirtualLoopPrivate::CastEntry(*B.Get());
+
+					return AData.PlaybackTime < BData.PlaybackTime;
+				});
+			}
+			else if (SortMode == EColumnSortMode::Descending)
+			{
+				DataViewEntries.Sort([](const TSharedPtr<IDashboardDataViewEntry>& A, const TSharedPtr<IDashboardDataViewEntry>& B)
+				{
+					const FVirtualLoopDashboardEntry& AData = VirtualLoopPrivate::CastEntry(*A.Get());
+					const FVirtualLoopDashboardEntry& BData = VirtualLoopPrivate::CastEntry(*B.Get());
+
+					return BData.PlaybackTime < AData.PlaybackTime;
+				});
+			}
+		}
+		else if (SortByColumn == "UpdateInterval")
+		{
+			if (SortMode == EColumnSortMode::Ascending)
+			{
+				DataViewEntries.Sort([](const TSharedPtr<IDashboardDataViewEntry>& A, const TSharedPtr<IDashboardDataViewEntry>& B)
+				{
+					const FVirtualLoopDashboardEntry& AData = VirtualLoopPrivate::CastEntry(*A.Get());
+					const FVirtualLoopDashboardEntry& BData = VirtualLoopPrivate::CastEntry(*B.Get());
+
+					return AData.UpdateInterval < BData.UpdateInterval;
+				});
+			}
+			else if (SortMode == EColumnSortMode::Descending)
+			{
+				DataViewEntries.Sort([](const TSharedPtr<IDashboardDataViewEntry>& A, const TSharedPtr<IDashboardDataViewEntry>& B)
+				{
+					const FVirtualLoopDashboardEntry& AData = VirtualLoopPrivate::CastEntry(*A.Get());
+					const FVirtualLoopDashboardEntry& BData = VirtualLoopPrivate::CastEntry(*B.Get());
+
+					return BData.UpdateInterval < BData.UpdateInterval;
+				});
+			}
+		}
+	}
+
+#if WITH_EDITOR
+	bool FVirtualLoopDashboardViewFactory::IsDebugDrawEnabled() const
+	{
+		if (FAudioDeviceManager* Manager = FAudioDeviceManager::Get())
+		{
+			return !(Manager->IsVisualizeDebug3dEnabled() && !::Audio::FAudioDebugger::IsVirtualLoopVisualizeEnabled());
+		}
+
+		return false;
+	}
+
+	void FVirtualLoopDashboardViewFactory::DebugDraw(float InElapsed, const IDashboardDataViewEntry& InEntry, ::Audio::FDeviceId DeviceId) const
+	{
+		const FVirtualLoopDashboardEntry& LoopData = VirtualLoopPrivate::CastEntry(InEntry);
+		const FRotator& Rotator = LoopData.Rotator;
+		const FVector& Location = LoopData.Location;
+		const FString Description = FString::Printf(TEXT("%s [Virt: %.2fs]"), *LoopData.Name, LoopData.TimeVirtualized);
+
+		const TArray<UWorld*> Worlds = FAudioDeviceManager::Get()->GetWorldsUsingAudioDevice(DeviceId);
+		for (UWorld* World : Worlds)
+		{
+			DrawDebugSphere(World, Location, 30.0f, 8, AttenuationVisualizer.Color, false, InElapsed, SDPG_Foreground);
+			DrawDebugString(World, Location + FVector(0, 0, 32), *Description, nullptr, AttenuationVisualizer.Color, InElapsed, false, 1.0f);
+
+			if (const UObject* Object = LoopData.GetObject())
+			{
+				FTransform Transform;
+				Transform.SetLocation(Location);
+				Transform.SetRotation(FQuat(Rotator));
+				AttenuationVisualizer.Draw(InElapsed, Transform, *Object, *World);
+			}
+		}
+	}
+#endif // WITH_EDITOR
+} // namespace UE::Audio::Insights
+#undef LOCTEXT_NAMESPACE
diff --git a/Engine/Plugins/AudioInsights/Source/AudioInsights/Private/Views/VirtualLoopDashboardViewFactory.h b/Engine/Plugins/AudioInsights/Source/AudioInsights/Private/Views/VirtualLoopDashboardViewFactory.h
new file mode 100644
--- /dev/null
+++ b/Engine/Plugins/AudioInsights/Source/AudioInsights/Private/Views/VirtualLoopDashboardViewFactory.h
@@ -0,0 +1,34 @@
+// Copyright Epic Games, Inc. All Rights Reserved.
+#pragma once
+
+#include "Views/TableDashboardViewFactory.h"
+
+
+namespace UE::Audio::Insights
+{
+	class FVirtualLoopDashboardViewFactory : public FTraceObjectTableDashboardViewFactory
+	{
+	public:
+		FVirtualLoopDashboardViewFactory();
+		virtual ~FVirtualLoopDashboardViewFactory() = default;
+
+		virtual FName GetName() const override;
+		virtual FText GetDisplayName() const override;
+		virtual FSlateIcon GetIcon() const override;
+		virtual EDefaultDashboardTabStack GetDefaultTabStack() const override;
+
+	protected:
+		virtual void ProcessEntries(FTraceTableDashboardViewFactory::EProcessReason Reason) override;
+		virtual const TMap<FName, FTraceTableDashboardViewFactory::FColumnData>& GetColumns() const override;
+
+		virtual void SortTable() override;
+
+#if WITH_EDITOR
+		virtual bool IsDebugDrawEnabled() const override;
+		virtual void DebugDraw(float InElapsed, const IDashboardDataViewEntry& InEntry, ::Audio::FDeviceId DeviceId) const override;
+#endif // WITH_EDITOR
+
+	private:
+		FSoundAttenuationVisualizer AttenuationVisualizer;
+	};
+} // namespace UE::Audio::Insights
diff --git a/Engine/Plugins/AudioInsights/Source/AudioInsights/Public/AudioInsightsDashboardAssetCommands.h b/Engine/Plugins/AudioInsights/Source/AudioInsights/Public/AudioInsightsDashboardAssetCommands.h
new file mode 100644
--- /dev/null
+++ b/Engine/Plugins/AudioInsights/Source/AudioInsights/Public/AudioInsightsDashboardAssetCommands.h
@@ -0,0 +1,56 @@
+// Copyright Epic Games, Inc. All Rights Reserved.
+#pragma once
+
+#include "Framework/Commands/Commands.h"
+#include "Framework/Commands/UICommandInfo.h"
+#include "Framework/MultiBox/MultiBoxBuilder.h"
+#include "Templates/SharedPointer.h"
+
+
+namespace UE::Audio::Insights
+{
+	class AUDIOINSIGHTS_API FDashboardAssetCommands : public TCommands<FDashboardAssetCommands>
+	{
+	public:
+		FDashboardAssetCommands();
+
+		virtual void RegisterCommands() override;
+
+		virtual void AddAssetCommands(FToolBarBuilder& OutToolbarBuilder) const;
+
+		virtual TSharedPtr<const FUICommandInfo> GetBrowserSyncCommand() const;
+		virtual TSharedPtr<const FUICommandInfo> GetOpenCommand() const;
+
+		virtual TSharedPtr<const FUICommandInfo> GetStartCommand() const;
+		virtual TSharedPtr<const FUICommandInfo> GetStopCommand() const;
+
+		virtual TSharedPtr<const FUICommandInfo> GetMuteCommand() const;
+		virtual TSharedPtr<const FUICommandInfo> GetSoloCommand() const;
+		virtual TSharedPtr<const FUICommandInfo> GetClearMuteSoloCommand() const;
+
+	private:
+		FSlateIcon GetStartIcon() const;
+		FSlateIcon GetStopIcon() const;
+
+		/** Selects the sound in the content browser. */
+		TSharedPtr<FUICommandInfo> BrowserSync;
+
+		/** Opens the sound asset in the content browser. */
+		TSharedPtr<FUICommandInfo> Open;
+
+		/** Starts the trace session & enables the required audio channels on. */
+		TSharedPtr<FUICommandInfo> Start;
+
+		/** Stops the trace session & disables the required audio channels off. */
+		TSharedPtr<FUICommandInfo> Stop;
+
+		/** Mutes the selected item in the Mixer Source View. */
+		TSharedPtr<FUICommandInfo> Mute;
+
+		/** Solo the selected item in the Mixer Source View. */
+		TSharedPtr<FUICommandInfo> Solo;
+
+		/** Clears any Mute/Solo state from any of the items in the Mixer Source View. */
+		TSharedPtr<FUICommandInfo> ClearMuteSolo;
+	};
+} // namespace UE::Audio::Insights
diff --git a/Engine/Plugins/AudioInsights/Source/AudioInsights/Public/AudioInsightsDataSource.h b/Engine/Plugins/AudioInsights/Source/AudioInsights/Public/AudioInsightsDataSource.h
new file mode 100644
--- /dev/null
+++ b/Engine/Plugins/AudioInsights/Source/AudioInsights/Public/AudioInsightsDataSource.h
@@ -0,0 +1,55 @@
+// Copyright Epic Games, Inc. All Rights Reserved.
+#pragma once
+
+#include "Delegates/Delegate.h"
+#include "Templates/SharedPointer.h"
+#include "Trace/Analyzer.h"
+#include "TraceServices/Model/AnalysisSession.h"
+#include "UObject/NameTypes.h"
+
+
+namespace UE::Audio::Insights
+{
+	class AUDIOINSIGHTS_API IDashboardDataViewEntry : public TSharedFromThis<IDashboardDataViewEntry>
+	{
+	public:
+		virtual ~IDashboardDataViewEntry() = default;
+
+		virtual bool IsValid() const = 0;
+	};
+
+	class AUDIOINSIGHTS_API FAudioDataTableView
+	{ 
+	public:
+		FAudioDataTableView()
+			: Entries(MakeShared<TArray<TSharedPtr<IDashboardDataViewEntry>>>())
+		{
+		}
+
+		FAudioDataTableView(TSharedRef<TArray<TSharedPtr<IDashboardDataViewEntry>>> InEntries)
+			: Entries(InEntries)
+		{
+		}
+
+		virtual ~FAudioDataTableView() = default;
+
+		TSharedRef<const TArray<TSharedPtr<IDashboardDataViewEntry>>> GetEntries() const
+		{
+			return Entries;
+		}
+
+		virtual bool FilterViewEntries() = 0;
+
+		virtual void ResetViewEntries()
+		{
+			Entries->Empty();
+			OnViewEntriesUpdatedDelegate.Broadcast();
+		}
+
+		DECLARE_MULTICAST_DELEGATE(FOnViewEntriesUpdatedDelegate);
+		FOnViewEntriesUpdatedDelegate OnViewEntriesUpdatedDelegate;
+
+	protected:
+		TSharedRef<TArray<TSharedPtr<IDashboardDataViewEntry>>> Entries;
+	};
+} // namespace UE::Audio::Insights
diff --git a/Engine/Plugins/AudioInsights/Source/AudioInsights/Public/AudioInsightsLog.h b/Engine/Plugins/AudioInsights/Source/AudioInsights/Public/AudioInsightsLog.h
new file mode 100644
--- /dev/null
+++ b/Engine/Plugins/AudioInsights/Source/AudioInsights/Public/AudioInsightsLog.h
@@ -0,0 +1,8 @@
+// Copyright Epic Games, Inc. All Rights Reserved.
+#pragma once
+
+#include "Logging/LogMacros.h"
+#include "Logging/LogVerbosity.h"
+
+
+AUDIOINSIGHTS_API DECLARE_LOG_CATEGORY_EXTERN(LogAudioInsights, Log, All);
diff --git a/Engine/Plugins/AudioInsights/Source/AudioInsights/Public/AudioInsightsTraceChannelHandle.h b/Engine/Plugins/AudioInsights/Source/AudioInsights/Public/AudioInsightsTraceChannelHandle.h
new file mode 100644
--- /dev/null
+++ b/Engine/Plugins/AudioInsights/Source/AudioInsights/Public/AudioInsightsTraceChannelHandle.h
@@ -0,0 +1,54 @@
+// Copyright Epic Games, Inc. All Rights Reserved.
+#pragma once
+
+#include "Containers/Map.h"
+#include "Templates/SharedPointer.h"
+
+
+namespace UE::Audio::Insights
+{
+	struct FTraceChannelHandle;
+
+	class AUDIOINSIGHTS_API FTraceChannelManager : public TSharedFromThis<FTraceChannelManager>
+	{
+		FCriticalSection CritSect;
+		TMap<FString, uint32> ActiveChannels;
+
+		friend struct FTraceChannelHandle;
+
+	public:
+		FTraceChannelHandle CreateHandle(FString&& InName);
+	};
+
+	struct AUDIOINSIGHTS_API FTraceChannelHandle
+	{
+	private:
+		FString ChannelName;
+		TSharedRef<FTraceChannelManager> Manager;
+
+		void Reset();
+		void Init();
+
+		FTraceChannelHandle(FString&& InName, TSharedRef<FTraceChannelManager> InManager);
+
+	public:
+		FTraceChannelHandle() = default;
+		FTraceChannelHandle(FTraceChannelHandle&&);
+		FTraceChannelHandle(const FTraceChannelHandle&);
+		FTraceChannelHandle& operator=(const FTraceChannelHandle& Other);
+		FTraceChannelHandle& operator=(FTraceChannelHandle&& Other);
+		~FTraceChannelHandle();
+
+		friend FORCEINLINE uint32 GetTypeHash(const FTraceChannelHandle& Handle)
+		{
+			return GetTypeHash(Handle.ChannelName);
+		}
+
+		bool operator==(const FTraceChannelHandle& Other) const;
+		bool operator!=(const FTraceChannelHandle& Other) const;
+
+		bool IsValid() const;
+
+	friend class FTraceChannelManager;
+	};
+} // namespace UE::Audio::Insights
\ No newline at end of file
diff --git a/Engine/Plugins/AudioInsights/Source/AudioInsights/Public/AudioInsightsTraceProviderBase.h b/Engine/Plugins/AudioInsights/Source/AudioInsights/Public/AudioInsightsTraceProviderBase.h
new file mode 100644
--- /dev/null
+++ b/Engine/Plugins/AudioInsights/Source/AudioInsights/Public/AudioInsightsTraceProviderBase.h
@@ -0,0 +1,211 @@
+// Copyright Epic Games, Inc. All Rights Reserved.
+#pragma once
+
+#include <atomic>
+#include "AudioDefines.h"
+#include "Containers/Ticker.h"
+#include "IAudioInsightsModule.h"
+#include "Messages/AnalyzerMessageQueue.h"
+#include "Modules/ModuleManager.h"
+#include "Templates/SharedPointer.h"
+#include "Trace/Analyzer.h"
+#include "TraceServices/Model/AnalysisSession.h"
+#include "UObject/NameTypes.h"
+
+
+namespace UE::Audio::Insights
+{
+	class AUDIOINSIGHTS_API FTraceProviderBase
+		: public TraceServices::IEditableProvider
+	{
+	public:
+		FTraceProviderBase(FName InName);
+		virtual ~FTraceProviderBase() = default;
+
+		virtual Trace::IAnalyzer* ConstructAnalyzer() = 0;
+		FName GetName() const;
+
+		virtual void Reset()
+		{
+			LastUpdateId = 0;
+			LastMessageId = 0;
+		}
+
+		virtual bool ProcessMessages()
+		{
+			LastUpdateId = LastMessageId;
+			return true;
+		}
+
+		uint64 GetLastUpdateId() const
+		{
+			return LastUpdateId;
+		}
+
+		bool IsUpdated() const
+		{
+			return GetLastMessageId() == LastUpdateId;
+		}
+
+	protected:
+		class AUDIOINSIGHTS_API FTraceAnalyzerBase : public Trace::IAnalyzer
+		{
+		public:
+			FTraceAnalyzerBase(TSharedRef<FTraceProviderBase> InProvider);
+			virtual ~FTraceAnalyzerBase() = default;
+
+			virtual void OnAnalysisBegin(const FOnAnalysisContext& Context) override;
+
+		protected:
+			virtual bool OnEventSuccess(uint16 RouteId, EStyle Style, const FOnEventContext& Context);
+			virtual bool OnEventFailure(uint16 RouteId, EStyle Style, const FOnEventContext& Context);
+
+			template <typename TProviderType>
+			TProviderType& GetProvider()
+			{
+				return *StaticCastSharedRef<TProviderType>(Provider);
+			}
+
+		private:
+			TSharedRef<FTraceProviderBase> Provider;
+		};
+
+		uint64 GetLastMessageId() const
+		{
+			return LastMessageId;
+		}
+
+		uint64 LastUpdateId = 0;
+
+	private:
+		std::atomic<uint64> LastMessageId { 0 };
+		FName Name;
+
+		friend class FTraceAnalyzerBase;
+	};
+
+	template<typename EntryKey, typename EntryType /* = IDashboardDataViewEntry */>
+	class TDeviceDataMapTraceProvider : public FTraceProviderBase
+	{
+	public:
+		using KeyType = EntryKey;
+		using ValueType = EntryType;
+
+		using FDeviceData = TSortedMap<EntryKey, EntryType>;
+		using FEntryPair = TPair<EntryKey, EntryType>;
+
+		TDeviceDataMapTraceProvider(FName InName)
+			: FTraceProviderBase(InName)
+		{
+			TickerHandle = FTSTicker::GetCoreTicker().AddTicker(*InName.ToString(), 0.0f, [this](float DeltaTime)
+			{
+				if (LastUpdateId != GetLastMessageId())
+				{
+					ProcessMessages();
+				}
+				LastUpdateId = GetLastMessageId();
+				return true;
+			});
+		}
+
+		virtual ~TDeviceDataMapTraceProvider()
+		{
+			FTSTicker::GetCoreTicker().RemoveTicker(TickerHandle);
+		}
+
+		const TMap<::Audio::FDeviceId, FDeviceData>& GetDeviceDataMap() const
+		{
+			return DeviceDataMap;
+		}
+
+		const FDeviceData* FindFilteredDeviceData() const
+		{
+			const IAudioInsightsModule& InsightsModule = FModuleManager::GetModuleChecked<IAudioInsightsModule>(IAudioInsightsModule::GetName());
+			const ::Audio::FDeviceId DeviceId = InsightsModule.GetDeviceId();
+
+			return DeviceDataMap.Find(DeviceId);
+		}
+
+	protected:
+		using Super = TDeviceDataMapTraceProvider<EntryKey, EntryType>;
+
+		TMap<::Audio::FDeviceId, FDeviceData> DeviceDataMap;
+
+		virtual void Reset() override
+		{
+			DeviceDataMap.Empty();
+			FTraceProviderBase::Reset();
+		}
+
+		template <typename TMsgType>
+		void ProcessMessageQueue(
+			TAnalyzerMessageQueue<TMsgType>& InQueue,
+			TFunctionRef<EntryType*(const TMsgType&)> GetEntry,
+			TFunctionRef<void(const TMsgType&, EntryType*)> ProcessEntry
+		)
+		{
+			TArray<TMsgType> Messages = InQueue.DequeueAll();
+			for (const TMsgType& Msg : Messages)
+			{
+				EntryType* Entry = GetEntry(Msg);
+				ProcessEntry(Msg, Entry);
+			};
+		}
+
+		EntryType* FindDeviceEntry(::Audio::FDeviceId InDeviceId, const EntryKey& InKey)
+		{
+			if (FDeviceData* Entry = DeviceDataMap.Find(InDeviceId))
+			{
+				return Entry->Find(InKey);
+			}
+
+			return nullptr;
+		}
+
+		const EntryType* FindDeviceEntry(::Audio::FDeviceId InDeviceId, const EntryKey& InKey) const
+		{
+			if (const FDeviceData* Entry = DeviceDataMap.Find(InDeviceId))
+			{
+				return Entry->Find(InKey);
+			}
+
+			return nullptr;
+		}
+
+		FDeviceData* FindFilteredDeviceData()
+		{
+			const IAudioInsightsModule& InsightsModule = FModuleManager::GetModuleChecked<IAudioInsightsModule>(IAudioInsightsModule::GetName());
+			const ::Audio::FDeviceId DeviceId = InsightsModule.GetDeviceId();
+
+			return DeviceDataMap.Find(DeviceId);
+		}
+
+		bool RemoveDeviceEntry(::Audio::FDeviceId InDeviceId, const EntryKey& InKey)
+		{
+			if (FDeviceData* DeviceData = DeviceDataMap.Find(InDeviceId))
+			{
+				if (DeviceData->Remove(InKey) > 0)
+				{
+					if (DeviceData->IsEmpty())
+					{
+						DeviceDataMap.Remove(InDeviceId);
+					}
+
+					return true;
+				}
+			}
+
+			return false;
+		}
+
+		void UpdateDeviceEntry(::Audio::FDeviceId InDeviceId, const EntryKey& InKey, TFunctionRef<void(EntryType&)> InEntryMutator)
+		{
+			FDeviceData& DeviceData = DeviceDataMap.FindOrAdd(InDeviceId);
+			EntryType& Entry = DeviceData.FindOrAdd(InKey);
+			InEntryMutator(Entry);
+		}
+
+	private:
+		FTSTicker::FDelegateHandle TickerHandle;
+	};
+} // namespace UE::Audio::Insights
diff --git a/Engine/Plugins/AudioInsights/Source/AudioInsights/Public/IAudioInsightsDashboardFactory.h b/Engine/Plugins/AudioInsights/Source/AudioInsights/Public/IAudioInsightsDashboardFactory.h
new file mode 100644
--- /dev/null
+++ b/Engine/Plugins/AudioInsights/Source/AudioInsights/Public/IAudioInsightsDashboardFactory.h
@@ -0,0 +1,20 @@
+// Copyright Epic Games, Inc. All Rights Reserved.
+#pragma once
+
+#include "AudioDefines.h"
+#include "UObject/NameTypes.h"
+#include "Views/DashboardViewFactory.h"
+
+
+namespace UE::Audio::Insights
+{
+	class AUDIOINSIGHTS_API IDashboardFactory
+	{
+	public:
+		virtual ~IDashboardFactory() = default;
+
+		virtual void RegisterViewFactory(TSharedRef<IDashboardViewFactory> InFactory) = 0;
+		virtual void UnregisterViewFactory(FName InDashboardName) = 0;
+		virtual ::Audio::FDeviceId GetDeviceId() const = 0;
+	};
+} // namespace UE::Audio::Insights
diff --git a/Engine/Plugins/AudioInsights/Source/AudioInsights/Public/IAudioInsightsModule.h b/Engine/Plugins/AudioInsights/Source/AudioInsights/Public/IAudioInsightsModule.h
new file mode 100644
--- /dev/null
+++ b/Engine/Plugins/AudioInsights/Source/AudioInsights/Public/IAudioInsightsModule.h
@@ -0,0 +1,26 @@
+// Copyright Epic Games, Inc. All Rights Reserved.
+#pragma once
+
+#include "Modules/ModuleInterface.h"
+#include "Templates/SharedPointer.h"
+#include "UObject/NameTypes.h"
+
+namespace UE::Audio::Insights
+{
+	class IDashboardViewFactory;
+} // namespace UE::Audio::Insights
+
+class AUDIOINSIGHTS_API IAudioInsightsModule : public IModuleInterface
+{
+public:
+	virtual void RegisterDashboardViewFactory(TSharedRef<UE::Audio::Insights::IDashboardViewFactory> InDashboardFactory) = 0;
+	virtual void UnregisterDashboardViewFactory(FName InName) = 0;
+
+	virtual ::Audio::FDeviceId GetDeviceId() const = 0;
+
+	static FName GetName()
+	{
+		const FLazyName ModuleName = "AudioInsights";
+		return ModuleName.Resolve();
+	}
+};
diff --git a/Engine/Plugins/AudioInsights/Source/AudioInsights/Public/Messages/AnalyzerMessageQueue.h b/Engine/Plugins/AudioInsights/Source/AudioInsights/Public/Messages/AnalyzerMessageQueue.h
new file mode 100644
--- /dev/null
+++ b/Engine/Plugins/AudioInsights/Source/AudioInsights/Public/Messages/AnalyzerMessageQueue.h
@@ -0,0 +1,96 @@
+// Copyright Epic Games, Inc. All Rights Reserved.
+#pragma once
+
+#include "Containers/MpscQueue.h"
+#include "Math/UnrealMathUtility.h"
+#include "Misc/ScopeLock.h"
+
+
+namespace UE::Audio::Insights
+{
+	template <typename T>
+	struct TAnalyzerMessageQueue
+	{
+		const double MaxHistoryLimitSec = 5.0;
+
+		TAnalyzerMessageQueue(double InHistoryLimitSec = MaxHistoryLimitSec)
+			: HistoryLimitSec(FMath::Clamp(InHistoryLimitSec, UE_DOUBLE_KINDA_SMALL_NUMBER, MaxHistoryLimitSec))
+		{
+		}
+
+	private:
+		double HistoryLimitSec = 0.0;
+		TMpscQueue<T> Data;
+
+		mutable FCriticalSection AccessCritSec;
+
+	public:
+		bool IsEmpty() const
+		{
+			return Data.IsEmpty();
+		}
+
+		TArray<T> DequeueAll()
+		{
+			TArray<T> Output;
+
+			{
+				FScopeLock Lock(&AccessCritSec);
+
+				if (!Data.IsEmpty())
+				{
+					double FirstTimeStamp = Data.Peek()->Timestamp;
+
+					T Message;
+					while (Data.Dequeue(Message))
+					{
+						if (HistoryLimitSec > Message.Timestamp - FirstTimeStamp)
+						{
+							Output.Add(MoveTemp(Message));
+						}
+
+						if (!Data.IsEmpty())
+						{
+							FirstTimeStamp = Data.Peek()->Timestamp;
+						}
+					}
+				}
+			}
+
+			return Output;
+		}
+
+		void Enqueue(T&& InMessage)
+		{
+			{
+				// Clear queue if queue time limit hit.
+				// Have to lock as queues do not support
+				// removal from producer thread. Should rarely
+				// be expensive, as that would hit only under
+				// conditions where consuming thread is stalled.
+				FScopeLock Lock(&AccessCritSec);
+
+				if (!Data.IsEmpty())
+				{
+					double FirstTimeStamp = Data.Peek()->Timestamp;
+					while (InMessage.Timestamp - FirstTimeStamp > MaxHistoryLimitSec)
+					{
+						if (!Data.Dequeue().IsSet() || Data.IsEmpty())
+						{
+							break;
+						}
+						FirstTimeStamp = Data.Peek()->Timestamp;
+					}
+				}
+			}
+
+			Data.Enqueue(MoveTemp(InMessage));
+		}
+
+		void Empty()
+		{
+			FScopeLock Lock(&AccessCritSec);
+			Data.Empty();
+		}
+	};
+} // namespace UE::Audio::Insights
diff --git a/Engine/Plugins/AudioInsights/Source/AudioInsights/Public/Views/DashboardViewFactory.h b/Engine/Plugins/AudioInsights/Source/AudioInsights/Public/Views/DashboardViewFactory.h
new file mode 100644
--- /dev/null
+++ b/Engine/Plugins/AudioInsights/Source/AudioInsights/Public/Views/DashboardViewFactory.h
@@ -0,0 +1,67 @@
+// Copyright Epic Games, Inc. All Rights Reserved.
+#pragma once
+
+#include "AudioInsightsTraceProviderBase.h"
+#include "Containers/Array.h"
+#include "HAL/Platform.h"
+#include "Templates/SharedPointer.h"
+#include "Textures/SlateIcon.h"
+#include "UObject/NameTypes.h"
+#include "Widgets/SWidget.h"
+
+
+namespace UE::Audio::Insights
+{
+	enum class AUDIOINSIGHTS_API EDefaultDashboardTabStack : uint8
+	{
+		Viewport,
+		Log,
+		Analysis,
+		AudioMeters,
+		AudioMeter,
+		Oscilloscope
+	};
+
+	class AUDIOINSIGHTS_API IDashboardViewFactory
+	{
+	public:
+		virtual ~IDashboardViewFactory() = default;
+
+		virtual EDefaultDashboardTabStack GetDefaultTabStack() const = 0;
+		virtual FText GetDisplayName() const = 0;
+		virtual FName GetName() const = 0;
+		virtual FSlateIcon GetIcon() const = 0;
+		virtual TSharedRef<SWidget> MakeWidget() = 0;
+	};
+
+	class AUDIOINSIGHTS_API FTraceDashboardViewFactoryBase : public IDashboardViewFactory
+	{
+	public:
+		const TArray<TSharedPtr<FTraceProviderBase>>& GetProviders() const
+		{
+			return Providers;
+		}
+
+		template <typename ProviderType>
+		TSharedPtr<ProviderType> FindProvider(bool bEnsureIfMissing = true) const
+		{
+			for (const TSharedPtr<FTraceProviderBase>& Provider : Providers)
+			{
+				if (Provider->GetName() == ProviderType::GetName_Static())
+				{
+					return StaticCastSharedPtr<ProviderType>(Provider);
+				}
+			}
+
+			if (bEnsureIfMissing)
+			{
+				ensureMsgf(false, TEXT("Failed to find associated provider '%s'"), *ProviderType::GetName_Static().ToString());
+			}
+
+			return TSharedPtr<ProviderType>();
+		}
+
+	protected:
+		TArray<TSharedPtr<FTraceProviderBase>> Providers;
+	};
+} // namespace UE::Audio::Insights
diff --git a/Engine/Plugins/AudioInsights/Source/AudioInsights/Public/Views/TableDashboardViewFactory.h b/Engine/Plugins/AudioInsights/Source/AudioInsights/Public/Views/TableDashboardViewFactory.h
new file mode 100644
--- /dev/null
+++ b/Engine/Plugins/AudioInsights/Source/AudioInsights/Public/Views/TableDashboardViewFactory.h
@@ -0,0 +1,218 @@
+// Copyright Epic Games, Inc. All Rights Reserved.
+#pragma once
+
+#include "AudioDefines.h"
+#include "AudioInsightsDataSource.h"
+#include "Engine/Attenuation.h"
+#include "Engine/World.h"
+#include "Templates/SharedPointer.h"
+#include "Views/DashboardViewFactory.h"
+#include "Widgets/Input/SSearchBox.h"
+#include "Widgets/Views/SHeaderRow.h"
+#include "Widgets/Views/SListView.h"
+
+
+namespace UE::Audio::Insights
+{
+	class AUDIOINSIGHTS_API IObjectDashboardEntry : public IDashboardDataViewEntry
+	{
+	public:
+		virtual ~IObjectDashboardEntry() = default;
+
+		virtual FText GetDisplayName() const = 0;
+		virtual const UObject* GetObject() const = 0;
+		virtual UObject* GetObject() = 0;
+	};
+
+	struct AUDIOINSIGHTS_API FSoundAssetDashboardEntry : public IObjectDashboardEntry
+	{
+		virtual ~FSoundAssetDashboardEntry() = default;
+
+		virtual FText GetDisplayName() const override;
+		virtual const UObject* GetObject() const override;
+		virtual UObject* GetObject() override;
+		virtual bool IsValid() const override;
+
+		::Audio::FDeviceId DeviceId = INDEX_NONE;
+		uint32 PlayOrder = INDEX_NONE;
+		uint64 ComponentId = TNumericLimits<uint64>::Max();
+		double Timestamp = 0.0;
+		FString Name;
+	};
+
+	class AUDIOINSIGHTS_API FTraceTableDashboardViewFactory : public FTraceDashboardViewFactoryBase, public TSharedFromThis<FTraceTableDashboardViewFactory>
+	{
+	public:
+		FTraceTableDashboardViewFactory();
+		virtual ~FTraceTableDashboardViewFactory();
+
+	protected:
+		struct SRowWidget : public SMultiColumnTableRow<TSharedPtr<IDashboardDataViewEntry>>
+		{
+			SLATE_BEGIN_ARGS(SRowWidget) { }
+			SLATE_END_ARGS()
+
+			void Construct(const FArguments& InArgs, const TSharedRef<STableViewBase>& InOwnerTable, TSharedPtr<IDashboardDataViewEntry> InData, TSharedRef<FTraceTableDashboardViewFactory> InFactory);
+			virtual TSharedRef<SWidget> GenerateWidgetForColumn(const FName& Column) override;
+
+			TSharedPtr<IDashboardDataViewEntry> Data;
+			TSharedPtr<FTraceTableDashboardViewFactory> Factory;
+		};
+
+	public:
+		virtual TSharedRef<SWidget> GenerateWidgetForColumn(TSharedRef<IDashboardDataViewEntry> InRowData, const FName& Column);
+		virtual FReply OnDataRowKeyInput(const FGeometry& Geometry, const FKeyEvent& KeyEvent) const;
+		virtual TSharedRef<SWidget> MakeWidget() override;
+		virtual void RebuildFilteredEntriesListView();
+
+		enum class EProcessReason : uint8
+		{
+			None,
+			FilterUpdated,
+			EntriesUpdated
+		};
+
+	protected:
+		void Tick(float InElapsed);
+
+		template<typename TableProviderType>
+		bool FilterEntries(TFunctionRef<bool(const IDashboardDataViewEntry&)> IsFiltered)
+		{
+			const TSharedPtr<const TableProviderType> Provider = FindProvider<const TableProviderType>();
+			if (Provider.IsValid())
+			{
+				if (const TableProviderType::FDeviceData* DeviceData = Provider->FindFilteredDeviceData())
+				{
+					DataViewEntries.Reset();
+
+					auto TransformEntry = [](const TableProviderType::FEntryPair& Pair)
+					{
+						return StaticCastSharedPtr<IDashboardDataViewEntry>(Pair.Value);
+					};
+
+					if (SearchBoxFilterText.IsEmpty())
+					{
+						Algo::Transform(*DeviceData, DataViewEntries, TransformEntry);
+					}
+					else
+					{
+						auto FilterEntry = [this, &IsFiltered](const TableProviderType::FEntryPair& Pair)
+						{
+							return !IsFiltered((const IDashboardDataViewEntry&)(*Pair.Value));
+						};
+						Algo::TransformIf(*DeviceData, DataViewEntries, FilterEntry, TransformEntry);
+					}
+
+					// Temp fix for UE-201431 (SoundAsset entries with empty name are displayed in the Sources tab)
+					DataViewEntries.RemoveAll([](const TSharedPtr<IDashboardDataViewEntry>& Item)
+					{
+						const TSharedPtr<FSoundAssetDashboardEntry> SoundAssetDashboardEntry = StaticCastSharedPtr<FSoundAssetDashboardEntry>(Item);
+
+						return SoundAssetDashboardEntry.IsValid() && SoundAssetDashboardEntry->Name.IsEmpty();
+					});
+					
+					// Sort list
+					RequestSort();
+
+					return true;
+				}
+				else
+				{
+					if (!DataViewEntries.IsEmpty())
+					{
+						DataViewEntries.Empty();
+						return true;
+					}
+				}
+			}
+
+			return false;
+		}
+
+#if WITH_EDITOR
+		virtual bool IsDebugDrawEnabled() const { return false; }
+		virtual void DebugDraw(float InElapsed, const IDashboardDataViewEntry& InDashboardEntry, ::Audio::FDeviceId DeviceId) const { }
+#endif // WITH_EDITOR
+
+		struct FColumnData
+		{
+			const FText DisplayName;
+			const TFunction<FText(const IDashboardDataViewEntry&)> GetDisplayValue;
+			bool bDefaultHidden = false;
+			const float FillWidth = 1.0f;
+			const EHorizontalAlignment Alignment = HAlign_Left;
+		};
+
+		virtual const TMap<FName, FColumnData>& GetColumns() const = 0;
+		virtual void ProcessEntries(EProcessReason Reason) = 0;
+
+		const FText& GetSearchFilterText() const;
+
+		virtual void SortTable() = 0;
+
+		virtual TSharedPtr<SWidget> OnConstructContextMenu();
+		virtual void OnSelectionChanged(TSharedPtr<IDashboardDataViewEntry> SelectedItem, ESelectInfo::Type SelectInfo);
+		virtual FSlateColor GetRowColor(const TSharedPtr<IDashboardDataViewEntry>& InRowDataPtr);
+
+		EProcessReason UpdateFilterReason = EProcessReason::None;
+		FTSTicker::FDelegateHandle TickerHandle;
+
+		FDelegateHandle OnEntriesUpdatedHandle;
+
+		TArray<TSharedPtr<IDashboardDataViewEntry>> DataViewEntries;
+		TMap<FName, uint64> UpdateIds;
+
+		TSharedPtr<SWidget> DashboardWidget;
+		TSharedPtr<SListView<TSharedPtr<IDashboardDataViewEntry>>> FilteredEntriesListView;
+
+		FName SortByColumn;
+		EColumnSortMode::Type SortMode;
+
+	private:
+		TSharedRef<SHeaderRow> MakeHeaderRowWidget();
+		void SetSearchBoxFilterText(const FText& NewText);
+
+		EColumnSortMode::Type GetColumnSortMode(const FName InColumnId) const;
+		void RequestSort();
+		void OnColumnSortModeChanged(const EColumnSortPriority::Type InSortPriority, const FName& InColumnId, const EColumnSortMode::Type InSortMode);
+
+		TSharedPtr<SSearchBox> SearchBoxWidget;
+		FText SearchBoxFilterText;
+	};
+
+	class AUDIOINSIGHTS_API FTraceObjectTableDashboardViewFactory : public FTraceTableDashboardViewFactory
+	{
+	public:
+		virtual ~FTraceObjectTableDashboardViewFactory() = default;
+
+		virtual TSharedRef<SWidget> GenerateWidgetForColumn(TSharedRef<IDashboardDataViewEntry> InRowData, const FName& Column) override;
+		virtual FReply OnDataRowKeyInput(const FGeometry& Geometry, const FKeyEvent& KeyEvent) const override;
+		virtual TSharedRef<SWidget> MakeWidget() override;
+
+	protected:
+		virtual TSharedRef<SWidget> MakeAssetMenuBar() const;
+
+	private:
+		TArray<UObject*> GetSelectedEditableAssets() const;
+
+		bool OpenAsset() const;
+		bool BrowseToAsset() const;
+	};
+
+	class AUDIOINSIGHTS_API FSoundAttenuationVisualizer
+	{
+		mutable TMultiMap<EAttenuationShape::Type, FBaseAttenuationSettings::AttenuationShapeDetails> ShapeDetailsMap;
+		mutable uint32 LastObjectId = INDEX_NONE;
+
+	public:
+		const FColor Color { 155, 155, 255 };
+
+	public:
+		FSoundAttenuationVisualizer(const FColor& InColor)
+			: Color(InColor)
+		{
+		}
+
+		void Draw(float InDeltaTime, const FTransform& InTransform, const UObject& InObject, const UWorld& InWorld) const;
+	};
+} // namespace UE::Audio::Insights
