[AudioInsights] Fixing compiler errors for Incremental Build Minimal / Incremental UnrealEditor MacArm64

diff --git a/Engine/Plugins/AudioInsights/Source/AudioInsights/Public/Messages/AnalyzerMessageQueue.h b/Engine/Plugins/AudioInsights/Source/AudioInsights/Public/Messages/AnalyzerMessageQueue.h
--- a/Engine/Plugins/AudioInsights/Source/AudioInsights/Public/Messages/AnalyzerMessageQueue.h
+++ b/Engine/Plugins/AudioInsights/Source/AudioInsights/Public/Messages/AnalyzerMessageQueue.h
@@ -11,9 +11,9 @@ namespace UE::Audio::Insights
 	template <typename T>
 	struct TAnalyzerMessageQueue
 	{
-		const double MaxHistoryLimitSec = 5.0;
+		static constexpr double MaxHistoryLimitSec = 5.0;
 
-		TAnalyzerMessageQueue(double InHistoryLimitSec = MaxHistoryLimitSec)
+		explicit TAnalyzerMessageQueue(double InHistoryLimitSec = MaxHistoryLimitSec)
 			: HistoryLimitSec(FMath::Clamp(InHistoryLimitSec, UE_DOUBLE_KINDA_SMALL_NUMBER, MaxHistoryLimitSec))
 		{
 		}
diff --git a/Engine/Plugins/AudioInsights/Source/AudioInsights/Public/Views/TableDashboardViewFactory.h b/Engine/Plugins/AudioInsights/Source/AudioInsights/Public/Views/TableDashboardViewFactory.h
--- a/Engine/Plugins/AudioInsights/Source/AudioInsights/Public/Views/TableDashboardViewFactory.h
+++ b/Engine/Plugins/AudioInsights/Source/AudioInsights/Public/Views/TableDashboardViewFactory.h
@@ -81,11 +81,11 @@ namespace UE::Audio::Insights
 			const TSharedPtr<const TableProviderType> Provider = FindProvider<const TableProviderType>();
 			if (Provider.IsValid())
 			{
-				if (const TableProviderType::FDeviceData* DeviceData = Provider->FindFilteredDeviceData())
+				if (const typename TableProviderType::FDeviceData* DeviceData = Provider->FindFilteredDeviceData())
 				{
 					DataViewEntries.Reset();
 
-					auto TransformEntry = [](const TableProviderType::FEntryPair& Pair)
+					auto TransformEntry = [](const typename TableProviderType::FEntryPair& Pair)
 					{
 						return StaticCastSharedPtr<IDashboardDataViewEntry>(Pair.Value);
 					};
@@ -96,7 +96,7 @@ namespace UE::Audio::Insights
 					}
 					else
 					{
-						auto FilterEntry = [this, &IsFiltered](const TableProviderType::FEntryPair& Pair)
+						auto FilterEntry = [this, &IsFiltered](const typename TableProviderType::FEntryPair& Pair)
 						{
 							return !IsFiltered((const IDashboardDataViewEntry&)(*Pair.Value));
 						};
