[AudioInsights] Fixing more compiler errors for Incremental Build / Incremental UnrealEditor Mac

diff --git a/Engine/Plugins/AudioInsights/Source/AudioInsights/Private/Views/OutputMeterDashboardViewFactory.cpp b/Engine/Plugins/AudioInsights/Source/AudioInsights/Private/Views/OutputMeterDashboardViewFactory.cpp
--- a/Engine/Plugins/AudioInsights/Source/AudioInsights/Private/Views/OutputMeterDashboardViewFactory.cpp
+++ b/Engine/Plugins/AudioInsights/Source/AudioInsights/Private/Views/OutputMeterDashboardViewFactory.cpp
@@ -44,7 +44,9 @@ namespace UE::Audio::Insights
 		if (TWeakObjectPtr<UAudioSettings> AudioSettings = GetMutableDefault<UAudioSettings>();
 			AudioSettings.IsValid())
 		{
-			if (MainSubmix = Cast<USoundSubmix>(AudioSettings->MasterSubmix.ResolveObject()))
+			MainSubmix = Cast<USoundSubmix>(AudioSettings->MasterSubmix.ResolveObject());
+
+			if (MainSubmix)
 			{
 				OutputMeterSubmixAnalyzer = MakeShared<FAudioMeterSubmixAnalyzer>(MainSubmix);
 
diff --git a/Engine/Plugins/AudioInsights/Source/AudioInsights/Private/Views/OutputOscilloscopeDashboardViewFactory.cpp b/Engine/Plugins/AudioInsights/Source/AudioInsights/Private/Views/OutputOscilloscopeDashboardViewFactory.cpp
--- a/Engine/Plugins/AudioInsights/Source/AudioInsights/Private/Views/OutputOscilloscopeDashboardViewFactory.cpp
+++ b/Engine/Plugins/AudioInsights/Source/AudioInsights/Private/Views/OutputOscilloscopeDashboardViewFactory.cpp
@@ -42,7 +42,9 @@ namespace UE::Audio::Insights
 		if (TWeakObjectPtr<UAudioSettings> AudioSettings = GetMutableDefault<UAudioSettings>();
 			AudioSettings.IsValid())
 		{
-			if (MainSubmix = Cast<USoundSubmix>(AudioSettings->MasterSubmix.ResolveObject()))
+			MainSubmix = Cast<USoundSubmix>(AudioSettings->MasterSubmix.ResolveObject());
+
+			if (MainSubmix)
 			{
 				OutputOscilloscopeAnalyzer = MakeShared<FAudioOscilloscopeAnalyzer>(MainSubmix);
 
