[AudioInsights] Small renaming in AudioMeterView

diff --git a/Engine/Plugins/AudioInsights/Source/AudioInsights/Private/Views/AudioMeterView.cpp b/Engine/Plugins/AudioInsights/Source/AudioInsights/Private/Views/AudioMeterView.cpp
--- a/Engine/Plugins/AudioInsights/Source/AudioInsights/Private/Views/AudioMeterView.cpp
+++ b/Engine/Plugins/AudioInsights/Source/AudioInsights/Private/Views/AudioMeterView.cpp
@@ -4,6 +4,7 @@
 #include "Analyzers/AudioMeterSubmixAnalyzer.h"
 #include "AudioInsightsDashboardFactory.h"
 #include "AudioMeter.h"
+#include "Sound/AudioBus.h"
 #include "Sound/SoundSubmix.h"
 #include "Views/SubmixesDashboardViewFactory.h"
 #include "Widgets/SNullWidget.h"
@@ -13,8 +14,8 @@ namespace UE::Audio::Insights
 	FAudioMeterView::FAudioMeterView(FAudioAssetVariant InAudioAssetVariant)
 		: AudioAssetVariant(InAudioAssetVariant)
 		, AudioMeterAnalyzerVariant(MakeAudioMeterAnalyzerVariant(InAudioAssetVariant))
-		, SubmixNameTextBlock(MakeSubmixNameTextBlock(InAudioAssetVariant))
-		, SubmixViewWidget(MakeWidget())
+		, AudioAssetNameTextBlock(MakeAudioAssetNameTextBlock(InAudioAssetVariant))
+		, AudioMeterViewWidget(MakeWidget())
 		, OnActiveAudioDeviceChangedHandle(FDashboardFactory::OnActiveAudioDeviceChanged.AddRaw(this, &FAudioMeterView::HandleOnActiveAudioDeviceChanged))
 	{
 		
@@ -44,7 +45,7 @@ namespace UE::Audio::Insights
 		}
 	}
 
-	TSharedRef<STextBlock> FAudioMeterView::MakeSubmixNameTextBlock(const FAudioAssetVariant InAudioAssetVariant)
+	TSharedRef<STextBlock> FAudioMeterView::MakeAudioAssetNameTextBlock(const FAudioAssetVariant InAudioAssetVariant)
 	{
 		switch (InAudioAssetVariant.GetIndex())
 		{
@@ -131,11 +132,11 @@ namespace UE::Audio::Insights
 				[
 					SNew(SBox)
 				]
-				// Submix name label
+				// Audio aaset name label
 				+ SVerticalBox::Slot()
 				.FillHeight(0.1)
 				[
-					SubmixNameTextBlock
+					AudioAssetNameTextBlock
 				]
 				// Height padding
 				+ SVerticalBox::Slot()
diff --git a/Engine/Plugins/AudioInsights/Source/AudioInsights/Private/Views/AudioMeterView.h b/Engine/Plugins/AudioInsights/Source/AudioInsights/Private/Views/AudioMeterView.h
--- a/Engine/Plugins/AudioInsights/Source/AudioInsights/Private/Views/AudioMeterView.h
+++ b/Engine/Plugins/AudioInsights/Source/AudioInsights/Private/Views/AudioMeterView.h
@@ -23,11 +23,11 @@ namespace UE::Audio::Insights
 		FAudioMeterView(FAudioAssetVariant InAudioAssetVariant);
 		virtual ~FAudioMeterView();
 
-		TSharedRef<SWidget> GetWidget() const { return SubmixViewWidget; };
+		TSharedRef<SWidget> GetWidget() const { return AudioMeterViewWidget; };
 
 	private:
 		FAudioMeterVariant MakeAudioMeterAnalyzerVariant(const FAudioAssetVariant InAudioAssetVariant);
-		TSharedRef<STextBlock> MakeSubmixNameTextBlock(const FAudioAssetVariant InAudioAssetVariant);
+		TSharedRef<STextBlock> MakeAudioAssetNameTextBlock(const FAudioAssetVariant InAudioAssetVariant);
 		TSharedRef<SWidget> MakeWidget();
 
 		TSharedRef<SWidget> GetAudioMeterAnalyzerWidget() const;
@@ -37,9 +37,9 @@ namespace UE::Audio::Insights
 		FAudioAssetVariant AudioAssetVariant;
 		FAudioMeterVariant AudioMeterAnalyzerVariant;
 
-		TSharedRef<STextBlock> SubmixNameTextBlock;
+		TSharedRef<STextBlock> AudioAssetNameTextBlock;
 
-		TSharedRef<SWidget> SubmixViewWidget;
+		TSharedRef<SWidget> AudioMeterViewWidget;
 
 		FDelegateHandle OnActiveAudioDeviceChangedHandle;
 	};
