[AudioInsights] Fixing small bug in VirtualLoop column sorting

diff --git a/Engine/Plugins/AudioInsights/Source/AudioInsights/Private/Views/VirtualLoopDashboardViewFactory.cpp b/Engine/Plugins/AudioInsights/Source/AudioInsights/Private/Views/VirtualLoopDashboardViewFactory.cpp
--- a/Engine/Plugins/AudioInsights/Source/AudioInsights/Private/Views/VirtualLoopDashboardViewFactory.cpp
+++ b/Engine/Plugins/AudioInsights/Source/AudioInsights/Private/Views/VirtualLoopDashboardViewFactory.cpp
@@ -170,7 +170,7 @@ namespace UE::Audio::Insights
 					const FVirtualLoopDashboardEntry& AData = VirtualLoopPrivate::CastEntry(*A.Get());
 					const FVirtualLoopDashboardEntry& BData = VirtualLoopPrivate::CastEntry(*B.Get());
 
-					return AData.GetDisplayName().CompareToCaseIgnored(BData.GetDisplayName()) < 0;
+					return BData.GetDisplayName().CompareToCaseIgnored(AData.GetDisplayName()) < 0;
 				});
 			}
 		}
@@ -239,7 +239,7 @@ namespace UE::Audio::Insights
 					const FVirtualLoopDashboardEntry& AData = VirtualLoopPrivate::CastEntry(*A.Get());
 					const FVirtualLoopDashboardEntry& BData = VirtualLoopPrivate::CastEntry(*B.Get());
 
-					return BData.UpdateInterval < BData.UpdateInterval;
+					return BData.UpdateInterval < AData.UpdateInterval;
 				});
 			}
 		}
