[AudioWidgets] Fixed LOCTEXT text conflict issues.

diff --git a/Engine/Plugins/AudioInsights/Source/AudioInsights/Private/Views/MixerSourceDashboardViewFactory.cpp b/Engine/Plugins/AudioInsights/Source/AudioInsights/Private/Views/MixerSourceDashboardViewFactory.cpp
--- a/Engine/Plugins/AudioInsights/Source/AudioInsights/Private/Views/MixerSourceDashboardViewFactory.cpp
+++ b/Engine/Plugins/AudioInsights/Source/AudioInsights/Private/Views/MixerSourceDashboardViewFactory.cpp
@@ -912,7 +912,7 @@ namespace UE::Audio::Insights
 				SAssignNew(SoloToggleButton, SCheckBox)
 				.HAlign(EHorizontalAlignment::HAlign_Center)
 				.Style(&SoloToggleButtonStyle)
-				.ToolTip(FSlateApplicationBase::Get().MakeToolTip(LOCTEXT("TableDashboardView_MuteButtonTooltipText", "Enabled/Disable Solo on all the items in the list.")))
+				.ToolTip(FSlateApplicationBase::Get().MakeToolTip(LOCTEXT("TableDashboardView_SoloButtonTooltipText", "Enabled/Disable Solo on all the items in the list.")))
 				.IsChecked(ECheckBoxState::Unchecked)
 				.OnCheckStateChanged(this, &FMixerSourceDashboardViewFactory::ToggleSoloForAllItems)
 				[
diff --git a/Engine/Plugins/AudioInsights/Source/AudioInsights/Private/Views/OutputMeterDashboardViewFactory.cpp b/Engine/Plugins/AudioInsights/Source/AudioInsights/Private/Views/OutputMeterDashboardViewFactory.cpp
--- a/Engine/Plugins/AudioInsights/Source/AudioInsights/Private/Views/OutputMeterDashboardViewFactory.cpp
+++ b/Engine/Plugins/AudioInsights/Source/AudioInsights/Private/Views/OutputMeterDashboardViewFactory.cpp
@@ -25,7 +25,7 @@ namespace UE::Audio::Insights
 
 	FText FOutputMeterDashboardViewFactory::GetDisplayName() const
 	{
-		return LOCTEXT("AudioDashboard_DashboardsTab_DisplayName", "Output Meter");
+		return LOCTEXT("AudioDashboard_DashboardsAudioMeterTab_DisplayName", "Output Meter");
 	}
 
 	EDefaultDashboardTabStack FOutputMeterDashboardViewFactory::GetDefaultTabStack() const
diff --git a/Engine/Plugins/AudioInsights/Source/AudioInsights/Private/Views/OutputOscilloscopeDashboardViewFactory.cpp b/Engine/Plugins/AudioInsights/Source/AudioInsights/Private/Views/OutputOscilloscopeDashboardViewFactory.cpp
--- a/Engine/Plugins/AudioInsights/Source/AudioInsights/Private/Views/OutputOscilloscopeDashboardViewFactory.cpp
+++ b/Engine/Plugins/AudioInsights/Source/AudioInsights/Private/Views/OutputOscilloscopeDashboardViewFactory.cpp
@@ -23,7 +23,7 @@ namespace UE::Audio::Insights
 
 	FText FOutputOscilloscopeDashboardViewFactory::GetDisplayName() const
 	{
-		return LOCTEXT("AudioDashboard_DashboardsTab_DisplayName", "Output Oscilloscope");
+		return LOCTEXT("AudioDashboard_DashboardsOscilloscopeTab_DisplayName", "Output Oscilloscope");
 	}
 
 	EDefaultDashboardTabStack FOutputOscilloscopeDashboardViewFactory::GetDefaultTabStack() const
