[AudioAnalyzers] Deprecating use of UWorld in favor of using Audio::FDeviceID in audio analyzers/widgets

diff --git a/Engine/Plugins/AudioInsights/Source/AudioInsights/Private/Analyzers/AudioMeterAnalyzer.cpp b/Engine/Plugins/AudioInsights/Source/AudioInsights/Private/Analyzers/AudioMeterAnalyzer.cpp
--- a/Engine/Plugins/AudioInsights/Source/AudioInsights/Private/Analyzers/AudioMeterAnalyzer.cpp
+++ b/Engine/Plugins/AudioInsights/Source/AudioInsights/Private/Analyzers/AudioMeterAnalyzer.cpp
@@ -5,7 +5,6 @@
 #include "AudioInsightsModule.h"
 #include "AudioMeter.h"
 #include "AudioMixerDevice.h"
-#include "Editor.h"
 
 namespace UE::Audio::Insights
 {
@@ -15,23 +14,20 @@ namespace UE::Audio::Insights
 		{
 			using namespace ::Audio;
 
-			if (FAudioDeviceManager* AudioDeviceManager = FAudioDeviceManager::Get())
-			{
-				const IAudioInsightsModule& InsightsModule = FModuleManager::GetModuleChecked<IAudioInsightsModule>(IAudioInsightsModule::GetName());
-				TArray<UWorld*> DeviceWorlds = AudioDeviceManager->GetWorldsUsingAudioDevice(InsightsModule.GetDeviceId());
+			const IAudioInsightsModule& InsightsModule = FModuleManager::GetModuleChecked<IAudioInsightsModule>(IAudioInsightsModule::GetName());
+			const FDeviceId AudioDeviceId = InsightsModule.GetDeviceId();
 
-				if (UWorld* World = (DeviceWorlds.Num() > 0) ? DeviceWorlds[0] : nullptr)
+			if (const FAudioDeviceManager* AudioDeviceManager = FAudioDeviceManager::Get())
+			{
+				if (const FMixerDevice* MixerDevice = static_cast<const FMixerDevice*>(AudioDeviceManager->GetAudioDeviceRaw(AudioDeviceId)))
 				{
-					if (FMixerDevice* MixerDevice = static_cast<FMixerDevice*>(World->GetAudioDeviceRaw()))
-					{
-						return MakeShared<AudioWidgets::FAudioMeter>(InExternalAudioBus.IsValid() ? InExternalAudioBus->GetNumChannels() : MixerDevice->GetNumDeviceChannels(), 
-							*DeviceWorlds[0], 
-							InExternalAudioBus.Get());
-					}
+					return MakeShared<AudioWidgets::FAudioMeter>(InExternalAudioBus.IsValid() ? InExternalAudioBus->GetNumChannels() : MixerDevice->GetNumDeviceChannels(),
+						AudioDeviceId,
+						InExternalAudioBus.Get());
 				}
 			}
 			
-			return MakeShared<AudioWidgets::FAudioMeter>(1, *GEditor->GetWorld());
+			return MakeShared<AudioWidgets::FAudioMeter>(1, AudioDeviceId);
 		}
 	}
 
@@ -52,20 +48,14 @@ namespace UE::Audio::Insights
 		}
 
 		const IAudioInsightsModule& InsightsModule = FModuleManager::GetModuleChecked<IAudioInsightsModule>(IAudioInsightsModule::GetName());
-		TArray<UWorld*> DeviceWorlds = AudioDeviceManager->GetWorldsUsingAudioDevice(InsightsModule.GetDeviceId());
-
-		UWorld* World = (DeviceWorlds.Num() > 0) ? DeviceWorlds[0] : nullptr;
-		if (!World)
-		{
-			return;
-		}
+		const FDeviceId AudioDeviceId = InsightsModule.GetDeviceId();
 
-		FMixerDevice* MixerDevice = static_cast<FMixerDevice*>(World->GetAudioDeviceRaw());
+		const FMixerDevice* MixerDevice = static_cast<const FMixerDevice*>(AudioDeviceManager->GetAudioDeviceRaw(AudioDeviceId));
 		if (!MixerDevice)
 		{
 			return;
 		}
 
-		AudioMeter->Init(InExternalAudioBus.IsValid() ? InExternalAudioBus->GetNumChannels() : MixerDevice->GetNumDeviceChannels(), *World, InExternalAudioBus.Get());
+		AudioMeter->Init(InExternalAudioBus.IsValid() ? InExternalAudioBus->GetNumChannels() : MixerDevice->GetNumDeviceChannels(), AudioDeviceId, InExternalAudioBus.Get());
 	}
 } // namespace UE::Audio::Insights
diff --git a/Engine/Plugins/AudioInsights/Source/AudioInsights/Private/Analyzers/AudioMeterSubmixAnalyzer.cpp b/Engine/Plugins/AudioInsights/Source/AudioInsights/Private/Analyzers/AudioMeterSubmixAnalyzer.cpp
--- a/Engine/Plugins/AudioInsights/Source/AudioInsights/Private/Analyzers/AudioMeterSubmixAnalyzer.cpp
+++ b/Engine/Plugins/AudioInsights/Source/AudioInsights/Private/Analyzers/AudioMeterSubmixAnalyzer.cpp
@@ -2,6 +2,7 @@
 #include "AudioMeterSubmixAnalyzer.h"
 
 #include "AudioBusSubsystem.h"
+#include "AudioDefines.h"
 #include "AudioDeviceManager.h"
 #include "AudioInsightsModule.h"
 #include "AudioMeter.h"
@@ -47,22 +48,16 @@ namespace UE::Audio::Insights
 		}
 
 		const IAudioInsightsModule& InsightsModule = FModuleManager::GetModuleChecked<IAudioInsightsModule>(IAudioInsightsModule::GetName());
-		TArray<UWorld*> DeviceWorlds = AudioDeviceManager->GetWorldsUsingAudioDevice(InsightsModule.GetDeviceId());
+		const FDeviceId AudioDeviceId = InsightsModule.GetDeviceId();
 
-		UWorld* World = (DeviceWorlds.Num() > 0) ? DeviceWorlds[0] : nullptr;
-		if (!World)
-		{
-			return;
-		}
-
-		FMixerDevice* MixerDevice = static_cast<FMixerDevice*>(World->GetAudioDeviceRaw());
+		const FMixerDevice* MixerDevice = static_cast<const FMixerDevice*>(AudioDeviceManager->GetAudioDeviceRaw(AudioDeviceId));
 		if (!MixerDevice)
 		{
 			return;
 		}
 
-		FMixerSubmixPtr MixerSubmix = MixerDevice->GetSubmixInstance(SoundSubmix.Get()).Pin();
-		if (!MixerSubmix.IsValid())
+		FMixerSubmixWeakPtr MixerSubmixWeakPtr = MixerDevice->GetSubmixInstance(SoundSubmix.Get());
+		if (!MixerSubmixWeakPtr.IsValid())
 		{
 			return;
 		}
@@ -77,12 +72,16 @@ namespace UE::Audio::Insights
 		const FAudioBusKey AudioBusKey(AudioBus->GetUniqueID());
 		const int32 AudioBusNumChannels = AudioBus->GetNumChannels();
 
-		FAudioThread::RunCommandOnAudioThread([MixerDevice, MixerSubmix, AudioBusKey, AudioBusNumChannels]()
+		FAudioThread::RunCommandOnAudioThread([MixerDevice, MixerSubmixWeakPtr, AudioBusKey, AudioBusNumChannels]()
 		{
 			TObjectPtr<UAudioBusSubsystem> AudioBusSubsystem = MixerDevice->GetSubsystem<UAudioBusSubsystem>();
 			check(AudioBusSubsystem);
 
-			MixerSubmix->RegisterAudioBus(AudioBusKey, AudioBusSubsystem->AddPatchInputForAudioBus(AudioBusKey, MixerDevice->GetNumOutputFrames(), AudioBusNumChannels));
+			if (FMixerSubmixPtr MixerSubmix = MixerSubmixWeakPtr.Pin();
+				MixerSubmix.IsValid())
+			{
+				MixerSubmix->RegisterAudioBus(AudioBusKey, AudioBusSubsystem->AddPatchInputForAudioBus(AudioBusKey, MixerDevice->GetNumOutputFrames(), AudioBusNumChannels));
+			}
 		});
 	}
 
@@ -97,15 +96,9 @@ namespace UE::Audio::Insights
 		}
 
 		const IAudioInsightsModule& InsightsModule = FModuleManager::GetModuleChecked<IAudioInsightsModule>(IAudioInsightsModule::GetName());
-		TArray<UWorld*> DeviceWorlds = AudioDeviceManager->GetWorldsUsingAudioDevice(InsightsModule.GetDeviceId());
-
-		UWorld* World = (DeviceWorlds.Num() > 0) ? DeviceWorlds[0] : nullptr;
-		if (!World)
-		{
-			return;
-		}
+		const FDeviceId AudioDeviceId = InsightsModule.GetDeviceId();
 
-		FMixerDevice* MixerDevice = static_cast<FMixerDevice*>(World->GetAudioDeviceRaw());
+		const FMixerDevice* MixerDevice = static_cast<const FMixerDevice*>(AudioDeviceManager->GetAudioDeviceRaw(AudioDeviceId));
 		if (!MixerDevice)
 		{
 			return;
@@ -116,8 +109,8 @@ namespace UE::Audio::Insights
 			return;
 		}
 
-		FMixerSubmixPtr MixerSubmix = MixerDevice->GetSubmixInstance(SoundSubmix.Get()).Pin();
-		if (!MixerSubmix.IsValid())
+		FMixerSubmixWeakPtr MixerSubmixWeakPtr = MixerDevice->GetSubmixInstance(SoundSubmix.Get());
+		if (!MixerSubmixWeakPtr.IsValid())
 		{
 			return;
 		}
@@ -131,9 +124,13 @@ namespace UE::Audio::Insights
 
 		const FAudioBusKey AudioBusKey(AudioBus->GetUniqueID());
 
-		FAudioThread::RunCommandOnAudioThread([MixerDevice, MixerSubmix, AudioBusKey]()
+		FAudioThread::RunCommandOnAudioThread([MixerSubmixWeakPtr, AudioBusKey]()
 		{
-			MixerSubmix->UnregisterAudioBus(AudioBusKey);
+			if (FMixerSubmixPtr MixerSubmix = MixerSubmixWeakPtr.Pin();
+				MixerSubmix.IsValid())
+			{
+				MixerSubmix->UnregisterAudioBus(AudioBusKey);
+			}
 		});
 	}
 
diff --git a/Engine/Plugins/AudioInsights/Source/AudioInsights/Private/Analyzers/AudioOscilloscopeAnalyzer.cpp b/Engine/Plugins/AudioInsights/Source/AudioInsights/Private/Analyzers/AudioOscilloscopeAnalyzer.cpp
--- a/Engine/Plugins/AudioInsights/Source/AudioInsights/Private/Analyzers/AudioOscilloscopeAnalyzer.cpp
+++ b/Engine/Plugins/AudioInsights/Source/AudioInsights/Private/Analyzers/AudioOscilloscopeAnalyzer.cpp
@@ -2,12 +2,12 @@
 #include "AudioOscilloscopeAnalyzer.h"
 
 #include "AudioBusSubsystem.h"
+#include "AudioDefines.h"
 #include "AudioDeviceManager.h"
 #include "AudioInsightsModule.h"
 #include "AudioOscilloscope.h"
 #include "AudioMixerDevice.h"
 #include "AudioMixerSubmix.h"
-#include "Editor.h"
 
 namespace UE::Audio::Insights
 {
@@ -20,26 +20,23 @@ namespace UE::Audio::Insights
 		{
 			using namespace ::Audio;
 
+			const IAudioInsightsModule& InsightsModule = FModuleManager::GetModuleChecked<IAudioInsightsModule>(IAudioInsightsModule::GetName());
+			const FDeviceId AudioDeviceId = InsightsModule.GetDeviceId();
+
 			if (FAudioDeviceManager* AudioDeviceManager = FAudioDeviceManager::Get())
 			{
-				const IAudioInsightsModule& InsightsModule = FModuleManager::GetModuleChecked<IAudioInsightsModule>(IAudioInsightsModule::GetName());
-				TArray<UWorld*> DeviceWorlds = AudioDeviceManager->GetWorldsUsingAudioDevice(InsightsModule.GetDeviceId());
-
-				if (UWorld* World = (DeviceWorlds.Num() > 0) ? DeviceWorlds[0] : nullptr)
+				if (const FMixerDevice* MixerDevice = static_cast<const FMixerDevice*>(AudioDeviceManager->GetAudioDeviceRaw(AudioDeviceId)))
 				{
-					if (FMixerDevice* MixerDevice = static_cast<FMixerDevice*>(World->GetAudioDeviceRaw()))
-					{
-						return MakeShared<AudioWidgets::FAudioOscilloscope>(DeviceWorlds[0],
-							MixerDevice->GetNumDeviceChannels(),
-							InTimeWindowMs,
-							InMaxTimeWindowMs,
-							InAnalysisPeriodMs,
-							InPanelLayoutType);
-					}
+					return MakeShared<AudioWidgets::FAudioOscilloscope>(AudioDeviceId,
+						MixerDevice->GetNumDeviceChannels(),
+						InTimeWindowMs,
+						InMaxTimeWindowMs,
+						InAnalysisPeriodMs,
+						InPanelLayoutType);
 				}
 			}
 			
-			return MakeShared<AudioWidgets::FAudioOscilloscope>(GEditor->GetWorld(), 1, InTimeWindowMs, InMaxTimeWindowMs, InAnalysisPeriodMs, InPanelLayoutType);
+			return MakeShared<AudioWidgets::FAudioOscilloscope>(AudioDeviceId, 1, InTimeWindowMs, InMaxTimeWindowMs, InAnalysisPeriodMs, InPanelLayoutType);
 		}
 	}
 
@@ -72,15 +69,9 @@ namespace UE::Audio::Insights
 		}
 
 		const IAudioInsightsModule& InsightsModule = FModuleManager::GetModuleChecked<IAudioInsightsModule>(IAudioInsightsModule::GetName());
-		TArray<UWorld*> DeviceWorlds = AudioDeviceManager->GetWorldsUsingAudioDevice(InsightsModule.GetDeviceId());
+		const FDeviceId AudioDeviceId = InsightsModule.GetDeviceId();
 
-		UWorld* World = (DeviceWorlds.Num() > 0) ? DeviceWorlds[0] : nullptr;
-		if (!World)
-		{
-			return;
-		}
-
-		FMixerDevice* MixerDevice = static_cast<FMixerDevice*>(World->GetAudioDeviceRaw());
+		const FMixerDevice* MixerDevice = static_cast<FMixerDevice*>(AudioDeviceManager->GetAudioDeviceRaw(AudioDeviceId));
 		if (!MixerDevice)
 		{
 			return;
@@ -91,13 +82,13 @@ namespace UE::Audio::Insights
 			return;
 		}
 
-		FMixerSubmixPtr MixerSubmix = MixerDevice->GetSubmixInstance(SoundSubmix.Get()).Pin();
-		if (!MixerSubmix.IsValid())
+		FMixerSubmixWeakPtr MixerSubmixWeakPtr = MixerDevice->GetSubmixInstance(SoundSubmix.Get());
+		if (!MixerSubmixWeakPtr.IsValid())
 		{
 			return;
 		}
 
-		AudioOscilloscope->CreateDataProvider(World, TimeWindowMs, MaxTimeWindowMs, AnalysisPeriodMs, PanelLayoutType);
+		AudioOscilloscope->CreateDataProvider(AudioDeviceId, TimeWindowMs, MaxTimeWindowMs, AnalysisPeriodMs, PanelLayoutType);
 		AudioOscilloscope->CreateOscilloscopeWidget(MixerDevice->GetNumDeviceChannels(), PanelLayoutType);
 
 		// Start processing
@@ -113,12 +104,16 @@ namespace UE::Audio::Insights
 		const FAudioBusKey AudioBusKey(AudioBus->GetUniqueID());
 		const int32 AudioBusNumChannels = AudioBus->GetNumChannels();
 
-		FAudioThread::RunCommandOnAudioThread([MixerDevice, MixerSubmix, AudioBusKey, AudioBusNumChannels]()
+		FAudioThread::RunCommandOnAudioThread([MixerDevice, MixerSubmixWeakPtr, AudioBusKey, AudioBusNumChannels]()
 		{
 			TObjectPtr<UAudioBusSubsystem> AudioBusSubsystem = MixerDevice->GetSubsystem<UAudioBusSubsystem>();
 			check(AudioBusSubsystem);
 
-			MixerSubmix->RegisterAudioBus(AudioBusKey, AudioBusSubsystem->AddPatchInputForAudioBus(AudioBusKey, MixerDevice->GetNumOutputFrames(), AudioBusNumChannels));
+			if (FMixerSubmixPtr MixerSubmix = MixerSubmixWeakPtr.Pin();
+				MixerSubmix.IsValid())
+			{
+				MixerSubmix->RegisterAudioBus(AudioBusKey, AudioBusSubsystem->AddPatchInputForAudioBus(AudioBusKey, MixerDevice->GetNumOutputFrames(), AudioBusNumChannels));
+			}
 		});
 	}
 
@@ -133,15 +128,9 @@ namespace UE::Audio::Insights
 		}
 
 		const IAudioInsightsModule& InsightsModule = FModuleManager::GetModuleChecked<IAudioInsightsModule>(IAudioInsightsModule::GetName());
-		TArray<UWorld*> DeviceWorlds = AudioDeviceManager->GetWorldsUsingAudioDevice(InsightsModule.GetDeviceId());
+		const FDeviceId AudioDeviceId = InsightsModule.GetDeviceId();
 
-		UWorld* World = (DeviceWorlds.Num() > 0) ? DeviceWorlds[0] : nullptr;
-		if (!World)
-		{
-			return;
-		}
-
-		FMixerDevice* MixerDevice = static_cast<FMixerDevice*>(World->GetAudioDeviceRaw());
+		const FMixerDevice* MixerDevice = static_cast<FMixerDevice*>(AudioDeviceManager->GetAudioDeviceRaw(AudioDeviceId));
 		if (!MixerDevice)
 		{
 			return;
@@ -152,8 +141,8 @@ namespace UE::Audio::Insights
 			return;
 		}
 
-		FMixerSubmixPtr MixerSubmix = MixerDevice->GetSubmixInstance(SoundSubmix.Get()).Pin();
-		if (!MixerSubmix.IsValid())
+		FMixerSubmixWeakPtr MixerSubmixWeakPtr = MixerDevice->GetSubmixInstance(SoundSubmix.Get());
+		if (!MixerSubmixWeakPtr.IsValid())
 		{
 			return;
 		}
@@ -167,9 +156,13 @@ namespace UE::Audio::Insights
 
 		const FAudioBusKey AudioBusKey(AudioBus->GetUniqueID());
 
-		FAudioThread::RunCommandOnAudioThread([MixerDevice, MixerSubmix, AudioBusKey]()
+		FAudioThread::RunCommandOnAudioThread([MixerSubmixWeakPtr, AudioBusKey]()
 		{
-			MixerSubmix->UnregisterAudioBus(AudioBusKey);
+			if (FMixerSubmixPtr MixerSubmix = MixerSubmixWeakPtr.Pin();
+				MixerSubmix.IsValid())
+			{
+				MixerSubmix->UnregisterAudioBus(AudioBusKey);
+			}
 		});
 
 		// Stop processing
diff --git a/Engine/Plugins/Runtime/AudioWidgets/Source/AudioWidgets/Private/AudioMeter.cpp b/Engine/Plugins/Runtime/AudioWidgets/Source/AudioWidgets/Private/AudioMeter.cpp
--- a/Engine/Plugins/Runtime/AudioWidgets/Source/AudioWidgets/Private/AudioMeter.cpp
+++ b/Engine/Plugins/Runtime/AudioWidgets/Source/AudioWidgets/Private/AudioMeter.cpp
@@ -162,6 +162,12 @@ const FText UAudioMeter::GetPaletteCategory()
 namespace AudioWidgets
 {
 	FAudioMeter::FAudioMeter(int32 InNumChannels, UWorld& InWorld, TObjectPtr<UAudioBus> InExternalAudioBus)
+		: FAudioMeter(InNumChannels, InWorld.GetAudioDevice().GetDeviceID(), InExternalAudioBus)
+	{
+		
+	}
+
+	FAudioMeter::FAudioMeter(const int32 InNumChannels, const Audio::FDeviceId InAudioDeviceId, const TObjectPtr<UAudioBus> InExternalAudioBus)
 		: Widget(SNew(SAudioMeter)
 			.Orientation(EOrientation::Orient_Vertical)
 			.BackgroundColor(FLinearColor::Transparent)
@@ -175,7 +181,7 @@ namespace AudioWidgets
 			.MeterScaleLabelColor(FLinearColor(0.442708f, 0.442708f, 0.442708f, 1.0f))
 		)
 	{
-		Init(InNumChannels, InWorld, InExternalAudioBus);
+		Init(InNumChannels, InAudioDeviceId, InExternalAudioBus);
 	}
 
 	FAudioMeter::~FAudioMeter()
@@ -194,6 +200,17 @@ namespace AudioWidgets
 	}
 
 	void FAudioMeter::Init(int32 InNumChannels, UWorld& InWorld, TObjectPtr<UAudioBus> InExternalAudioBus)
+	{
+		const FAudioDeviceHandle AudioDevice = InWorld.GetAudioDevice();
+		if (!AudioDevice.IsValid())
+		{
+			return;
+		}
+
+		Init(InNumChannels, AudioDevice.GetDeviceID(), InExternalAudioBus);
+	}
+
+	void FAudioMeter::Init(const int32 InNumChannels, const Audio::FDeviceId InAudioDeviceId, const TObjectPtr<UAudioBus> InExternalAudioBus)
 	{
 		check(InNumChannels > 0);
 
@@ -212,8 +229,7 @@ namespace AudioWidgets
 
 		ResultsDelegateHandle = Analyzer->OnLatestPerChannelMeterResultsNative.AddRaw(this, &FAudioMeter::OnMeterOutput);
 
-		WorldPtr = &InWorld;
-		Analyzer->StartAnalyzing(&InWorld, AudioBus.Get());
+		Analyzer->StartAnalyzing(InAudioDeviceId, AudioBus.Get());
 
 		constexpr float DefaultMeterValue = -60.0f;
 		constexpr float DefaultPeakValue = -60.0f;
diff --git a/Engine/Plugins/Runtime/AudioWidgets/Source/AudioWidgets/Private/AudioOscilloscope.cpp b/Engine/Plugins/Runtime/AudioWidgets/Source/AudioWidgets/Private/AudioOscilloscope.cpp
--- a/Engine/Plugins/Runtime/AudioWidgets/Source/AudioWidgets/Private/AudioOscilloscope.cpp
+++ b/Engine/Plugins/Runtime/AudioWidgets/Source/AudioWidgets/Private/AudioOscilloscope.cpp
@@ -2,11 +2,9 @@
 
 #include "AudioOscilloscope.h"
 
-#include "Engine/World.h"
-
 namespace AudioWidgets
 {
-	FAudioOscilloscope::FAudioOscilloscope(UWorld* InWorld, 
+	FAudioOscilloscope::FAudioOscilloscope(Audio::FDeviceId InAudioDeviceId,
 		const uint32 InNumChannels, 
 		const float InTimeWindowMs, 
 		const float InMaxTimeWindowMs, 
@@ -15,7 +13,7 @@ namespace AudioWidgets
 		: OscilloscopePanelStyle(FAudioOscilloscopePanelStyle::GetDefault())
 	{
 		CreateAudioBus(InNumChannels);
-		CreateDataProvider(InWorld, InTimeWindowMs, InMaxTimeWindowMs, InAnalysisPeriodMs, InPanelLayoutType);
+		CreateDataProvider(InAudioDeviceId, InTimeWindowMs, InMaxTimeWindowMs, InAnalysisPeriodMs, InPanelLayoutType);
 		CreateOscilloscopeWidget(InNumChannels, InPanelLayoutType);
 	}
 
@@ -25,7 +23,7 @@ namespace AudioWidgets
 		AudioBus->AudioBusChannels = AudioBusUtils::ConvertIntToEAudioBusChannels(InNumChannels);
 	}
 
-	void FAudioOscilloscope::CreateDataProvider(UWorld* InWorld,
+	void FAudioOscilloscope::CreateDataProvider(Audio::FDeviceId InAudioDeviceId,
 		const float InTimeWindowMs,
 		const float InMaxTimeWindowMs,
 		const float InAnalysisPeriodMs,
@@ -34,7 +32,7 @@ namespace AudioWidgets
 		check(AudioBus);
 
 		const uint32 NumChannelsToProvide = (InPanelLayoutType == EAudioPanelLayoutType::Advanced) ? 1 : AudioBus->GetNumChannels(); // Advanced mode waveform display is based on channel selection
-		AudioSamplesDataProvider = MakeShared<FWaveformAudioSamplesDataProvider>(InWorld, AudioBus.Get(), NumChannelsToProvide, InTimeWindowMs, InMaxTimeWindowMs, InAnalysisPeriodMs);
+		AudioSamplesDataProvider = MakeShared<FWaveformAudioSamplesDataProvider>(InAudioDeviceId, AudioBus.Get(), NumChannelsToProvide, InTimeWindowMs, InMaxTimeWindowMs, InAnalysisPeriodMs);
 	}
 
 	void FAudioOscilloscope::CreateOscilloscopeWidget(const uint32 InNumChannels, const EAudioPanelLayoutType InPanelLayoutType)
diff --git a/Engine/Plugins/Runtime/AudioWidgets/Source/AudioWidgets/Private/AudioOscilloscopeUMG.cpp b/Engine/Plugins/Runtime/AudioWidgets/Source/AudioWidgets/Private/AudioOscilloscopeUMG.cpp
--- a/Engine/Plugins/Runtime/AudioWidgets/Source/AudioWidgets/Private/AudioOscilloscopeUMG.cpp
+++ b/Engine/Plugins/Runtime/AudioWidgets/Source/AudioWidgets/Private/AudioOscilloscopeUMG.cpp
@@ -40,7 +40,20 @@ void UAudioOscilloscope::CreateDataProvider()
 {
 	constexpr uint32 NumChannelsToProvide = 1;
 	constexpr float  MaxTimeWindowMs      = 5000.0f; // TODO alex.perez: should we expose this as a UPROPERTY?
-	AudioSamplesDataProvider = MakeShared<AudioWidgets::FWaveformAudioSamplesDataProvider>(GetWorld(), AudioBus, NumChannelsToProvide, TimeWindowMs, MaxTimeWindowMs, AnalysisPeriodMs);
+
+	UWorld* World = GetWorld();
+	if (!World)
+	{
+		return;
+	}
+
+	const FAudioDeviceHandle AudioDevice = World->GetAudioDevice();
+	if (!AudioDevice.IsValid())
+	{
+		return;
+	}
+
+	AudioSamplesDataProvider = MakeShared<AudioWidgets::FWaveformAudioSamplesDataProvider>(AudioDevice.GetDeviceID(), AudioBus, NumChannelsToProvide, TimeWindowMs, MaxTimeWindowMs, AnalysisPeriodMs);
 }
 
 void UAudioOscilloscope::CreateOscilloscopeWidget()
diff --git a/Engine/Plugins/Runtime/AudioWidgets/Source/AudioWidgets/Private/AudioSpectrumAnalyzer.cpp b/Engine/Plugins/Runtime/AudioWidgets/Source/AudioWidgets/Private/AudioSpectrumAnalyzer.cpp
--- a/Engine/Plugins/Runtime/AudioWidgets/Source/AudioWidgets/Private/AudioSpectrumAnalyzer.cpp
+++ b/Engine/Plugins/Runtime/AudioWidgets/Source/AudioWidgets/Private/AudioSpectrumAnalyzer.cpp
@@ -3,18 +3,17 @@
 #include "AudioSpectrumAnalyzer.h"
 
 #include "DSP/EnvelopeFollower.h"
-#include "Engine/World.h"
 
 namespace AudioWidgets
 {
-	FAudioSpectrumAnalyzer::FAudioSpectrumAnalyzer(int32 InNumChannels, UWorld& InWorld, TObjectPtr<UAudioBus> InExternalAudioBus)
+	FAudioSpectrumAnalyzer::FAudioSpectrumAnalyzer(int32 InNumChannels, Audio::FDeviceId InAudioDeviceId, TObjectPtr<UAudioBus> InExternalAudioBus)
 		: Widget(SNew(SAudioSpectrumPlot)
 			.Clipping(EWidgetClipping::ClipToBounds)
 			.DisplayFrequencyAxisLabels(false)
 			.DisplaySoundLevelAxisLabels(false)
 			.OnGetAudioSpectrumData_Raw(this, &FAudioSpectrumAnalyzer::GetAudioSpectrumData))
 	{
-		Init(InNumChannels, InWorld, InExternalAudioBus);
+		Init(InNumChannels, InAudioDeviceId, InExternalAudioBus);
 	}
 
 	FAudioSpectrumAnalyzer::~FAudioSpectrumAnalyzer()
@@ -32,7 +31,7 @@ namespace AudioWidgets
 		return Widget->AsShared();
 	}
 
-	void FAudioSpectrumAnalyzer::Init(int32 InNumChannels, UWorld& InWorld, TObjectPtr<UAudioBus> InExternalAudioBus)
+	void FAudioSpectrumAnalyzer::Init(int32 InNumChannels, Audio::FDeviceId InAudioDeviceId, TObjectPtr<UAudioBus> InExternalAudioBus)
 	{
 		check(InNumChannels > 0);
 
@@ -57,8 +56,7 @@ namespace AudioWidgets
 
 		ResultsDelegateHandle = Analyzer->OnConstantQResultsNative.AddRaw(this, &FAudioSpectrumAnalyzer::OnConstantQResults);
 
-		WorldPtr = &InWorld;
-		Analyzer->StartAnalyzing(&InWorld, AudioBus.Get());
+		Analyzer->StartAnalyzing(InAudioDeviceId, AudioBus.Get());
 	}
 
 	void FAudioSpectrumAnalyzer::OnConstantQResults(UConstantQAnalyzer* InSpectrumAnalyzer, int32 ChannelIndex, const TArray<FConstantQResults>& InSpectrumResultsArray)
diff --git a/Engine/Plugins/Runtime/AudioWidgets/Source/AudioWidgets/Private/AudioVectorscope.cpp b/Engine/Plugins/Runtime/AudioWidgets/Source/AudioWidgets/Private/AudioVectorscope.cpp
--- a/Engine/Plugins/Runtime/AudioWidgets/Source/AudioWidgets/Private/AudioVectorscope.cpp
+++ b/Engine/Plugins/Runtime/AudioWidgets/Source/AudioWidgets/Private/AudioVectorscope.cpp
@@ -2,12 +2,11 @@
 
 #include "AudioVectorscope.h"
 
-#include "Engine/World.h"
 #include "SAudioVectorscopePanelWidget.h"
 
 namespace AudioWidgets
 {
-	FAudioVectorscope::FAudioVectorscope(UWorld* InWorld, 
+	FAudioVectorscope::FAudioVectorscope(Audio::FDeviceId InAudioDeviceId,
 		const uint32 InNumChannels, 
 		const float InTimeWindowMs, 
 		const float InMaxTimeWindowMs, 
@@ -16,7 +15,7 @@ namespace AudioWidgets
 		: VectorscopePanelStyle(FAudioVectorscopePanelStyle::GetDefault())
 	{
 		CreateAudioBus(InNumChannels);
-		CreateDataProvider(InWorld, InTimeWindowMs, InMaxTimeWindowMs, InAnalysisPeriodMs);
+		CreateDataProvider(InAudioDeviceId, InTimeWindowMs, InMaxTimeWindowMs, InAnalysisPeriodMs);
 		CreateVectorscopeWidget(InPanelLayoutType);
 	}
 
@@ -26,11 +25,11 @@ namespace AudioWidgets
 		AudioBus->AudioBusChannels = AudioBusUtils::ConvertIntToEAudioBusChannels(InNumChannels);
 	}
 
-	void FAudioVectorscope::CreateDataProvider(UWorld* InWorld, const float InTimeWindowMs,	const float InMaxTimeWindowMs, const float InAnalysisPeriodMs)
+	void FAudioVectorscope::CreateDataProvider(Audio::FDeviceId InAudioDeviceId, const float InTimeWindowMs,	const float InMaxTimeWindowMs, const float InAnalysisPeriodMs)
 	{
 		check(AudioBus);
 
-		AudioSamplesDataProvider = MakeShared<FWaveformAudioSamplesDataProvider>(InWorld, AudioBus.Get(), AudioBus->GetNumChannels(), InTimeWindowMs, InMaxTimeWindowMs, InAnalysisPeriodMs);
+		AudioSamplesDataProvider = MakeShared<FWaveformAudioSamplesDataProvider>(InAudioDeviceId, AudioBus.Get(), AudioBus->GetNumChannels(), InTimeWindowMs, InMaxTimeWindowMs, InAnalysisPeriodMs);
 	}
 
 	void FAudioVectorscope::CreateVectorscopeWidget(const EAudioPanelLayoutType InPanelLayoutType)
diff --git a/Engine/Plugins/Runtime/AudioWidgets/Source/AudioWidgets/Private/AudioVectorscopeUMG.cpp b/Engine/Plugins/Runtime/AudioWidgets/Source/AudioWidgets/Private/AudioVectorscopeUMG.cpp
--- a/Engine/Plugins/Runtime/AudioWidgets/Source/AudioWidgets/Private/AudioVectorscopeUMG.cpp
+++ b/Engine/Plugins/Runtime/AudioWidgets/Source/AudioWidgets/Private/AudioVectorscopeUMG.cpp
@@ -3,6 +3,7 @@
 #include "AudioVectorscopeUMG.h"
 
 #include "AudioMixerDevice.h"
+#include "Engine/World.h"
 #include "SAudioVectorscopePanelWidget.h"
 #include "WaveformAudioSamplesDataProvider.h"
 
@@ -33,7 +34,20 @@ void UAudioVectorscope::CreateDummyVectorscopeWidget()
 void UAudioVectorscope::CreateDataProvider()
 {
 	constexpr float MaxDisplayPersistenceMs = 500.0f; // TODO alex.perez: should we expose this as a UPROPERTY?
-	AudioSamplesDataProvider = MakeShared<AudioWidgets::FWaveformAudioSamplesDataProvider>(GetWorld(), AudioBus, AudioBus->GetNumChannels(), DisplayPersistenceMs, MaxDisplayPersistenceMs, AnalysisPeriodMs);
+
+	UWorld* World = GetWorld();
+	if (!World)
+	{
+		return;
+	}
+
+	const FAudioDeviceHandle AudioDevice = World->GetAudioDevice();
+	if (!AudioDevice.IsValid())
+	{
+		return;
+	}
+
+	AudioSamplesDataProvider = MakeShared<AudioWidgets::FWaveformAudioSamplesDataProvider>(AudioDevice.GetDeviceID(), AudioBus, AudioBus->GetNumChannels(), DisplayPersistenceMs, MaxDisplayPersistenceMs, AnalysisPeriodMs);
 }
 
 void UAudioVectorscope::CreateVectorscopeWidget()
diff --git a/Engine/Plugins/Runtime/AudioWidgets/Source/AudioWidgets/Private/WaveformAudioSamplesDataProvider.cpp b/Engine/Plugins/Runtime/AudioWidgets/Source/AudioWidgets/Private/WaveformAudioSamplesDataProvider.cpp
--- a/Engine/Plugins/Runtime/AudioWidgets/Source/AudioWidgets/Private/WaveformAudioSamplesDataProvider.cpp
+++ b/Engine/Plugins/Runtime/AudioWidgets/Source/AudioWidgets/Private/WaveformAudioSamplesDataProvider.cpp
@@ -35,34 +35,33 @@ namespace AudioWidgets
 		};
 	}
 
-	FWaveformAudioSamplesDataProvider::FWaveformAudioSamplesDataProvider(UWorld* InWorld,
+	FWaveformAudioSamplesDataProvider::FWaveformAudioSamplesDataProvider(const Audio::FDeviceId InAudioDeviceId,
 		UAudioBus* InAudioBus,
 		const uint32 InNumChannelToProvide,
 		const float InTimeWindowMs,
 		const float InMaxTimeWindowMs,
 		const float InAnalysisPeriodMs)
 	{
-		if (!InWorld)
+		if (!InAudioBus)
 		{
-			UE_LOG(LogAudioWidgets, Error, TEXT("World pointer needed to retrieve the AudioDevice is not available."));
+			UE_LOG(LogAudioWidgets, Error, TEXT("Unable to obtain audio samples for visualization without a valid audio bus."));
 			return;
 		}
 
-		if (!InAudioBus)
+		const FAudioDeviceManager* AudioDeviceManager = FAudioDeviceManager::Get();
+		if (!AudioDeviceManager)
 		{
-			UE_LOG(LogAudioWidgets, Error, TEXT("Unable to obtain audio samples for visualization without a valid audio bus."));
+			UE_LOG(LogAudioWidgets, Error, TEXT("Unable to obtain audio samples for visualization without a valid audio device manager."));
 			return;
 		}
 
-		MixerDevice = static_cast<Audio::FMixerDevice*>(InWorld->GetAudioDeviceRaw());
-
+		MixerDevice = static_cast<const Audio::FMixerDevice*>(AudioDeviceManager->GetAudioDeviceRaw(InAudioDeviceId));
 		if (!MixerDevice)
 		{
 			UE_LOG(LogAudioWidgets, Error, TEXT("Unable to obtain audio samples for visualization without a valid audio mixer."));
 			return;
 		}
 
-		WorldPtr = InWorld;
 		AudioBus = InAudioBus;
 
 		NumChannelsToProvide = InNumChannelToProvide;
@@ -115,7 +114,6 @@ namespace AudioWidgets
 			return;
 		}
 
-		check(WorldPtr);
 		check(AudioBus);
 		check(MixerDevice);
 
diff --git a/Engine/Plugins/Runtime/AudioWidgets/Source/AudioWidgets/Public/AudioMeter.h b/Engine/Plugins/Runtime/AudioWidgets/Source/AudioWidgets/Public/AudioMeter.h
--- a/Engine/Plugins/Runtime/AudioWidgets/Source/AudioWidgets/Public/AudioMeter.h
+++ b/Engine/Plugins/Runtime/AudioWidgets/Source/AudioWidgets/Public/AudioMeter.h
@@ -2,6 +2,7 @@
 #pragma once
 
 #include "AudioAnalyzer.h"
+#include "AudioDefines.h"
 #include "AudioMeterStyle.h"
 #include "AudioMeterTypes.h"
 #include "Components/Widget.h"
@@ -155,15 +156,21 @@ namespace AudioWidgets
 	class AUDIOWIDGETS_API FAudioMeter
 	{
 	public:
+		UE_DEPRECATED(5.4, "Use the FAudioMeter constructor that uses Audio::FDeviceId.")
 		FAudioMeter(int32 InNumChannels, UWorld& InWorld, TObjectPtr<UAudioBus> InExternalAudioBus = nullptr); // OPTIONAL PARAM InExternalAudioBus: An audio meter can be constructed from this audio bus.
+		
+		FAudioMeter(const int32 InNumChannels, const Audio::FDeviceId InAudioDeviceId, TObjectPtr<UAudioBus> InExternalAudioBus = nullptr); // OPTIONAL PARAM InExternalAudioBus: An audio meter can be constructed from this audio bus.
 		~FAudioMeter();
 
 		UAudioBus* GetAudioBus() const;
 
 		TSharedRef<SAudioMeter> GetWidget() const;
 
+		UE_DEPRECATED(5.4, "Use the Init method that uses Audio::FDeviceId.")
 		void Init(int32 InNumChannels, UWorld& InWorld, TObjectPtr<UAudioBus> InExternalAudioBus = nullptr);
 
+		void Init(const int32 InNumChannels, const Audio::FDeviceId InAudioDeviceId, TObjectPtr<UAudioBus> InExternalAudioBus);
+
 	protected:
 		void OnMeterOutput(UMeterAnalyzer* InMeterAnalyzer, int32 ChannelIndex, const FMeterResults& InMeterResults);
 
@@ -188,8 +195,6 @@ namespace AudioWidgets
 		/** MetaSound Output Meter widget */
 		TSharedPtr<SAudioMeter> Widget;
 
-		TWeakObjectPtr<UWorld> WorldPtr;
-
 		bool bUseExternalAudioBus = false;
 	};
 } // namespace AudioWidgets
diff --git a/Engine/Plugins/Runtime/AudioWidgets/Source/AudioWidgets/Public/AudioOscilloscope.h b/Engine/Plugins/Runtime/AudioWidgets/Source/AudioWidgets/Public/AudioOscilloscope.h
--- a/Engine/Plugins/Runtime/AudioWidgets/Source/AudioWidgets/Public/AudioOscilloscope.h
+++ b/Engine/Plugins/Runtime/AudioWidgets/Source/AudioWidgets/Public/AudioOscilloscope.h
@@ -2,6 +2,7 @@
 
 #pragma once
 
+#include "AudioDefines.h"
 #include "AudioOscilloscopePanelStyle.h"
 #include "AudioWidgetsEnums.h"
 #include "SAudioOscilloscopePanelWidget.h"
@@ -18,7 +19,7 @@ namespace AudioWidgets
 	class AUDIOWIDGETS_API FAudioOscilloscope
 	{
 	public:
-		FAudioOscilloscope(UWorld* InWorld, 
+		FAudioOscilloscope(Audio::FDeviceId InAudioDeviceId,
 			const uint32 InNumChannels, 
 			const float InTimeWindowMs, 
 			const float InMaxTimeWindowMs, 
@@ -27,7 +28,7 @@ namespace AudioWidgets
 
 		void CreateAudioBus(const uint32 InNumChannels);
 
-		void CreateDataProvider(UWorld* InWorld,
+		void CreateDataProvider(Audio::FDeviceId InAudioDeviceId,
 			const float InTimeWindowMs,
 			const float InMaxTimeWindowMs,
 			const float InAnalysisPeriodMs,
diff --git a/Engine/Plugins/Runtime/AudioWidgets/Source/AudioWidgets/Public/AudioSpectrumAnalyzer.h b/Engine/Plugins/Runtime/AudioWidgets/Source/AudioWidgets/Public/AudioSpectrumAnalyzer.h
--- a/Engine/Plugins/Runtime/AudioWidgets/Source/AudioWidgets/Public/AudioSpectrumAnalyzer.h
+++ b/Engine/Plugins/Runtime/AudioWidgets/Source/AudioWidgets/Public/AudioSpectrumAnalyzer.h
@@ -18,14 +18,14 @@ namespace AudioWidgets
 	class AUDIOWIDGETS_API FAudioSpectrumAnalyzer
 	{
 	public:
-		FAudioSpectrumAnalyzer(int32 InNumChannels, UWorld& InWorld, TObjectPtr<UAudioBus> InExternalAudioBus = nullptr);
+		FAudioSpectrumAnalyzer(int32 InNumChannels, Audio::FDeviceId InAudioDeviceId, TObjectPtr<UAudioBus> InExternalAudioBus = nullptr);
 		~FAudioSpectrumAnalyzer();
 
 		UAudioBus* GetAudioBus() const;
 
 		TSharedRef<SWidget> GetWidget() const;
 
-		void Init(int32 InNumChannels, UWorld& InWorld, TObjectPtr<UAudioBus> InExternalAudioBus = nullptr);
+		void Init(int32 InNumChannels, Audio::FDeviceId InAudioDeviceId, TObjectPtr<UAudioBus> InExternalAudioBus = nullptr);
 
 	protected:
 		void OnConstantQResults(UConstantQAnalyzer* InSpectrumAnalyzer, int32 ChannelIndex, const TArray<FConstantQResults>& InSpectrumResultsArray);
@@ -55,8 +55,6 @@ namespace AudioWidgets
 		/** Slate widget for spectrum display */
 		TSharedPtr<SAudioSpectrumPlot> Widget;
 
-		TWeakObjectPtr<UWorld> WorldPtr;
-
 		bool bUseExternalAudioBus = false;
 
 		TOptional<float> PrevTimeStamp;
diff --git a/Engine/Plugins/Runtime/AudioWidgets/Source/AudioWidgets/Public/AudioVectorscope.h b/Engine/Plugins/Runtime/AudioWidgets/Source/AudioWidgets/Public/AudioVectorscope.h
--- a/Engine/Plugins/Runtime/AudioWidgets/Source/AudioWidgets/Public/AudioVectorscope.h
+++ b/Engine/Plugins/Runtime/AudioWidgets/Source/AudioWidgets/Public/AudioVectorscope.h
@@ -2,6 +2,7 @@
 
 #pragma once
 
+#include "AudioDefines.h"
 #include "AudioVectorscopePanelStyle.h"
 #include "AudioWidgetsEnums.h"
 #include "Sound/AudioBus.h"
@@ -17,7 +18,7 @@ namespace AudioWidgets
 	class AUDIOWIDGETS_API FAudioVectorscope
 	{
 	public:
-		FAudioVectorscope(UWorld* InWorld, 
+		FAudioVectorscope(Audio::FDeviceId InAudioDeviceId,
 			const uint32 InNumChannels, 
 			const float InTimeWindowMs, 
 			const float InMaxTimeWindowMs, 
@@ -25,7 +26,7 @@ namespace AudioWidgets
 			const EAudioPanelLayoutType InPanelLayoutType);
 
 		void CreateAudioBus(const uint32 InNumChannels);
-		void CreateDataProvider(UWorld* InWorld, const float InTimeWindowMs, const float InMaxTimeWindowMs, const float InAnalysisPeriodMs);
+		void CreateDataProvider(Audio::FDeviceId InAudioDeviceId, const float InTimeWindowMs, const float InMaxTimeWindowMs, const float InAnalysisPeriodMs);
 		void CreateVectorscopeWidget(const EAudioPanelLayoutType InPanelLayoutType);
 
 		void StartProcessing();
diff --git a/Engine/Plugins/Runtime/AudioWidgets/Source/AudioWidgets/Public/WaveformAudioSamplesDataProvider.h b/Engine/Plugins/Runtime/AudioWidgets/Source/AudioWidgets/Public/WaveformAudioSamplesDataProvider.h
--- a/Engine/Plugins/Runtime/AudioWidgets/Source/AudioWidgets/Public/WaveformAudioSamplesDataProvider.h
+++ b/Engine/Plugins/Runtime/AudioWidgets/Source/AudioWidgets/Public/WaveformAudioSamplesDataProvider.h
@@ -2,15 +2,20 @@
 
 #pragma once
 
+#include "AudioDefines.h"
 #include "AudioOscilloscopeEnums.h"
 #include "Containers/ArrayView.h"
 #include "Containers/Ticker.h"
 #include "DSP/Dsp.h"
 #include "DSP/MultithreadedPatching.h"
-#include "Engine/World.h"
 #include "FixedSampledSequenceView.h"
 #include "IFixedSampledSequenceViewProvider.h"
 
+namespace Audio
+{
+	class FMixerDevice;
+}
+
 class UAudioBus;
 
 struct FWaveformAudioSamplesResult;
@@ -23,7 +28,7 @@ namespace AudioWidgets
 															   public TSharedFromThis<FWaveformAudioSamplesDataProvider>
 	{
 	public:
-		FWaveformAudioSamplesDataProvider(UWorld* InWorld, 
+		FWaveformAudioSamplesDataProvider(const Audio::FDeviceId InAudioDeviceId,
 			UAudioBus* InAudioBus, 
 			const uint32 InNumChannelToProvide, 
 			const float InTimeWindowMs, 
@@ -55,7 +60,6 @@ namespace AudioWidgets
 		void PushAudioSamplesToCircularBuffer();
 		bool Tick(float DeltaTime);
 
-		UWorld* WorldPtr = nullptr;
 		FTSTicker::FDelegateHandle TickerHandle = nullptr;
 
 		uint32 NumChannelsToProvide = 0;
diff --git a/Engine/Plugins/Runtime/Metasound/Source/MetasoundEditor/Private/MetasoundEditor.cpp b/Engine/Plugins/Runtime/Metasound/Source/MetasoundEditor/Private/MetasoundEditor.cpp
--- a/Engine/Plugins/Runtime/Metasound/Source/MetasoundEditor/Private/MetasoundEditor.cpp
+++ b/Engine/Plugins/Runtime/Metasound/Source/MetasoundEditor/Private/MetasoundEditor.cpp
@@ -1170,67 +1170,65 @@ namespace Metasound
 			{
 				if (ensure(GEditor))
 				{
-					UWorld* EditorWorld = GEditor->GetEditorWorldContext().World();
-					if (ensure(EditorWorld))
+					const Audio::FDeviceId AudioDeviceId = GEditor->GetMainAudioDeviceID();
+
+					if (!OutputMeter.IsValid())
 					{
-						if (!OutputMeter.IsValid())
-						{
-							OutputMeter = MakeShared<AudioWidgets::FAudioMeter>(MetaSoundSource->NumChannels, *EditorWorld);
-						}
+						OutputMeter = MakeShared<AudioWidgets::FAudioMeter>(MetaSoundSource->NumChannels, AudioDeviceId);
+					}
 
-						const uint32 MetaSoundNumChannels = static_cast<uint32>(MetaSoundSource->NumChannels);
+					const uint32 MetaSoundNumChannels = static_cast<uint32>(MetaSoundSource->NumChannels);
 
-						// Init Oscilloscope
-						constexpr float OscilloscopeTimeWindowMs     = 10.0f;
-						constexpr float OscilloscopeMaxTimeWindowMs  = 10.0f;
-						constexpr float OscilloscopeAnalysisPeriodMs = 10.0f;
-						constexpr EAudioPanelLayoutType OscilloscopePanelLayoutType = EAudioPanelLayoutType::Basic;
+					// Init Oscilloscope
+					constexpr float OscilloscopeTimeWindowMs     = 10.0f;
+					constexpr float OscilloscopeMaxTimeWindowMs  = 10.0f;
+					constexpr float OscilloscopeAnalysisPeriodMs = 10.0f;
+					constexpr EAudioPanelLayoutType OscilloscopePanelLayoutType = EAudioPanelLayoutType::Basic;
 
-						if (!OutputOscilloscope.IsValid())
-						{
-							OutputOscilloscope = MakeShared<AudioWidgets::FAudioOscilloscope>(EditorWorld,
-								MetaSoundNumChannels,
-								OscilloscopeTimeWindowMs,
-								OscilloscopeMaxTimeWindowMs,
-								OscilloscopeAnalysisPeriodMs,
-								OscilloscopePanelLayoutType);
-						}
-						else
-						{
-							OutputOscilloscope->CreateAudioBus(MetaSoundNumChannels);
-							OutputOscilloscope->CreateDataProvider(EditorWorld,	OscilloscopeTimeWindowMs, OscilloscopeMaxTimeWindowMs, OscilloscopeAnalysisPeriodMs, OscilloscopePanelLayoutType);
-							OutputOscilloscope->CreateOscilloscopeWidget(MetaSoundNumChannels, EAudioPanelLayoutType::Basic);
-						}
+					if (!OutputOscilloscope.IsValid())
+					{
+						OutputOscilloscope = MakeShared<AudioWidgets::FAudioOscilloscope>(AudioDeviceId,
+							MetaSoundNumChannels,
+							OscilloscopeTimeWindowMs,
+							OscilloscopeMaxTimeWindowMs,
+							OscilloscopeAnalysisPeriodMs,
+							OscilloscopePanelLayoutType);
+					}
+					else
+					{
+						OutputOscilloscope->CreateAudioBus(MetaSoundNumChannels);
+						OutputOscilloscope->CreateDataProvider(AudioDeviceId, OscilloscopeTimeWindowMs, OscilloscopeMaxTimeWindowMs, OscilloscopeAnalysisPeriodMs, OscilloscopePanelLayoutType);
+						OutputOscilloscope->CreateOscilloscopeWidget(MetaSoundNumChannels, EAudioPanelLayoutType::Basic);
+					}
 
-						// Init Vectorscope
-						constexpr float VectorscopeTimeWindowMs     = 30.0f;
-						constexpr float VectorscopeMaxTimeWindowMs  = 30.0f;
-						constexpr float VectorscopeAnalysisPeriodMs = 10.0f;
-						constexpr EAudioPanelLayoutType VectorscopePanelLayoutType = EAudioPanelLayoutType::Basic;
+					// Init Vectorscope
+					constexpr float VectorscopeTimeWindowMs     = 30.0f;
+					constexpr float VectorscopeMaxTimeWindowMs  = 30.0f;
+					constexpr float VectorscopeAnalysisPeriodMs = 10.0f;
+					constexpr EAudioPanelLayoutType VectorscopePanelLayoutType = EAudioPanelLayoutType::Basic;
 
-						if (!OutputVectorscope.IsValid())
-						{
-							OutputVectorscope = MakeShared<AudioWidgets::FAudioVectorscope>(EditorWorld,
-								MetaSoundNumChannels,
-								VectorscopeTimeWindowMs,
-								VectorscopeMaxTimeWindowMs,
-								VectorscopeAnalysisPeriodMs,
-								VectorscopePanelLayoutType);
-						}
-						else
-						{
-							OutputVectorscope->CreateAudioBus(MetaSoundNumChannels);
-							OutputVectorscope->CreateDataProvider(EditorWorld, VectorscopeTimeWindowMs, VectorscopeMaxTimeWindowMs, VectorscopeAnalysisPeriodMs);
-							OutputVectorscope->CreateVectorscopeWidget(VectorscopePanelLayoutType);
-						}
-
-						if (!OutputSpectrumAnalyzer.IsValid())
-						{
-							OutputSpectrumAnalyzer = MakeShared<AudioWidgets::FAudioSpectrumAnalyzer>(MetaSoundSource->NumChannels, *EditorWorld);
-						}
+					if (!OutputVectorscope.IsValid())
+					{
+						OutputVectorscope = MakeShared<AudioWidgets::FAudioVectorscope>(AudioDeviceId,
+							MetaSoundNumChannels,
+							VectorscopeTimeWindowMs,
+							VectorscopeMaxTimeWindowMs,
+							VectorscopeAnalysisPeriodMs,
+							VectorscopePanelLayoutType);
+					}
+					else
+					{
+						OutputVectorscope->CreateAudioBus(MetaSoundNumChannels);
+						OutputVectorscope->CreateDataProvider(AudioDeviceId, VectorscopeTimeWindowMs, VectorscopeMaxTimeWindowMs, VectorscopeAnalysisPeriodMs);
+						OutputVectorscope->CreateVectorscopeWidget(VectorscopePanelLayoutType);
+					}
 
-						return;
+					if (!OutputSpectrumAnalyzer.IsValid())
+					{
+						OutputSpectrumAnalyzer = MakeShared<AudioWidgets::FAudioSpectrumAnalyzer>(MetaSoundSource->NumChannels, AudioDeviceId);
 					}
+
+					return;
 				}
 			}
 
diff --git a/Engine/Plugins/Runtime/Metasound/Source/MetasoundEditor/Private/MetasoundPinAudioInspector.cpp b/Engine/Plugins/Runtime/Metasound/Source/MetasoundEditor/Private/MetasoundPinAudioInspector.cpp
--- a/Engine/Plugins/Runtime/Metasound/Source/MetasoundEditor/Private/MetasoundPinAudioInspector.cpp
+++ b/Engine/Plugins/Runtime/Metasound/Source/MetasoundEditor/Private/MetasoundPinAudioInspector.cpp
@@ -3,6 +3,8 @@
 
 #include "Analysis/MetasoundFrontendVertexAnalyzerAudioBuffer.h"
 #include "AudioBusSubsystem.h"
+#include "AudioDefines.h"
+#include "AudioDeviceManager.h"
 #include "AudioMixerDevice.h"
 #include "AudioOscilloscope.h"
 #include "MetasoundEditor.h"
@@ -36,18 +38,19 @@ namespace Metasound
 		FMetasoundPinAudioInspector::FMetasoundPinAudioInspector(FEdGraphPinReference InPinRef)
 			: GraphPinObj(FMetasoundPinAudioInspectorPrivate::ResolvePinObjectAsOutput(InPinRef.Get()))
 		{
+			using namespace Audio;
+
 			// Initialize Oscilloscope
-			UWorld* EditorWorld = GEditor->GetEditorWorldContext().World();
-			check(EditorWorld);
+			const FDeviceId AudioDeviceId = GEditor->GetMainAudioDeviceID();
 
 			constexpr int32 NumChannels = 1; // Audio wires are currently mono signals
 
-			Oscilloscope = MakeShared<AudioWidgets::FAudioOscilloscope>(EditorWorld,
-																		NumChannels,
-																		/*InTimeWindowMs*/     10.0f,
-																		/*InMaxTimeWindowMs*/  10.0f,
-																		/*InAnalysisPeriodMs*/ 10.0f,
-																		/*InPanelLayoutType*/  EAudioPanelLayoutType::Basic);
+			Oscilloscope = MakeShared<AudioWidgets::FAudioOscilloscope>(AudioDeviceId,
+				NumChannels,
+				/*InTimeWindowMs*/     10.0f,
+				/*InMaxTimeWindowMs*/  10.0f,
+				/*InAnalysisPeriodMs*/ 10.0f,
+				/*InPanelLayoutType*/  EAudioPanelLayoutType::Basic);
 
 			PinAudioInspectorWidget = SNew(SMetasoundPinAudioInspector)
 									  .VisualizationWidget(Oscilloscope->GetPanelWidget());
@@ -55,14 +58,17 @@ namespace Metasound
 			Oscilloscope->StartProcessing();
 
 			// Set PatchInput from Oscilloscope AudioBus
-			Audio::FMixerDevice* MixerDevice = static_cast<Audio::FMixerDevice*>(EditorWorld->GetAudioDeviceRaw());
+			const FAudioDeviceManager* AudioDeviceManager = FAudioDeviceManager::Get();
+			check(AudioDeviceManager);
+
+			const FMixerDevice* MixerDevice = static_cast<const FMixerDevice*>(AudioDeviceManager->GetAudioDeviceRaw(AudioDeviceId));
 			check(MixerDevice);
 
 			UAudioBusSubsystem* AudioBusSubsystem = MixerDevice->GetSubsystem<UAudioBusSubsystem>();
 			check(AudioBusSubsystem);
 
 			const uint32 AudioBusId = Oscilloscope->GetAudioBus()->GetUniqueID();
-			PatchInput = AudioBusSubsystem->AddPatchInputForAudioBus(Audio::FAudioBusKey(AudioBusId), MixerDevice->GetNumOutputFrames(), NumChannels);
+			PatchInput = AudioBusSubsystem->AddPatchInputForAudioBus(FAudioBusKey(AudioBusId), MixerDevice->GetNumOutputFrames(), NumChannels);
 
 			// Track Audio Pin
 			if (GraphPinObj && !GraphPinObj->LinkedTo.IsEmpty() && GraphPinObj->PinType.PinCategory == FGraphBuilder::PinCategoryAudio)
diff --git a/Engine/Source/Runtime/AudioAnalyzer/Classes/AudioAnalyzer.h b/Engine/Source/Runtime/AudioAnalyzer/Classes/AudioAnalyzer.h
--- a/Engine/Source/Runtime/AudioAnalyzer/Classes/AudioAnalyzer.h
+++ b/Engine/Source/Runtime/AudioAnalyzer/Classes/AudioAnalyzer.h
@@ -3,12 +3,14 @@
 #pragma once
 
 #include "Async/AsyncWork.h"
-#include "CoreMinimal.h"
-#include "IAudioAnalyzerInterface.h"
 #include "AudioAnalyzerAsset.h"
 #include "AudioAnalyzerFacade.h"
-#include "Sound/AudioBus.h"
+#include "AudioDefines.h"
+#include "CoreMinimal.h"
 #include "DSP/MultithreadedPatching.h"
+#include "IAudioAnalyzerInterface.h"
+#include "Sound/AudioBus.h"
+
 #include "AudioAnalyzer.generated.h"
 
 class UAudioAnalyzerSubsystem;
@@ -94,8 +96,11 @@ public:
 	AUDIOANALYZER_API void StartAnalyzing(const UObject* WorldContextObject, UAudioBus* AudioBusToAnalyze);
 
 	/** Starts analyzing using the given world.*/
+	UE_DEPRECATED(5.4, "Use the StartAnalyzing method that uses Audio::FDeviceId.")
 	AUDIOANALYZER_API void StartAnalyzing(UWorld* InWorld, UAudioBus* AudioBusToAnalyze);
 
+	AUDIOANALYZER_API void StartAnalyzing(const Audio::FDeviceId InAudioDeviceId, UAudioBus* AudioBusToAnalyze);
+
 	/** Stops analyzing audio. */
 	UFUNCTION(BlueprintCallable, BlueprintCosmetic, Category = AudioAnalyzer, meta = (WorldContext = "WorldContextObject"))
 	AUDIOANALYZER_API void StopAnalyzing(const UObject* WorldContextObject = nullptr);
diff --git a/Engine/Source/Runtime/AudioAnalyzer/Private/AudioAnalyzer.cpp b/Engine/Source/Runtime/AudioAnalyzer/Private/AudioAnalyzer.cpp
--- a/Engine/Source/Runtime/AudioAnalyzer/Private/AudioAnalyzer.cpp
+++ b/Engine/Source/Runtime/AudioAnalyzer/Private/AudioAnalyzer.cpp
@@ -1,14 +1,15 @@
 // Copyright Epic Games, Inc. All Rights Reserved.
 
-
 #include "AudioAnalyzer.h"
+
+#include "Async/Async.h"
 #include "AudioAnalyzerFacade.h"
 #include "AudioAnalyzerModule.h"
-#include "Async/Async.h"
-#include "AudioMixerDevice.h"
-#include "AudioDeviceManager.h"
 #include "AudioAnalyzerSubsystem.h"
 #include "AudioBusSubsystem.h"
+#include "AudioDeviceHandle.h"
+#include "AudioDeviceManager.h"
+#include "AudioMixerDevice.h"
 #include "Engine/Engine.h"
 
 #include UE_INLINE_GENERATED_CPP_BY_NAME(AudioAnalyzer)
@@ -38,6 +39,22 @@ void FAudioAnalyzeTask::DoWork()
 }
 
 void UAudioAnalyzer::StartAnalyzing(UWorld* InWorld, UAudioBus* AudioBusToAnalyze)
+{
+	if (!InWorld)
+	{
+		return;
+	}
+
+	const FAudioDeviceHandle AudioDevice = InWorld->GetAudioDevice();
+	if (!AudioDevice.IsValid())
+	{
+		return;
+	}
+
+	StartAnalyzing(AudioDevice.GetDeviceID(), AudioBusToAnalyze);
+}
+
+void UAudioAnalyzer::StartAnalyzing(const Audio::FDeviceId InAudioDeviceId, UAudioBus* AudioBusToAnalyze)
 {
 	if (!AudioBusToAnalyze)
 	{
@@ -45,12 +62,14 @@ void UAudioAnalyzer::StartAnalyzing(UWorld* InWorld, UAudioBus* AudioBusToAnalyz
 		return;
 	}
 
-	if (!InWorld)
+	const FAudioDeviceManager* AudioDeviceManager = FAudioDeviceManager::Get();
+	if (!AudioDeviceManager)
 	{
+		UE_LOG(LogAudioAnalyzer, Error, TEXT("Unable to analyze audio with a null audio device manager."));
 		return;
 	}
 
-	Audio::FMixerDevice* MixerDevice = static_cast<Audio::FMixerDevice*>(InWorld->GetAudioDeviceRaw());
+	const Audio::FMixerDevice* MixerDevice = static_cast<const Audio::FMixerDevice*>(AudioDeviceManager->GetAudioDeviceRaw(InAudioDeviceId));
 
 	if (!MixerDevice)
 	{
@@ -120,7 +139,13 @@ void UAudioAnalyzer::StartAnalyzing(const UObject* WorldContextObject, UAudioBus
 		return;
 	}
 
-	StartAnalyzing(ThisWorld, AudioBusToAnalyze);
+	const FAudioDeviceHandle AudioDevice = ThisWorld->GetAudioDevice();
+	if (!AudioDevice.IsValid())
+	{
+		return;
+	}
+
+	StartAnalyzing(AudioDevice.GetDeviceID(), AudioBusToAnalyze);
 }
 
 void UAudioAnalyzer::StopAnalyzing(const UObject* WorldContextObject)
