Fix for compiler errors in shipping builds due to insights calling into AudioDebugger which is not available therein
#tests

diff --git a/Engine/Plugins/AudioInsights/Source/AudioInsights/Private/Views/MixerSourceDashboardViewFactory.cpp b/Engine/Plugins/AudioInsights/Source/AudioInsights/Private/Views/MixerSourceDashboardViewFactory.cpp
--- a/Engine/Plugins/AudioInsights/Source/AudioInsights/Private/Views/MixerSourceDashboardViewFactory.cpp
+++ b/Engine/Plugins/AudioInsights/Source/AudioInsights/Private/Views/MixerSourceDashboardViewFactory.cpp
@@ -359,29 +359,28 @@ namespace UE::Audio::Insights
 
 	FSlateColor FMixerSourceDashboardViewFactory::GetRowColor(const TSharedPtr<IDashboardDataViewEntry>& InRowDataPtr)
 	{
-		bool bIsSolo = false;
-		bool bIsMute = false;
-
 		FColor RowTextColor(255, 255, 255);
 
-		if (FAudioDeviceManager* AudioDeviceManager = FAudioDeviceManager::Get())
+#if ENABLE_AUDIO_DEBUG
+		if (const FAudioDeviceManager* AudioDeviceManager = FAudioDeviceManager::Get())
 		{
 			const FSoundAssetDashboardEntry& SoundAssetDashboardEntry = *StaticCastSharedPtr<FSoundAssetDashboardEntry>(InRowDataPtr).Get();
-
-			const FName SoundAssetName(SoundAssetDashboardEntry.GetDisplayName().ToString());
-
-			bIsSolo = AudioDeviceManager->GetDebugger().IsSoloSoundWave(SoundAssetName);
-			bIsMute = AudioDeviceManager->GetDebugger().IsMuteSoundWave(SoundAssetName);
-
+			const FName SoundAssetName { SoundAssetDashboardEntry.Name };
+			const bool bIsSolo = AudioDeviceManager->GetDebugger().IsSoloSoundWave(SoundAssetName);
 			if (bIsSolo)
 			{
 				RowTextColor = FColor(255, 255, 0);
 			}
-			else if (bIsMute)
+			else
 			{
-				RowTextColor = FColor(255, 0, 0);
+				const bool bIsMute = AudioDeviceManager->GetDebugger().IsMuteSoundWave(SoundAssetName);
+				if (bIsMute)
+				{
+					RowTextColor = FColor(255, 0, 0);
+				}
 			}
 		}
+#endif // ENABLE_AUDIO_DEBUG
 
 		return FSlateColor(RowTextColor);
 	}
@@ -401,7 +400,7 @@ namespace UE::Audio::Insights
 					{
 						const FSoundAssetDashboardEntry& SoundAssetDashboardEntry = *StaticCastSharedPtr<FSoundAssetDashboardEntry>(Item).Get();
 
-						AudioDeviceManager->GetDebugger().SetMuteSoundWave(FName(SoundAssetDashboardEntry.GetDisplayName().ToString()), NewState == ECheckBoxState::Checked);
+						AudioDeviceManager->GetDebugger().SetMuteSoundWave(FName { SoundAssetDashboardEntry.Name }, NewState == ECheckBoxState::Checked);
 					}
 				}
 			}
@@ -424,7 +423,7 @@ namespace UE::Audio::Insights
 					{
 						const FSoundAssetDashboardEntry& SoundAssetDashboardEntry = *StaticCastSharedPtr<FSoundAssetDashboardEntry>(Item).Get();
 
-						AudioDeviceManager->GetDebugger().SetSoloSoundWave(FName(SoundAssetDashboardEntry.GetDisplayName().ToString()), NewState == ECheckBoxState::Checked);
+						AudioDeviceManager->GetDebugger().SetSoloSoundWave(FName { SoundAssetDashboardEntry.Name }, NewState == ECheckBoxState::Checked);
 					}
 				}
 			}
@@ -446,8 +445,7 @@ namespace UE::Audio::Insights
 					if (SelectedItem.IsValid())
 					{
 						const FSoundAssetDashboardEntry& SoundAssetDashboardEntry = *StaticCastSharedPtr<FSoundAssetDashboardEntry>(SelectedItem).Get();
-						const FName SoundAssetDisplayName = FName(SoundAssetDashboardEntry.GetDisplayName().ToString());
-
+						const FName SoundAssetDisplayName { SoundAssetDashboardEntry.Name };
 						AudioDeviceManager->GetDebugger().ToggleMuteSoundWave(SoundAssetDisplayName);
 					}
 				}
@@ -465,8 +463,7 @@ namespace UE::Audio::Insights
 					}
 
 					const FSoundAssetDashboardEntry& SoundAssetDashboardEntry = *StaticCastSharedPtr<FSoundAssetDashboardEntry>(Item).Get();
-					const FName SoundAssetDisplayName = FName(SoundAssetDashboardEntry.GetDisplayName().ToString());
-
+					const FName SoundAssetDisplayName { SoundAssetDashboardEntry.Name };
 					if (AudioDeviceManager->GetDebugger().IsMuteSoundWave(SoundAssetDisplayName))
 					{
 						bIsAnySoundMuted = true;
@@ -494,8 +491,7 @@ namespace UE::Audio::Insights
 					if (SelectedItem.IsValid())
 					{
 						const FSoundAssetDashboardEntry& SoundAssetDashboardEntry = *StaticCastSharedPtr<FSoundAssetDashboardEntry>(SelectedItem).Get();
-						const FName SoundAssetDisplayName = FName(SoundAssetDashboardEntry.GetDisplayName().ToString());
-
+						const FName SoundAssetDisplayName { SoundAssetDashboardEntry.Name };
 						AudioDeviceManager->GetDebugger().ToggleSoloSoundWave(SoundAssetDisplayName);
 					}
 				}
@@ -513,8 +509,7 @@ namespace UE::Audio::Insights
 					}
 
 					const FSoundAssetDashboardEntry& SoundAssetDashboardEntry = *StaticCastSharedPtr<FSoundAssetDashboardEntry>(Item).Get();
-					const FName SoundAssetDisplayName = FName(SoundAssetDashboardEntry.GetDisplayName().ToString());
-
+					const FName SoundAssetDisplayName { SoundAssetDashboardEntry.Name };
 					if (AudioDeviceManager->GetDebugger().IsSoloSoundWave(SoundAssetDisplayName))
 					{
 						bIsAnySoundSoloed = true;
diff --git a/Engine/Plugins/AudioInsights/Source/AudioInsights/Private/Views/VirtualLoopDashboardViewFactory.cpp b/Engine/Plugins/AudioInsights/Source/AudioInsights/Private/Views/VirtualLoopDashboardViewFactory.cpp
--- a/Engine/Plugins/AudioInsights/Source/AudioInsights/Private/Views/VirtualLoopDashboardViewFactory.cpp
+++ b/Engine/Plugins/AudioInsights/Source/AudioInsights/Private/Views/VirtualLoopDashboardViewFactory.cpp
@@ -248,10 +248,12 @@ namespace UE::Audio::Insights
 #if WITH_EDITOR
 	bool FVirtualLoopDashboardViewFactory::IsDebugDrawEnabled() const
 	{
+#if ENABLE_AUDIO_DEBUG
 		if (FAudioDeviceManager* Manager = FAudioDeviceManager::Get())
 		{
 			return !(Manager->IsVisualizeDebug3dEnabled() && !::Audio::FAudioDebugger::IsVirtualLoopVisualizeEnabled());
 		}
+#endif // ENABLE_AUDIO_DEBUG
 
 		return false;
 	}
