Minor Audio Insights plugin cleanup (change plugin from beta to experimental, remove Audio Insights profiling menu option, and rename Audio Dashboard to Audio Insights)

diff --git a/Engine/Plugins/AudioInsights/AudioInsights.uplugin b/Engine/Plugins/AudioInsights/AudioInsights.uplugin
--- a/Engine/Plugins/AudioInsights/AudioInsights.uplugin
+++ b/Engine/Plugins/AudioInsights/AudioInsights.uplugin
@@ -11,8 +11,8 @@
 	"MarketplaceURL": "",
 	"SupportURL": "",
 	"CanContainContent": true,
-	"IsBetaVersion": true,
-	"IsExperimentalVersion": false,
+	"IsBetaVersion": false,
+	"IsExperimentalVersion": true,
 	"Installed": false,
 	"Modules":
 	[
diff --git a/Engine/Plugins/AudioInsights/Source/AudioInsights/Private/AudioInsightsDashboardFactory.cpp b/Engine/Plugins/AudioInsights/Source/AudioInsights/Private/AudioInsightsDashboardFactory.cpp
--- a/Engine/Plugins/AudioInsights/Source/AudioInsights/Private/AudioInsightsDashboardFactory.cpp
+++ b/Engine/Plugins/AudioInsights/Source/AudioInsights/Private/AudioInsightsDashboardFactory.cpp
@@ -22,7 +22,7 @@ namespace UE::Audio::Insights
 {
 	namespace DashboardFactoryPrivate
 	{
-		static const FText ToolName = LOCTEXT("AudioDashboard_ToolName", "Audio Dashboard");
+		static const FText ToolName = LOCTEXT("AudioDashboard_ToolName", "Audio Insights");
 
 		static const FLazyName MainToolbarName = "MainToolbar";
 		static const FText MainToolbarDisplayName = LOCTEXT("AudioDashboard_MainToolbarDisplayName", "Dashboard Transport");
@@ -554,7 +554,7 @@ namespace UE::Audio::Insights
 	void FDashboardFactory::RegisterViewFactory(TSharedRef<IDashboardViewFactory> InFactory)
 	{
 		const FName Name = InFactory->GetName();
-		if (ensureAlwaysMsgf(!DashboardViewFactories.Contains(Name), TEXT("Failed to register Audio Dashboard '%s': Dashboard with name already registered"), *Name.ToString()))
+		if (ensureAlwaysMsgf(!DashboardViewFactories.Contains(Name), TEXT("Failed to register Audio Insights Dashboard '%s': Dashboard with name already registered"), *Name.ToString()))
 		{
 			DashboardViewFactories.Add(Name, InFactory);
 		}
diff --git a/Engine/Plugins/AudioInsights/Source/AudioInsights/Private/AudioInsightsModule.cpp b/Engine/Plugins/AudioInsights/Source/AudioInsights/Private/AudioInsightsModule.cpp
--- a/Engine/Plugins/AudioInsights/Source/AudioInsights/Private/AudioInsightsModule.cpp
+++ b/Engine/Plugins/AudioInsights/Source/AudioInsights/Private/AudioInsightsModule.cpp
@@ -98,26 +98,14 @@ namespace UE::Audio::Insights
 		return DashboardFactory->MakeDockTabWidget(Args);
 	}
 
-	TSharedRef<SDockTab> FAudioInsightsModule::CreateInsightsTabWidget(const FSpawnTabArgs& Args)
-	{
-		// TODO: Implement
-		return SNew(SDockTab);
-	}
-
 	void FAudioInsightsModule::RegisterMenus()
 	{
 		const IWorkspaceMenuStructure& MenuStructure = WorkspaceMenu::GetMenuStructure();
-		FGlobalTabmanager::Get()->RegisterNomadTabSpawner("AudioDashboard", FOnSpawnTab::CreateRaw(this, &FAudioInsightsModule::CreateDashboardTabWidget))
-			.SetDisplayName(LOCTEXT("OpenDashboard_TabDisplayName", "Audio Dashboard"))
-			.SetTooltipText(LOCTEXT("OpenDashboard_TabTooltip", "Opens Audio Dashboard, an extensible suite of tools and visualizers which enable monitoring and debugging audio in the Unreal Engine."))
+		FGlobalTabmanager::Get()->RegisterNomadTabSpawner("AudioInsights", FOnSpawnTab::CreateRaw(this, &FAudioInsightsModule::CreateDashboardTabWidget))
+			.SetDisplayName(LOCTEXT("OpenDashboard_TabDisplayName", "Audio Insights"))
+			.SetTooltipText(LOCTEXT("OpenDashboard_TabTooltip", "Opens Audio Insights, an extensible suite of tools and visualizers which enable monitoring and debugging audio in the Unreal Engine."))
 			.SetGroup(MenuStructure.GetToolsCategory())
 			.SetIcon(FSlateStyle::Get().CreateIcon("AudioInsights.Icon.Dashboard"));
-
-		FGlobalTabmanager::Get()->RegisterNomadTabSpawner(GetName(), FOnSpawnTab::CreateRaw(this, &FAudioInsightsModule::CreateInsightsTabWidget))
-			.SetDisplayName(LOCTEXT("OpenAudioInsights_TabDisplayName", "Audio Insights"))
-			.SetTooltipText(LOCTEXT("OpenAudioInsights_TabTooltip", "Opens Insights audio profile mode."))
-			.SetGroup(MenuStructure.GetDeveloperToolsProfilingCategory())
-			.SetIcon(FSlateStyle::Get().CreateIcon("AudioInsights.Icon"));
 	};
 } // namespace UE::Audio::Insights
 #undef LOCTEXT_NAMESPACE // AudioInsights
diff --git a/Engine/Plugins/AudioInsights/Source/AudioInsights/Private/AudioInsightsModule.h b/Engine/Plugins/AudioInsights/Source/AudioInsights/Private/AudioInsightsModule.h
--- a/Engine/Plugins/AudioInsights/Source/AudioInsights/Private/AudioInsightsModule.h
+++ b/Engine/Plugins/AudioInsights/Source/AudioInsights/Private/AudioInsightsModule.h
@@ -40,7 +40,6 @@ namespace UE::Audio::Insights
 
 	private:
 		TSharedRef<SDockTab> CreateDashboardTabWidget(const FSpawnTabArgs& Args);
-		TSharedRef<SDockTab> CreateInsightsTabWidget(const FSpawnTabArgs& Args);
 		void RegisterMenus();
 	};
 } // namespace UE::Audio::Insights
