Change audio insights to be exactly like old command line utilities for solo and mute "exclusive" mode.
Much more intuitive behavior with the new audio insights tool.

diff --git a/Engine/Plugins/AudioInsights/Source/AudioInsights/Private/Views/MixerSourceDashboardViewFactory.cpp b/Engine/Plugins/AudioInsights/Source/AudioInsights/Private/Views/MixerSourceDashboardViewFactory.cpp
--- a/Engine/Plugins/AudioInsights/Source/AudioInsights/Private/Views/MixerSourceDashboardViewFactory.cpp
+++ b/Engine/Plugins/AudioInsights/Source/AudioInsights/Private/Views/MixerSourceDashboardViewFactory.cpp
@@ -388,22 +388,10 @@ namespace UE::Audio::Insights
 	void FMixerSourceDashboardViewFactory::ToggleMuteForAllItems(ECheckBoxState NewState)
 	{
 #if ENABLE_AUDIO_DEBUG
-		if (FilteredEntriesListView.IsValid())
+		if (MuteState != NewState)
 		{
-			if (FAudioDeviceManager* AudioDeviceManager = FAudioDeviceManager::Get())
-			{
-				const TArrayView<const TSharedPtr<IDashboardDataViewEntry>> TableItems = FilteredEntriesListView->GetItems();
-
-				for (const TSharedPtr<IDashboardDataViewEntry>& Item : TableItems)
-				{
-					if (Item.IsValid())
-					{
-						const FSoundAssetDashboardEntry& SoundAssetDashboardEntry = *StaticCastSharedPtr<FSoundAssetDashboardEntry>(Item).Get();
-
-						AudioDeviceManager->GetDebugger().SetMuteSoundWave(FName { SoundAssetDashboardEntry.Name }, NewState == ECheckBoxState::Checked);
-					}
-				}
-			}
+			MuteState = NewState;
+			UpdateSoloMuteState();
 		}
 #endif
 	}
@@ -411,22 +399,10 @@ namespace UE::Audio::Insights
 	void FMixerSourceDashboardViewFactory::ToggleSoloForAllItems(ECheckBoxState NewState)
 	{
 #if ENABLE_AUDIO_DEBUG
-		if (FilteredEntriesListView.IsValid())
+		if (SoloState != NewState)
 		{
-			if (FAudioDeviceManager* AudioDeviceManager = FAudioDeviceManager::Get())
-			{
-				const TArrayView<const TSharedPtr<IDashboardDataViewEntry>> TableItems = FilteredEntriesListView->GetItems();
-
-				for (const TSharedPtr<IDashboardDataViewEntry>& Item : TableItems)
-				{
-					if (Item.IsValid())
-					{
-						const FSoundAssetDashboardEntry& SoundAssetDashboardEntry = *StaticCastSharedPtr<FSoundAssetDashboardEntry>(Item).Get();
-
-						AudioDeviceManager->GetDebugger().SetSoloSoundWave(FName { SoundAssetDashboardEntry.Name }, NewState == ECheckBoxState::Checked);
-					}
-				}
-			}
+			SoloState = NewState;
+			UpdateSoloMuteState();
 		}
 #endif
 	}
@@ -536,6 +512,33 @@ namespace UE::Audio::Insights
 #endif
 	}
 
+	void FMixerSourceDashboardViewFactory::UpdateSoloMuteState()
+	{
+#if ENABLE_AUDIO_DEBUG
+		if (FAudioDeviceManager* AudioDeviceManager = FAudioDeviceManager::Get())
+		{
+			FName CurrentFilterStringName = FName{ CurrentFilterString };
+			if (MuteState == ECheckBoxState::Checked && !CurrentFilterString.IsEmpty())
+			{
+				AudioDeviceManager->GetDebugger().ToggleMuteSoundWave(CurrentFilterStringName, true);
+			}
+			else
+			{
+				AudioDeviceManager->GetDebugger().ToggleMuteSoundWave(NAME_None, true);
+			}
+
+			if (SoloState == ECheckBoxState::Checked && !CurrentFilterString.IsEmpty())
+			{
+				AudioDeviceManager->GetDebugger().ToggleSoloSoundWave(CurrentFilterStringName, true);
+			}
+			else
+			{
+				AudioDeviceManager->GetDebugger().ToggleSoloSoundWave(NAME_None, true);
+			}
+		}
+#endif
+	}
+
 	void FMixerSourceDashboardViewFactory::UpdatePlotsWidgetsData()
 	{
 		if (DataViewEntries.Num() <= 0)
@@ -966,6 +969,15 @@ namespace UE::Audio::Insights
 		});
 
 		UpdatePlotsWidgetsData();
+
+#if ENABLE_AUDIO_DEBUG
+		// Update the mute and solo states if the filter string changes
+		if (CurrentFilterString != FilterString)
+		{
+			CurrentFilterString = FilterString;
+			UpdateSoloMuteState();
+		}
+#endif
 	}
 
 #if WITH_EDITOR
diff --git a/Engine/Plugins/AudioInsights/Source/AudioInsights/Private/Views/MixerSourceDashboardViewFactory.h b/Engine/Plugins/AudioInsights/Source/AudioInsights/Private/Views/MixerSourceDashboardViewFactory.h
--- a/Engine/Plugins/AudioInsights/Source/AudioInsights/Private/Views/MixerSourceDashboardViewFactory.h
+++ b/Engine/Plugins/AudioInsights/Source/AudioInsights/Private/Views/MixerSourceDashboardViewFactory.h
@@ -43,6 +43,8 @@ namespace UE::Audio::Insights
 
 	private:
 		void UpdatePlotsWidgetsData();
+		void UpdateSoloMuteState();
+
 		// Column information used by plot widgets, keyed by column name. These keys should be a subset of the keys in GetColumns(). 
 		const TMap<FName, FPlotColumnInfo>& GetPlotColumnInfo();
 		const TFunctionRef<float(const IDashboardDataViewEntry& InData)> GetPlotColumnDataFunc(const FName& ColumnName);
@@ -86,5 +88,12 @@ namespace UE::Audio::Insights
 		const static int32 NumPlotWidgets = 2;
 		TArray<FName> SelectedPlotColumnNames;
 		TArray<TSharedPtr<SAudioCurveView>> PlotWidgets;
+
+		// State of the mute and solo buttons
+#if ENABLE_AUDIO_DEBUG
+		ECheckBoxState MuteState = ECheckBoxState::Unchecked;
+		ECheckBoxState SoloState = ECheckBoxState::Unchecked;
+		FString CurrentFilterString;
+#endif
 	};
 } // namespace UE::Audio::Insights
