[AudioInsights] Implemented Main Menu Bar + Misc. layout adjustments

diff --git a/Engine/Plugins/AudioInsights/Source/AudioInsights/Private/AudioInsightsDashboardFactory.cpp b/Engine/Plugins/AudioInsights/Source/AudioInsights/Private/AudioInsightsDashboardFactory.cpp
--- a/Engine/Plugins/AudioInsights/Source/AudioInsights/Private/AudioInsightsDashboardFactory.cpp
+++ b/Engine/Plugins/AudioInsights/Source/AudioInsights/Private/AudioInsightsDashboardFactory.cpp
@@ -4,6 +4,7 @@
 #include "AudioDevice.h"
 #include "AudioDeviceManager.h"
 #include "AudioInsightsModule.h"
+#include "AudioInsightsStyle.h"
 #include "Editor.h"
 #include "Engine/World.h"
 #include "Framework/Docking/TabManager.h"
@@ -24,7 +25,7 @@ namespace UE::Audio::Insights
 	{
 		static const FText ToolName = LOCTEXT("AudioDashboard_ToolName", "Audio Insights");
 
-		static const FLazyName MainToolbarName = "MainToolbar";
+		static const FName MainToolbarName = "MainToolbar";
 		static const FText MainToolbarDisplayName = LOCTEXT("AudioDashboard_MainToolbarDisplayName", "Dashboard Transport");
 
 		static const FText PreviewDeviceDisplayName = LOCTEXT("AudioDashboard_PreviewDevice", "[Preview Audio]");
@@ -67,6 +68,7 @@ namespace UE::Audio::Insights
 			{
 				const FTraceModule& TraceModule = FAudioInsightsModule::GetChecked().GetTraceModule();
 				TraceModule.StartTraceAnalysis();
+
 				ActiveDeviceId = InDeviceId;
 			}
 		}
@@ -94,7 +96,6 @@ namespace UE::Audio::Insights
 		{
 			const FTraceModule& TraceModule = FAudioInsightsModule::GetChecked().GetTraceModule();
 			TraceModule.StopTraceAnalysis();
-
 		}
 
 		RefreshDeviceSelector();
@@ -199,35 +200,166 @@ namespace UE::Audio::Insights
 
 	TSharedRef<SDockTab> FDashboardFactory::MakeDockTabWidget(const FSpawnTabArgs& Args)
 	{
-		UnregisterTabSpawners();
-
-		InitDelegates();
-
-		InitTabLayout();
-
 		const TSharedRef<SDockTab> DockTab = SNew(SDockTab)
 			.Label(DashboardFactoryPrivate::ToolName)
 			.Clipping(EWidgetClipping::ClipToBounds)
 			.TabRole(ETabRole::NomadTab);
 
-		TSharedPtr<SWindow> Window = Args.GetOwnerWindow();
-
 		DashboardTabManager = FGlobalTabmanager::Get()->NewTabManager(DockTab);
 
+		InitDelegates();
+		TabLayout = GetDefaultTabLayout();
+
 		RegisterTabSpawners();
 		RefreshDeviceSelector();
 
-		TSharedPtr<SWidget> Content = DashboardTabManager->RestoreFrom(TabLayout->AsShared(), Window);
-		DockTab->SetContent(Content->AsShared());
+		const TSharedRef<SWidget> TabContent = SNew(SVerticalBox)
+			+ SVerticalBox::Slot()
+			.AutoHeight()
+			[
+				MakeMenuBarWidget()
+			]
+			+ SVerticalBox::Slot()
+			[
+				DashboardTabManager->RestoreFrom(TabLayout->AsShared(), Args.GetOwnerWindow()).ToSharedRef()
+			];
+
+		DockTab->SetContent(TabContent);
+
 		DockTab->SetOnTabClosed(SDockTab::FOnTabClosedCallback::CreateLambda([this](TSharedRef<SDockTab> TabClosed)
 		{
 			ResetDelegates();
+			UnregisterTabSpawners();
 		}));
+
 		return DockTab;
 	}
 
+	TSharedRef<SWidget> FDashboardFactory::MakeMenuBarWidget()
+	{
+		FMenuBarBuilder MenuBarBuilder = FMenuBarBuilder(TSharedPtr<FUICommandList>());
+
+		MenuBarBuilder.AddPullDownMenu(
+			LOCTEXT("File_MenuLabel", "File"),
+			FText::GetEmpty(),
+			FNewMenuDelegate::CreateLambda([this](FMenuBuilder& MenuBuilder)
+			{
+				MenuBuilder.AddMenuEntry(LOCTEXT("Close_MenuLabel", "Close"),
+					LOCTEXT("Close_MenuLabel_Tooltip", "Closes the Audio Insights dashboard."),
+					FSlateIcon(),
+					FUIAction(FExecuteAction::CreateLambda([this]()
+					{
+						if (DashboardTabManager.IsValid())
+						{
+							if (TSharedPtr<SDockTab> OwnerTab = DashboardTabManager->GetOwnerTab())
+							{
+								OwnerTab->RequestCloseTab();
+							}
+						}
+					}))
+				);
+			}),
+			"File"
+		);
+
+		MenuBarBuilder.AddPullDownMenu(
+			LOCTEXT("ViewMenuLabel", "View"),
+			FText::GetEmpty(),
+			FNewMenuDelegate::CreateLambda([this](FMenuBuilder& MenuBuilder)
+			{
+				MenuBuilder.AddMenuEntry(LOCTEXT("Transport_MenuLabel", "Transport"),
+				LOCTEXT("Transport_MenuLabel_Tooltip", "Shows the Transport tab."),
+					FSlateStyle::Get().CreateIcon("AudioInsights.Icon.Start.Inactive"),
+					FUIAction(FExecuteAction::CreateLambda([&DashboardTabManager = DashboardTabManager]()
+					{
+						if (DashboardTabManager.IsValid())
+						{
+							TSharedPtr<SDockTab> TransportTab = DashboardTabManager->FindExistingLiveTab(DashboardFactoryPrivate::MainToolbarName);
+							if (!TransportTab.IsValid())
+							{
+								TSharedPtr<SDockTab> InvokedTransportTab = DashboardTabManager->TryInvokeTab(DashboardFactoryPrivate::MainToolbarName);
+								if (InvokedTransportTab.IsValid())
+								{
+									InvokedTransportTab->SetParentDockTabStackTabWellHidden(true);
+								}
+							}
+							else
+							{
+								TransportTab->RequestCloseTab();
+							}
+						}
+					}),
+					FCanExecuteAction(),
+					FIsActionChecked::CreateLambda([&DashboardTabManager = DashboardTabManager]()
+					{
+						return DashboardTabManager.IsValid() ? DashboardTabManager->FindExistingLiveTab(DashboardFactoryPrivate::MainToolbarName).IsValid() : false;
+					})),
+					NAME_None,
+					EUserInterfaceActionType::Check
+				);
+
+				MenuBuilder.AddMenuSeparator();
+
+				for (const auto& KVP : DashboardViewFactories)
+				{
+					const FName& FactoryName = KVP.Key;
+					const TSharedPtr<IDashboardViewFactory>& Factory = KVP.Value;
+
+					MenuBuilder.AddMenuEntry(Factory->GetDisplayName(),
+						FText::GetEmpty(),
+						FSlateStyle::Get().CreateIcon(Factory->GetIcon().GetStyleName()),
+						FUIAction(FExecuteAction::CreateLambda([this, FactoryName]()
+						{
+							if (DashboardTabManager.IsValid())
+							{
+								if (TSharedPtr<SDockTab> ViewportTab = DashboardTabManager->FindExistingLiveTab(FactoryName);
+									!ViewportTab.IsValid())
+								{
+									DashboardTabManager->TryInvokeTab(FactoryName);
+
+									if (TSharedPtr<SDockTab> InvokedOutputMeterTab = DashboardTabManager->TryInvokeTab(FactoryName);
+										InvokedOutputMeterTab.IsValid() && DashboardViewFactories[FactoryName].IsValid())
+									{
+										if (const EDefaultDashboardTabStack DefaultTabStack = DashboardViewFactories[FactoryName]->GetDefaultTabStack();
+											DefaultTabStack == EDefaultDashboardTabStack::AudioMeter ||
+											DefaultTabStack == EDefaultDashboardTabStack::Oscilloscope)
+										{
+											InvokedOutputMeterTab->SetParentDockTabStackTabWellHidden(true);
+										}
+									}
+								}
+								else
+								{
+									ViewportTab->RequestCloseTab();
+								}
+							}
+						}),
+						FCanExecuteAction(),
+						FIsActionChecked::CreateLambda([&DashboardTabManager = DashboardTabManager, FactoryName]()
+						{
+							return DashboardTabManager.IsValid() ? DashboardTabManager->FindExistingLiveTab(FactoryName).IsValid() : false;
+						})),
+						NAME_None,
+						EUserInterfaceActionType::Check
+					);
+
+					if (const EDefaultDashboardTabStack DefaultTabStack = DashboardViewFactories[FactoryName]->GetDefaultTabStack();
+						DefaultTabStack == EDefaultDashboardTabStack::Log || DefaultTabStack == EDefaultDashboardTabStack::AudioMeters)
+					{
+						MenuBuilder.AddMenuSeparator();
+					}
+				}
+			}),
+			"View"
+		);
+
+		return MenuBarBuilder.MakeWidget();
+	}
+
 	TSharedRef<SWidget> FDashboardFactory::MakeMainToolbarWidget()
 	{
+		using namespace DashboardFactoryPrivate;
+
 		static const FName PlayWorldToolBarName = "Kismet.DebuggingViewToolBar";
 		if (!UToolMenus::Get()->IsMenuRegistered(PlayWorldToolBarName))
 		{
@@ -248,7 +380,6 @@ namespace UE::Audio::Insights
 				UToolMenus::Get()->GenerateWidget(PlayWorldToolBarName, { FPlayWorldCommands::GlobalPlayWorldActions })
 			]
 		]
-
 		+ SHorizontalBox::Slot()
 		.HAlign(HAlign_Right)
 		.VAlign(VAlign_Center)
@@ -296,7 +427,7 @@ namespace UE::Audio::Insights
 		[
 			SNew(STextBlock)
 			.Text(LOCTEXT("SelectDashboardWorld_DisplayName", "World Filter:"))
-			.ToolTipText(DashboardFactoryPrivate::DashboardWorldSelectDescription)
+			.ToolTipText(DashboardWorldSelectDescription)
 			.Font(IPropertyTypeCustomizationUtils::GetRegularFont())
 		]
 		+ SHorizontalBox::Slot()
@@ -306,11 +437,11 @@ namespace UE::Audio::Insights
 		.Padding(2.0f, 0.0f)
 		[
 			SAssignNew(AudioDeviceComboBox, SComboBox<TSharedPtr<::Audio::FDeviceId>>)
-			.ToolTipText(DashboardFactoryPrivate::DashboardWorldSelectDescription)
+			.ToolTipText(DashboardWorldSelectDescription)
 			.OptionsSource(&AudioDeviceIds)
 			.OnGenerateWidget_Lambda([](const TSharedPtr<::Audio::FDeviceId>& WidgetDeviceId)
 			{
-				FText NameText = DashboardFactoryPrivate::GetDebugNameFromDeviceId(*WidgetDeviceId);
+				FText NameText = GetDebugNameFromDeviceId(*WidgetDeviceId);
 				return SNew(STextBlock)
 					.Text(NameText)
 					.Font(IPropertyTypeCustomizationUtils::GetRegularFont());
@@ -330,7 +461,7 @@ namespace UE::Audio::Insights
 				.Font(IPropertyTypeCustomizationUtils::GetRegularFont())
 				.Text_Lambda([this]()
 				{
-					return DashboardFactoryPrivate::GetDebugNameFromDeviceId(ActiveDeviceId);
+					return GetDebugNameFromDeviceId(ActiveDeviceId);
 				})
 			]
 		];
@@ -369,12 +500,10 @@ namespace UE::Audio::Insights
 		}
 	}
 
-	void FDashboardFactory::InitTabLayout()
+	TSharedPtr<FTabManager::FLayout> FDashboardFactory::GetDefaultTabLayout()
 	{
 		using namespace DashboardFactoryPrivate;
 
-		TabLayout.Reset();
-
 		TSharedRef<FTabManager::FStack> MainMenuTabStack = FTabManager::NewStack();
 		TSharedRef<FTabManager::FStack> ViewportTabStack = FTabManager::NewStack();
 		TSharedRef<FTabManager::FStack> LogTabStack = FTabManager::NewStack();
@@ -385,44 +514,47 @@ namespace UE::Audio::Insights
 
 		MainMenuTabStack->AddTab(MainToolbarName, ETabState::OpenedTab);
 
-		for (const TPair<FName, TSharedPtr<IDashboardViewFactory>>& Factory : DashboardViewFactories)
+		for (const auto& KVP : DashboardViewFactories)
 		{
-			EDefaultDashboardTabStack DefaultTabStack = Factory.Value->GetDefaultTabStack();
+			const FName& FactoryName = KVP.Key;
+			const TSharedPtr<IDashboardViewFactory>& Factory = KVP.Value;
+
+			const EDefaultDashboardTabStack DefaultTabStack = Factory->GetDefaultTabStack();
 			switch (DefaultTabStack)
 			{
 				case EDefaultDashboardTabStack::Viewport:
 				{
-					ViewportTabStack->AddTab(Factory.Key, ETabState::OpenedTab);
+					ViewportTabStack->AddTab(FactoryName, ETabState::OpenedTab);
 				}
 				break;
 
 				case EDefaultDashboardTabStack::Log:
 				{
-					LogTabStack->AddTab(Factory.Key, ETabState::OpenedTab);
+					LogTabStack->AddTab(FactoryName, ETabState::OpenedTab);
 				}
 				break;
 
 				case EDefaultDashboardTabStack::Analysis:
 				{
-					AnalysisTabStack->AddTab(Factory.Key, ETabState::OpenedTab);
+					AnalysisTabStack->AddTab(FactoryName, ETabState::OpenedTab);
 				}
 				break;
 
 				case EDefaultDashboardTabStack::AudioMeters:
 				{
-					AudioMetersTabStack->AddTab(Factory.Key, ETabState::OpenedTab);
+					AudioMetersTabStack->AddTab(FactoryName, ETabState::OpenedTab);
 				}
 				break;
 
 				case EDefaultDashboardTabStack::AudioMeter:
 				{
-					AudioMeterTabStack->AddTab(Factory.Key, ETabState::OpenedTab);
+					AudioMeterTabStack->AddTab(FactoryName, ETabState::OpenedTab);
 				}
 				break;
 
 				case EDefaultDashboardTabStack::Oscilloscope:
 				{
-					OscilloscopeTabStack->AddTab(Factory.Key, ETabState::OpenedTab);
+					OscilloscopeTabStack->AddTab(FactoryName, ETabState::OpenedTab);
 				}
 				break;
 
@@ -433,7 +565,7 @@ namespace UE::Audio::Insights
 
 		AnalysisTabStack->SetForegroundTab(FName("MixerSources"));
 
-		TabLayout = FTabManager::NewLayout("AudioDashboard_Layout_v1")
+		return FTabManager::NewLayout("AudioDashboard_Layout_v1")
 		->AddArea
 		(
 			FTabManager::NewPrimaryArea()
@@ -441,7 +573,7 @@ namespace UE::Audio::Insights
 			->Split
 			(
 				MainMenuTabStack
-				->SetSizeCoefficient(0.065)
+				->SetSizeCoefficient(0.085f)
 				->SetHideTabWell(true)
 			)
 			->Split
@@ -519,23 +651,27 @@ namespace UE::Audio::Insights
 
 	void FDashboardFactory::RegisterTabSpawners()
 	{
-		DashboardWorkspace = DashboardTabManager->AddLocalWorkspaceMenuCategory(DashboardFactoryPrivate::ToolName);
-		DashboardTabManager->RegisterTabSpawner(DashboardFactoryPrivate::MainToolbarName, FOnSpawnTab::CreateLambda([this](const FSpawnTabArgs& Args)
+		using namespace DashboardFactoryPrivate;
+
+		DashboardWorkspace = DashboardTabManager->AddLocalWorkspaceMenuCategory(ToolName);
+		DashboardTabManager->RegisterTabSpawner(MainToolbarName, FOnSpawnTab::CreateLambda([this](const FSpawnTabArgs& Args)
 		{
 			return SNew(SDockTab)
 				.Clipping(EWidgetClipping::ClipToBounds)
-				.Label(DashboardFactoryPrivate::MainToolbarDisplayName)
+				.Label(MainToolbarDisplayName)
 				[
 					MakeMainToolbarWidget()
 				];
 		}))
-		.SetDisplayName(DashboardFactoryPrivate::MainToolbarDisplayName)
+		.SetDisplayName(MainToolbarDisplayName)
 		.SetGroup(DashboardWorkspace->AsShared());
 
-		for (const TPair<FName, TSharedPtr<IDashboardViewFactory>>& KVP : DashboardViewFactories)
+		for (const auto& KVP : DashboardViewFactories)
 		{
-			const FName FactoryName = KVP.Value->GetName();
-			DashboardTabManager->RegisterTabSpawner(FactoryName, FOnSpawnTab::CreateLambda([this, Factory = KVP.Value](const FSpawnTabArgs& Args)
+			const FName& FactoryName = KVP.Key;
+			const TSharedPtr<IDashboardViewFactory>& Factory = KVP.Value;
+
+			DashboardTabManager->RegisterTabSpawner(FactoryName, FOnSpawnTab::CreateLambda([this, Factory](const FSpawnTabArgs& Args)
 			{
 				TSharedPtr<SWidget> DashboardView = Factory->MakeWidget();
 				return SNew(SDockTab)
@@ -545,35 +681,29 @@ namespace UE::Audio::Insights
 						DashboardView->AsShared()
 					];
 			}))
-			.SetDisplayName(KVP.Value->GetDisplayName())
+			.SetDisplayName(Factory->GetDisplayName())
 			.SetGroup(DashboardWorkspace->AsShared())
-			.SetIcon(KVP.Value->GetIcon());
+			.SetIcon(Factory->GetIcon());
 		}
 	}
 
 	void FDashboardFactory::RegisterViewFactory(TSharedRef<IDashboardViewFactory> InFactory)
 	{
-		const FName Name = InFactory->GetName();
-		if (ensureAlwaysMsgf(!DashboardViewFactories.Contains(Name), TEXT("Failed to register Audio Insights Dashboard '%s': Dashboard with name already registered"), *Name.ToString()))
+		if (const FName Name = InFactory->GetName(); 
+			ensureAlwaysMsgf(!DashboardViewFactories.Contains(Name), TEXT("Failed to register Audio Insights Dashboard '%s': Dashboard with name already registered"), *Name.ToString()))
 		{
 			DashboardViewFactories.Add(Name, InFactory);
 		}
 	}
 
-	void FDashboardFactory::UnregisterViewFactory(FName InName)
-	{
-		DashboardViewFactories.Remove(InName);
-	}
-
 	void FDashboardFactory::UnregisterTabSpawners()
 	{
-		using namespace DashboardFactoryPrivate;
-
 		if (DashboardTabManager.IsValid())
 		{
-			for (const TPair<FName, TSharedPtr<IDashboardViewFactory>>& Factory : DashboardViewFactories)
+			for (const auto& KVP : DashboardViewFactories)
 			{
-				DashboardTabManager->UnregisterTabSpawner(Factory.Value->GetName());
+				const FName& FactoryName = KVP.Key;
+				DashboardTabManager->UnregisterTabSpawner(FactoryName);
 			}
 
 			DashboardTabManager.Reset();
@@ -581,5 +711,10 @@ namespace UE::Audio::Insights
 
 		DashboardWorkspace.Reset();
 	}
+
+	void FDashboardFactory::UnregisterViewFactory(FName InName)
+	{
+		DashboardViewFactories.Remove(InName);
+	}
 } // namespace UE::Audio::Insights
 #undef LOCTEXT_NAMESPACE
diff --git a/Engine/Plugins/AudioInsights/Source/AudioInsights/Private/AudioInsightsDashboardFactory.h b/Engine/Plugins/AudioInsights/Source/AudioInsights/Private/AudioInsightsDashboardFactory.h
--- a/Engine/Plugins/AudioInsights/Source/AudioInsights/Private/AudioInsightsDashboardFactory.h
+++ b/Engine/Plugins/AudioInsights/Source/AudioInsights/Private/AudioInsightsDashboardFactory.h
@@ -42,10 +42,11 @@ namespace UE::Audio::Insights
 		void RefreshDeviceSelector();
 		void ResetDelegates();
 
+		TSharedRef<SWidget> MakeMenuBarWidget();
 		TSharedRef<SWidget> MakeMainToolbarWidget();
 
 		void InitDelegates();
-		void InitTabLayout();
+		TSharedPtr<FTabManager::FLayout> GetDefaultTabLayout();
 
 		void RegisterTabSpawners();
 		void UnregisterTabSpawners();
