[AudioInsights] Converted Transport tab to an element of the Audio insights main window, this fixes the squeezing issue + makes code logic simpler and UI more consistent.

diff --git a/Engine/Plugins/AudioInsights/Source/AudioInsights/Private/AudioInsightsDashboardFactory.cpp b/Engine/Plugins/AudioInsights/Source/AudioInsights/Private/AudioInsightsDashboardFactory.cpp
--- a/Engine/Plugins/AudioInsights/Source/AudioInsights/Private/AudioInsightsDashboardFactory.cpp
+++ b/Engine/Plugins/AudioInsights/Source/AudioInsights/Private/AudioInsightsDashboardFactory.cpp
@@ -220,6 +220,17 @@ namespace UE::Audio::Insights
 				MakeMenuBarWidget()
 			]
 			+ SVerticalBox::Slot()
+			.AutoHeight()
+			[
+				MakeMainToolbarWidget()
+			]
+			+ SVerticalBox::Slot()
+			.AutoHeight()
+			[
+				SNew(SBox)
+				.HeightOverride(4.0f)
+			]
+			+ SVerticalBox::Slot()
 			[
 				DashboardTabManager->RestoreFrom(TabLayout->AsShared(), Args.GetOwnerWindow()).ToSharedRef()
 			];
@@ -267,39 +278,6 @@ namespace UE::Audio::Insights
 			FText::GetEmpty(),
 			FNewMenuDelegate::CreateLambda([this](FMenuBuilder& MenuBuilder)
 			{
-				MenuBuilder.AddMenuEntry(LOCTEXT("Transport_MenuLabel", "Transport"),
-				LOCTEXT("Transport_MenuLabel_Tooltip", "Shows the Transport tab."),
-					FSlateStyle::Get().CreateIcon("AudioInsights.Icon.Start.Inactive"),
-					FUIAction(FExecuteAction::CreateLambda([&DashboardTabManager = DashboardTabManager]()
-					{
-						if (DashboardTabManager.IsValid())
-						{
-							TSharedPtr<SDockTab> TransportTab = DashboardTabManager->FindExistingLiveTab(DashboardFactoryPrivate::MainToolbarName);
-							if (!TransportTab.IsValid())
-							{
-								TSharedPtr<SDockTab> InvokedTransportTab = DashboardTabManager->TryInvokeTab(DashboardFactoryPrivate::MainToolbarName);
-								if (InvokedTransportTab.IsValid())
-								{
-									InvokedTransportTab->SetParentDockTabStackTabWellHidden(true);
-								}
-							}
-							else
-							{
-								TransportTab->RequestCloseTab();
-							}
-						}
-					}),
-					FCanExecuteAction(),
-					FIsActionChecked::CreateLambda([&DashboardTabManager = DashboardTabManager]()
-					{
-						return DashboardTabManager.IsValid() ? DashboardTabManager->FindExistingLiveTab(DashboardFactoryPrivate::MainToolbarName).IsValid() : false;
-					})),
-					NAME_None,
-					EUserInterfaceActionType::Check
-				);
-
-				MenuBuilder.AddMenuSeparator();
-
 				for (const auto& KVP : DashboardViewFactories)
 				{
 					const FName& FactoryName = KVP.Key;
@@ -368,103 +346,111 @@ namespace UE::Audio::Insights
 			FPlayWorldCommands::BuildToolbar(Section);
 		}
 
-		return SNew(SHorizontalBox)
-		+ SHorizontalBox::Slot()
-		.HAlign(HAlign_Left)
-		.VAlign(VAlign_Center)
-		.AutoWidth()
-		[
-			SNew(SBorder)
-			.BorderImage(FAppStyle::Get().GetBrush("NoBorder"))
-			[
-				UToolMenus::Get()->GenerateWidget(PlayWorldToolBarName, { FPlayWorldCommands::GlobalPlayWorldActions })
-			]
-		]
-		+ SHorizontalBox::Slot()
-		.HAlign(HAlign_Right)
-		.VAlign(VAlign_Center)
-		.AutoWidth()
-		.Padding(2.0f, 0.0f)
-		[
-			SNew(STextBlock)
-			.Text(LOCTEXT("StartOnPIE_DisplayName", "Start with PIE:"))
-			.Font(IPropertyTypeCustomizationUtils::GetRegularFont())
-		]
-		+ SHorizontalBox::Slot()
-		.HAlign(HAlign_Left)
-		.VAlign(VAlign_Center)
-		.AutoWidth()
-		.Padding(2.0f, 0.0f)
-		[
-			SNew(SCheckBox)
-			.IsChecked_Lambda([this]() { return bStartWithPIE ? ECheckBoxState::Checked : ECheckBoxState::Unchecked; })
-			.OnCheckStateChanged_Lambda([this](ECheckBoxState NewState) { bStartWithPIE = NewState == ECheckBoxState::Checked; })
-		]
-		+ SHorizontalBox::Slot()
-		.HAlign(HAlign_Right)
-		.VAlign(VAlign_Center)
-		.AutoWidth()
-		.Padding(2.0f, 0.0f)
-		[
-			SNew(STextBlock)
-			.Text(LOCTEXT("StopOnPIE_DisplayName", "Stop with PIE:"))
-			.Font(IPropertyTypeCustomizationUtils::GetRegularFont())
-		]
-		+ SHorizontalBox::Slot()
-		.HAlign(HAlign_Left)
-		.VAlign(VAlign_Center)
-		.AutoWidth()
-		.Padding(2.0f, 0.0f)
-		[
-			SNew(SCheckBox)
-			.IsChecked_Lambda([this]() { return bStopWithPIE ? ECheckBoxState::Checked : ECheckBoxState::Unchecked; })
-			.OnCheckStateChanged_Lambda([this](ECheckBoxState NewState) { bStopWithPIE = NewState == ECheckBoxState::Checked; })
-		]
-		+ SHorizontalBox::Slot()
-		.HAlign(HAlign_Right)
-		.VAlign(VAlign_Center)
-		.Padding(2.0f, 0.0f)
-		[
-			SNew(STextBlock)
-			.Text(LOCTEXT("SelectDashboardWorld_DisplayName", "World Filter:"))
-			.ToolTipText(DashboardWorldSelectDescription)
-			.Font(IPropertyTypeCustomizationUtils::GetRegularFont())
-		]
-		+ SHorizontalBox::Slot()
-		.HAlign(HAlign_Left)
-		.VAlign(VAlign_Center)
-		.AutoWidth()
-		.Padding(2.0f, 0.0f)
-		[
-			SAssignNew(AudioDeviceComboBox, SComboBox<TSharedPtr<::Audio::FDeviceId>>)
-			.ToolTipText(DashboardWorldSelectDescription)
-			.OptionsSource(&AudioDeviceIds)
-			.OnGenerateWidget_Lambda([](const TSharedPtr<::Audio::FDeviceId>& WidgetDeviceId)
-			{
-				FText NameText = GetDebugNameFromDeviceId(*WidgetDeviceId);
-				return SNew(STextBlock)
-					.Text(NameText)
-					.Font(IPropertyTypeCustomizationUtils::GetRegularFont());
-			})
-			.OnSelectionChanged_Lambda([this](TSharedPtr<::Audio::FDeviceId> NewDeviceId, ESelectInfo::Type)
-			{
-				if (NewDeviceId.IsValid())
-				{
-					ActiveDeviceId = *NewDeviceId;
-					RefreshDeviceSelector();
+		static FSlateBrush TransportBackgroundColorBrush;
+		TransportBackgroundColorBrush.TintColor = FSlateColor(FLinearColor(0.018f, 0.018f, 0.018f, 1.0f));
+		TransportBackgroundColorBrush.DrawAs    = ESlateBrushDrawType::Box;
 
-					OnActiveAudioDeviceChanged.Broadcast();
-				}
-			})
+		return SNew(SBorder)
+			.BorderImage(&TransportBackgroundColorBrush)
 			[
-				SNew(STextBlock)
-				.Font(IPropertyTypeCustomizationUtils::GetRegularFont())
-				.Text_Lambda([this]()
-				{
-					return GetDebugNameFromDeviceId(ActiveDeviceId);
-				})
-			]
-		];
+				SNew(SHorizontalBox)
+				+ SHorizontalBox::Slot()
+				.HAlign(HAlign_Left)
+				.VAlign(VAlign_Center)
+				.AutoWidth()
+				[
+					SNew(SBorder)
+					.BorderImage(FAppStyle::Get().GetBrush("NoBorder"))
+					[
+						UToolMenus::Get()->GenerateWidget(PlayWorldToolBarName, { FPlayWorldCommands::GlobalPlayWorldActions })
+					]
+				]
+				+ SHorizontalBox::Slot()
+				.HAlign(HAlign_Right)
+				.VAlign(VAlign_Center)
+				.AutoWidth()
+				.Padding(2.0f, 0.0f)
+				[
+					SNew(STextBlock)
+					.Text(LOCTEXT("StartOnPIE_DisplayName", "Start with PIE:"))
+					.Font(IPropertyTypeCustomizationUtils::GetRegularFont())
+				]
+				+ SHorizontalBox::Slot()
+				.HAlign(HAlign_Left)
+				.VAlign(VAlign_Center)
+				.AutoWidth()
+				.Padding(2.0f, 0.0f)
+				[
+					SNew(SCheckBox)
+					.IsChecked_Lambda([this]() { return bStartWithPIE ? ECheckBoxState::Checked : ECheckBoxState::Unchecked; })
+					.OnCheckStateChanged_Lambda([this](ECheckBoxState NewState) { bStartWithPIE = NewState == ECheckBoxState::Checked; })
+				]
+				+ SHorizontalBox::Slot()
+				.HAlign(HAlign_Right)
+				.VAlign(VAlign_Center)
+				.AutoWidth()
+				.Padding(2.0f, 0.0f)
+				[
+					SNew(STextBlock)
+					.Text(LOCTEXT("StopOnPIE_DisplayName", "Stop with PIE:"))
+					.Font(IPropertyTypeCustomizationUtils::GetRegularFont())
+				]
+				+ SHorizontalBox::Slot()
+				.HAlign(HAlign_Left)
+				.VAlign(VAlign_Center)
+				.AutoWidth()
+				.Padding(2.0f, 0.0f)
+				[
+					SNew(SCheckBox)
+					.IsChecked_Lambda([this]() { return bStopWithPIE ? ECheckBoxState::Checked : ECheckBoxState::Unchecked; })
+					.OnCheckStateChanged_Lambda([this](ECheckBoxState NewState) { bStopWithPIE = NewState == ECheckBoxState::Checked; })
+				]
+				+ SHorizontalBox::Slot()
+				.HAlign(HAlign_Right)
+				.VAlign(VAlign_Center)
+				.Padding(2.0f, 0.0f)
+				[
+					SNew(STextBlock)
+					.Text(LOCTEXT("SelectDashboardWorld_DisplayName", "World Filter:"))
+					.ToolTipText(DashboardWorldSelectDescription)
+					.Font(IPropertyTypeCustomizationUtils::GetRegularFont())
+				]
+				+ SHorizontalBox::Slot()
+				.HAlign(HAlign_Left)
+				.VAlign(VAlign_Center)
+				.AutoWidth()
+				.Padding(2.0f, 0.0f)
+				[
+					SAssignNew(AudioDeviceComboBox, SComboBox<TSharedPtr<::Audio::FDeviceId>>)
+					.ToolTipText(DashboardWorldSelectDescription)
+					.OptionsSource(&AudioDeviceIds)
+					.OnGenerateWidget_Lambda([](const TSharedPtr<::Audio::FDeviceId>& WidgetDeviceId)
+					{
+						FText NameText = GetDebugNameFromDeviceId(*WidgetDeviceId);
+						return SNew(STextBlock)
+							.Text(NameText)
+							.Font(IPropertyTypeCustomizationUtils::GetRegularFont());
+					})
+					.OnSelectionChanged_Lambda([this](TSharedPtr<::Audio::FDeviceId> NewDeviceId, ESelectInfo::Type)
+					{
+						if (NewDeviceId.IsValid())
+						{
+							ActiveDeviceId = *NewDeviceId;
+							RefreshDeviceSelector();
+
+							OnActiveAudioDeviceChanged.Broadcast();
+						}
+					})
+					[
+						SNew(STextBlock)
+						.Font(IPropertyTypeCustomizationUtils::GetRegularFont())
+						.Text_Lambda([this]()
+						{
+							return GetDebugNameFromDeviceId(ActiveDeviceId);
+						})
+					]
+				]
+			];
 	}
 
 	void FDashboardFactory::InitDelegates()
@@ -504,7 +490,6 @@ namespace UE::Audio::Insights
 	{
 		using namespace DashboardFactoryPrivate;
 
-		TSharedRef<FTabManager::FStack> MainMenuTabStack = FTabManager::NewStack();
 		TSharedRef<FTabManager::FStack> ViewportTabStack = FTabManager::NewStack();
 		TSharedRef<FTabManager::FStack> LogTabStack = FTabManager::NewStack();
 		TSharedRef<FTabManager::FStack> AnalysisTabStack = FTabManager::NewStack();
@@ -512,8 +497,6 @@ namespace UE::Audio::Insights
 		TSharedRef<FTabManager::FStack> AudioMeterTabStack = FTabManager::NewStack();
 		TSharedRef<FTabManager::FStack> OscilloscopeTabStack = FTabManager::NewStack();
 
-		MainMenuTabStack->AddTab(MainToolbarName, ETabState::OpenedTab);
-
 		for (const auto& KVP : DashboardViewFactories)
 		{
 			const FName& FactoryName = KVP.Key;
@@ -571,12 +554,6 @@ namespace UE::Audio::Insights
 			FTabManager::NewPrimaryArea()
 			->SetOrientation(Orient_Vertical)
 			->Split
-			(
-				MainMenuTabStack
-				->SetSizeCoefficient(0.085f)
-				->SetHideTabWell(true)
-			)
-			->Split
 			(
 				// Left column
 				FTabManager::NewSplitter()
@@ -654,17 +631,6 @@ namespace UE::Audio::Insights
 		using namespace DashboardFactoryPrivate;
 
 		DashboardWorkspace = DashboardTabManager->AddLocalWorkspaceMenuCategory(ToolName);
-		DashboardTabManager->RegisterTabSpawner(MainToolbarName, FOnSpawnTab::CreateLambda([this](const FSpawnTabArgs& Args)
-		{
-			return SNew(SDockTab)
-				.Clipping(EWidgetClipping::ClipToBounds)
-				.Label(MainToolbarDisplayName)
-				[
-					MakeMainToolbarWidget()
-				];
-		}))
-		.SetDisplayName(MainToolbarDisplayName)
-		.SetGroup(DashboardWorkspace->AsShared());
 
 		for (const auto& KVP : DashboardViewFactories)
 		{
