[AudioInsights] Cleanup plots when stopping PIE

diff --git a/Engine/Plugins/AudioInsights/Source/AudioInsights/Private/Views/MixerSourceDashboardViewFactory.cpp b/Engine/Plugins/AudioInsights/Source/AudioInsights/Private/Views/MixerSourceDashboardViewFactory.cpp
--- a/Engine/Plugins/AudioInsights/Source/AudioInsights/Private/Views/MixerSourceDashboardViewFactory.cpp
+++ b/Engine/Plugins/AudioInsights/Source/AudioInsights/Private/Views/MixerSourceDashboardViewFactory.cpp
@@ -6,6 +6,7 @@
 #include "AudioInsightsDashboardAssetCommands.h"
 #include "AudioInsightsModule.h"
 #include "AudioInsightsStyle.h"
+#include "Editor.h"
 #include "Internationalization/Text.h"
 #include "IPropertyTypeCustomization.h"
 #include "Providers/MixerSourceTraceProvider.h"
@@ -334,6 +335,25 @@ namespace UE::Audio::Insights
 		}
 	}
 
+	void FMixerSourceDashboardViewFactory::ResetPlots()
+	{
+		for (const auto& KVP : PlotWidgetCurveIdToPointDataMapPerColumn)
+		{
+			const TSharedPtr<FPointDataPerCurveMap>& PointDataPerCurveMap = KVP.Value;
+			PointDataPerCurveMap->Empty();
+		}
+
+		PlotWidgetMetadataPerCurve->Empty();
+
+		BeginTimestamp = TNumericLimits<double>::Max();
+		CurrentTimestamp = 0;
+	}
+
+	void FMixerSourceDashboardViewFactory::OnPIEStopped(bool bSimulating)
+	{
+		ResetPlots();
+	}
+
 	TSharedPtr<SWidget> FMixerSourceDashboardViewFactory::OnConstructContextMenu()
 	{
 		const FDashboardAssetCommands& Commands = FDashboardAssetCommands::Get();
@@ -922,6 +942,7 @@ namespace UE::Audio::Insights
 	TSharedRef<SWidget> FMixerSourceDashboardViewFactory::MakeWidget()
 	{
 		FDashboardFactory::OnActiveAudioDeviceChanged.AddSP(this, &FMixerSourceDashboardViewFactory::ClearMutesAndSolos);
+		FEditorDelegates::EndPIE.AddSP(this, &FMixerSourceDashboardViewFactory::OnPIEStopped);
 
 		TSharedRef<SWidget> MuteSoloWidget = MakeMuteSoloWidget();
 		TSharedRef<SWidget> TableDashboardWidget = FTraceTableDashboardViewFactory::MakeWidget();
diff --git a/Engine/Plugins/AudioInsights/Source/AudioInsights/Private/Views/MixerSourceDashboardViewFactory.h b/Engine/Plugins/AudioInsights/Source/AudioInsights/Private/Views/MixerSourceDashboardViewFactory.h
--- a/Engine/Plugins/AudioInsights/Source/AudioInsights/Private/Views/MixerSourceDashboardViewFactory.h
+++ b/Engine/Plugins/AudioInsights/Source/AudioInsights/Private/Views/MixerSourceDashboardViewFactory.h
@@ -42,6 +42,8 @@ namespace UE::Audio::Insights
 		TSharedRef<SWidget> MakePlotsWidget();
 
 	private:
+		void ResetPlots();
+		void OnPIEStopped(bool bSimulating);
 		void UpdatePlotsWidgetsData();
 		void UpdateSoloMuteState();
 
