[AudioInsights] Show items in VirtualLoops table regardless of the used concurrency rule.

diff --git a/Engine/Plugins/AudioInsights/Source/AudioInsights/Private/Providers/VirtualLoopTraceProvider.cpp b/Engine/Plugins/AudioInsights/Source/AudioInsights/Private/Providers/VirtualLoopTraceProvider.cpp
--- a/Engine/Plugins/AudioInsights/Source/AudioInsights/Private/Providers/VirtualLoopTraceProvider.cpp
+++ b/Engine/Plugins/AudioInsights/Source/AudioInsights/Private/Providers/VirtualLoopTraceProvider.cpp
@@ -15,6 +15,21 @@ namespace UE::Audio::Insights
 
 	bool FVirtualLoopTraceProvider::ProcessMessages()
 	{
+		auto RemoveEntryFunc = [this](const FVirtualLoopRealizeMessage& Msg, TSharedPtr<FVirtualLoopDashboardEntry>* OutEntry)
+		{
+			if (OutEntry && (*OutEntry)->Timestamp < Msg.Timestamp)
+			{
+				RemoveDeviceEntry(Msg.DeviceId, Msg.PlayOrder);
+			}
+		};
+
+		auto GetEntryFunc = [this](const FVirtualLoopMessageBase& Msg)
+		{
+			return FindDeviceEntry(Msg.DeviceId, Msg.PlayOrder);
+		};
+
+		ProcessMessageQueue<FVirtualLoopStopMessage>(TraceMessages.StopMessages, GetEntryFunc, RemoveEntryFunc);
+
 		ProcessMessageQueue<FVirtualLoopVirtualizeMessage>(TraceMessages.VirtualizeMessages,
 		[this](const FVirtualLoopMessageBase& Msg)
 		{
@@ -40,12 +55,29 @@ namespace UE::Audio::Insights
 			EntryRef.ComponentId = Msg.ComponentId;
 		});
 
-		auto GetEntryFunc = [this](const FVirtualLoopMessageBase& Msg)
+		ProcessMessageQueue<FVirtualLoopRealizeMessage>(TraceMessages.RealizeMessages, 
+		[this](const FVirtualLoopMessageBase& Msg)
 		{
-			return FindDeviceEntry(Msg.DeviceId, Msg.PlayOrder);
-		};
+			TSharedPtr<FVirtualLoopDashboardEntry>* ToReturn = nullptr;
+
+			UpdateDeviceEntry(Msg.DeviceId, Msg.PlayOrder, [&ToReturn, &Msg](TSharedPtr<FVirtualLoopDashboardEntry>& Entry)
+			{
+				if (!Entry.IsValid())
+				{
+					Entry = MakeShared<FVirtualLoopDashboardEntry>();
+					Entry->DeviceId = Msg.DeviceId;
+					Entry->PlayOrder = Msg.PlayOrder;
+				}
+				Entry->Timestamp = Msg.Timestamp;
+				ToReturn = &Entry;
+			});
+
+			return ToReturn;
+		}, 
+		RemoveEntryFunc);
 
-		ProcessMessageQueue<FVirtualLoopUpdateMessage>(TraceMessages.UpdateMessages, GetEntryFunc,
+		ProcessMessageQueue<FVirtualLoopUpdateMessage>(TraceMessages.UpdateMessages, 
+		GetEntryFunc,
 		[](const FVirtualLoopUpdateMessage& Msg, TSharedPtr<FVirtualLoopDashboardEntry>* OutEntry)
 		{
 			if (OutEntry)
@@ -59,17 +91,6 @@ namespace UE::Audio::Insights
 			}
 		});
 
-		auto RemoveEntryFunc = [this](const FVirtualLoopMessageBase& Msg, TSharedPtr<FVirtualLoopDashboardEntry>* OutEntry)
-		{
-			if (OutEntry && (*OutEntry)->Timestamp < Msg.Timestamp)
-			{
-				RemoveDeviceEntry(Msg.DeviceId, Msg.PlayOrder);
-			}
-		};
-
-		ProcessMessageQueue<FVirtualLoopRealizeMessage>(TraceMessages.RealizeMessages, GetEntryFunc, RemoveEntryFunc);
-		ProcessMessageQueue<FVirtualLoopStopMessage>(TraceMessages.StopMessages, GetEntryFunc, RemoveEntryFunc);
-
 		return true;
 	}
 
