Audio Insights: Stop the trace module spam and remove the delay before items are shown in the sources tab

diff --git a/Engine/Plugins/AudioInsights/Source/AudioInsights/Private/AudioInsightsTraceModule.cpp b/Engine/Plugins/AudioInsights/Source/AudioInsights/Private/AudioInsightsTraceModule.cpp
--- a/Engine/Plugins/AudioInsights/Source/AudioInsights/Private/AudioInsightsTraceModule.cpp
+++ b/Engine/Plugins/AudioInsights/Source/AudioInsights/Private/AudioInsightsTraceModule.cpp
@@ -6,6 +6,7 @@
 #include "Providers/MixerSourceTraceProvider.h"
 #include "Providers/VirtualLoopTraceProvider.h"
 #include "Templates/SharedPointer.h"
+#include "Trace/Trace.h"
 
 
 namespace UE::Audio::Insights
@@ -37,6 +38,25 @@ namespace UE::Audio::Insights
 		return TraceName.Resolve();
 	}
 
+	void FTraceModule::DisableAllTraceChannels()
+	{
+		UE::Trace::EnumerateChannels([](const ANSICHAR* ChannelName, bool bEnabled, void*)
+			{
+				if (bEnabled)
+				{
+					FString ChannelNameFString(ChannelName);
+					UE::Trace::ToggleChannel(ChannelNameFString.GetCharArray().GetData(), false);
+				}
+			}
+		, nullptr);
+	}
+
+	void FTraceModule::EnableAudioInsightsTraceChannels()
+	{
+		UE::Trace::ToggleChannel(TEXT("Audio"), true);
+		UE::Trace::ToggleChannel(TEXT("AudioMixer"), true);
+	}
+
 	void FTraceModule::OnAnalysisBegin(TraceServices::IAnalysisSession& InSession)
 	{
 		for (const TPair<FName, TSharedPtr<FTraceProviderBase>>& Pair : TraceProviders)
@@ -50,7 +70,14 @@ namespace UE::Audio::Insights
 	{
 		if (!FTraceAuxiliary::IsConnected())
 		{
-			FTraceAuxiliary::Start(FTraceAuxiliary::EConnectionType::Network, TEXT("localhost"), nullptr);
+			DisableAllTraceChannels();
+			EnableAudioInsightsTraceChannels();
+
+			// Clear all buffered data and prevent data from previous recordings from leaking into the new recording
+			FTraceAuxiliary::FOptions Options;
+			Options.bExcludeTail = true;
+
+			FTraceAuxiliary::Start(FTraceAuxiliary::EConnectionType::Network, TEXT("localhost"), TEXT(""), &Options);
 
 			IUnrealInsightsModule& UnrealInsightsModule = FModuleManager::LoadModuleChecked<IUnrealInsightsModule>("TraceInsights");
 			UnrealInsightsModule.StartAnalysisForLastLiveSession();
diff --git a/Engine/Plugins/AudioInsights/Source/AudioInsights/Private/AudioInsightsTraceModule.h b/Engine/Plugins/AudioInsights/Source/AudioInsights/Private/AudioInsightsTraceModule.h
--- a/Engine/Plugins/AudioInsights/Source/AudioInsights/Private/AudioInsightsTraceModule.h
+++ b/Engine/Plugins/AudioInsights/Source/AudioInsights/Private/AudioInsightsTraceModule.h
@@ -40,6 +40,9 @@ namespace UE::Audio::Insights
 	private:
 		static const FName GetName();
 
+		static void DisableAllTraceChannels();
+		static void EnableAudioInsightsTraceChannels();
+
 		TSharedRef<FTraceChannelManager> ChannelManager;
 		TMap<FName, TSharedPtr<FTraceProviderBase>> TraceProviders;
 	};
