Removing audio insights source-view right-click context menu for 5.4 -- will revisit in 5.5 with a better design for this menu.

diff --git a/Engine/Plugins/AudioInsights/Source/AudioInsights/Private/Views/MixerSourceDashboardViewFactory.cpp b/Engine/Plugins/AudioInsights/Source/AudioInsights/Private/Views/MixerSourceDashboardViewFactory.cpp
--- a/Engine/Plugins/AudioInsights/Source/AudioInsights/Private/Views/MixerSourceDashboardViewFactory.cpp
+++ b/Engine/Plugins/AudioInsights/Source/AudioInsights/Private/Views/MixerSourceDashboardViewFactory.cpp
@@ -354,6 +354,7 @@ namespace UE::Audio::Insights
 		ResetPlots();
 	}
 
+#if AUDIO_INSIGHTS_SHOW_SOURCE_CONTEXT_MENU
 	TSharedPtr<SWidget> FMixerSourceDashboardViewFactory::OnConstructContextMenu()
 	{
 		const FDashboardAssetCommands& Commands = FDashboardAssetCommands::Get();
@@ -376,6 +377,7 @@ namespace UE::Audio::Insights
 
 		return MenuBuilder.MakeWidget();
 	}
+#endif // AUDIO_INSIGHTS_SHOW_SOURCE_CONTEXT_MENU 
 
 	FSlateColor FMixerSourceDashboardViewFactory::GetRowColor(const TSharedPtr<IDashboardDataViewEntry>& InRowDataPtr)
 	{
diff --git a/Engine/Plugins/AudioInsights/Source/AudioInsights/Private/Views/MixerSourceDashboardViewFactory.h b/Engine/Plugins/AudioInsights/Source/AudioInsights/Private/Views/MixerSourceDashboardViewFactory.h
--- a/Engine/Plugins/AudioInsights/Source/AudioInsights/Private/Views/MixerSourceDashboardViewFactory.h
+++ b/Engine/Plugins/AudioInsights/Source/AudioInsights/Private/Views/MixerSourceDashboardViewFactory.h
@@ -5,6 +5,10 @@
 #include "Views/TableDashboardViewFactory.h"
 #include "Widgets/Input/SCheckBox.h"
 
+
+// Removing temporarily due to redesign of solo and mute functionality
+#define AUDIO_INSIGHTS_SHOW_SOURCE_CONTEXT_MENU 0
+
 namespace UE::Audio::Insights
 {
 	class FMixerSourceDashboardViewFactory : public FTraceObjectTableDashboardViewFactory
@@ -53,7 +57,9 @@ namespace UE::Audio::Insights
 		const FNumberFormattingOptions* GetPlotColumnNumberFormat(const FName& ColumnName);
 		const FText GetPlotColumnDisplayName(const FName& ColumnName);
 
+#if AUDIO_INSIGHTS_SHOW_SOURCE_CONTEXT_MENU
 		virtual TSharedPtr<SWidget> OnConstructContextMenu() override;
+#endif
 		virtual FSlateColor GetRowColor(const TSharedPtr<IDashboardDataViewEntry>& InRowDataPtr) override;
 
 		TSharedRef<SWidget> MakeMuteSoloWidget();
