Don't run audio insights providers in cook commandlets where they aren't necessary
#rnx

diff --git a/Engine/Plugins/AudioInsights/Source/AudioInsights/Private/AudioInsightsModule.cpp b/Engine/Plugins/AudioInsights/Source/AudioInsights/Private/AudioInsightsModule.cpp
--- a/Engine/Plugins/AudioInsights/Source/AudioInsights/Private/AudioInsightsModule.cpp
+++ b/Engine/Plugins/AudioInsights/Source/AudioInsights/Private/AudioInsightsModule.cpp
@@ -32,30 +32,37 @@ namespace UE::Audio::Insights
 {
 	void FAudioInsightsModule::StartupModule()
 	{
-		IModularFeatures::Get().RegisterModularFeature(TraceServices::ModuleFeatureName, &TraceModule);
-
-		RegisterMenus();
-
-		DashboardFactory = MakeShared<FDashboardFactory>();
-		DashboardFactory->RegisterViewFactory(MakeShared<FViewportDashboardViewFactory>());
-		DashboardFactory->RegisterViewFactory(MakeShared<FLogDashboardViewFactory>());
-		DashboardFactory->RegisterViewFactory(MakeShared<FMixerSourceDashboardViewFactory>());
-		DashboardFactory->RegisterViewFactory(MakeShared<FVirtualLoopDashboardViewFactory>());
-		DashboardFactory->RegisterViewFactory(MakeShared<FSubmixesDashboardViewFactory>());
-		DashboardFactory->RegisterViewFactory(MakeShared<FAudioBusesDashboardViewFactory>());
-		DashboardFactory->RegisterViewFactory(MakeShared<FAudioMetersDashboardViewFactory>());
-		DashboardFactory->RegisterViewFactory(MakeShared<FOutputMeterDashboardViewFactory>());
-		DashboardFactory->RegisterViewFactory(MakeShared<FOutputOscilloscopeDashboardViewFactory>());
-
-		FDashboardAssetCommands::Register();
+		// Don't run providers in cook commandlet to avoid additional, unnecessary overhead as audio insights is dormant.
+		if (!IsRunningCookCommandlet())
+		{
+			IModularFeatures::Get().RegisterModularFeature(TraceServices::ModuleFeatureName, &TraceModule);
+
+			RegisterMenus();
+
+			DashboardFactory = MakeShared<FDashboardFactory>();
+			DashboardFactory->RegisterViewFactory(MakeShared<FViewportDashboardViewFactory>());
+			DashboardFactory->RegisterViewFactory(MakeShared<FLogDashboardViewFactory>());
+			DashboardFactory->RegisterViewFactory(MakeShared<FMixerSourceDashboardViewFactory>());
+			DashboardFactory->RegisterViewFactory(MakeShared<FVirtualLoopDashboardViewFactory>());
+			DashboardFactory->RegisterViewFactory(MakeShared<FSubmixesDashboardViewFactory>());
+			DashboardFactory->RegisterViewFactory(MakeShared<FAudioBusesDashboardViewFactory>());
+			DashboardFactory->RegisterViewFactory(MakeShared<FAudioMetersDashboardViewFactory>());
+			DashboardFactory->RegisterViewFactory(MakeShared<FOutputMeterDashboardViewFactory>());
+			DashboardFactory->RegisterViewFactory(MakeShared<FOutputOscilloscopeDashboardViewFactory>());
+
+			FDashboardAssetCommands::Register();
+		}
 	}
 
 	void FAudioInsightsModule::ShutdownModule()
 	{
-		DashboardFactory.Reset();
-		IModularFeatures::Get().UnregisterModularFeature(TraceServices::ModuleFeatureName, &TraceModule);
+		if (!IsRunningCookCommandlet())
+		{
+			DashboardFactory.Reset();
+			IModularFeatures::Get().UnregisterModularFeature(TraceServices::ModuleFeatureName, &TraceModule);
 
-		FDashboardAssetCommands::Unregister();
+			FDashboardAssetCommands::Unregister();
+		}
 	}
 
 	void FAudioInsightsModule::RegisterDashboardViewFactory(TSharedRef<IDashboardViewFactory> InDashboardFactory)
diff --git a/Engine/Plugins/AudioInsights/Source/AudioInsights/Private/AudioInsightsTraceModule.cpp b/Engine/Plugins/AudioInsights/Source/AudioInsights/Private/AudioInsightsTraceModule.cpp
--- a/Engine/Plugins/AudioInsights/Source/AudioInsights/Private/AudioInsightsTraceModule.cpp
+++ b/Engine/Plugins/AudioInsights/Source/AudioInsights/Private/AudioInsightsTraceModule.cpp
@@ -1,6 +1,7 @@
 // Copyright Epic Games, Inc. All Rights Reserved.
 #include "AudioInsightsTraceModule.h"
 
+#include "CoreGlobals.h"
 #include "Insights/IUnrealInsightsModule.h"
 #include "ProfilingDebugging/TraceAuxiliary.h"
 #include "Providers/MixerSourceTraceProvider.h"
@@ -14,11 +15,15 @@ namespace UE::Audio::Insights
 	FTraceModule::FTraceModule()
 		: ChannelManager(MakeShared<FTraceChannelManager>())
 	{
-		TSharedPtr<FMixerSourceTraceProvider> SourceProvider = MakeShared<FMixerSourceTraceProvider>(ChannelManager);
-		TSharedPtr<FVirtualLoopTraceProvider> VirtualLoopProvider = MakeShared<FVirtualLoopTraceProvider>(ChannelManager);
+		// Don't run providers in cook commandlet to avoid additional, unnecessary overhead as audio insights is dormant.
+		if (!IsRunningCookCommandlet())
+		{
+			TSharedPtr<FMixerSourceTraceProvider> SourceProvider = MakeShared<FMixerSourceTraceProvider>(ChannelManager);
+			TSharedPtr<FVirtualLoopTraceProvider> VirtualLoopProvider = MakeShared<FVirtualLoopTraceProvider>(ChannelManager);
 
-		TraceProviders.Add(SourceProvider->GetName(), StaticCastSharedPtr<FTraceProviderBase>(SourceProvider));
-		TraceProviders.Add(VirtualLoopProvider->GetName(), StaticCastSharedPtr<FTraceProviderBase>(VirtualLoopProvider));
+			TraceProviders.Add(SourceProvider->GetName(), StaticCastSharedPtr<FTraceProviderBase>(SourceProvider));
+			TraceProviders.Add(VirtualLoopProvider->GetName(), StaticCastSharedPtr<FTraceProviderBase>(VirtualLoopProvider));
+		}
 	}
 
 	void FTraceModule::GetModuleInfo(TraceServices::FModuleInfo& OutModuleInfo)
diff --git a/Engine/Plugins/AudioInsights/Source/AudioInsights/Private/Providers/SoundSubmixProvider.cpp b/Engine/Plugins/AudioInsights/Source/AudioInsights/Private/Providers/SoundSubmixProvider.cpp
--- a/Engine/Plugins/AudioInsights/Source/AudioInsights/Private/Providers/SoundSubmixProvider.cpp
+++ b/Engine/Plugins/AudioInsights/Source/AudioInsights/Private/Providers/SoundSubmixProvider.cpp
@@ -142,8 +142,7 @@ namespace UE::Audio::Insights
 		const IAudioInsightsModule& InsightsModule = FModuleManager::GetModuleChecked<IAudioInsightsModule>(IAudioInsightsModule::GetName());
 		const ::Audio::FDeviceId DeviceId = InsightsModule.GetDeviceId();
 
-		for (uint32 Index = 0;
-			 const TSharedPtr<FSoundSubmixAssetDashboardEntry>& SubmixDataViewEntry : SubmixDataViewEntries)
+		for (const TSharedPtr<FSoundSubmixAssetDashboardEntry>& SubmixDataViewEntry : SubmixDataViewEntries)
 		{
 			UpdateDeviceEntry(DeviceId, SubmixDataViewEntry->SoundSubmix->GetUniqueID(), [&SubmixDataViewEntry](TSharedPtr<FSoundSubmixAssetDashboardEntry>& Entry)
 			{
@@ -152,8 +151,6 @@ namespace UE::Audio::Insights
 					Entry = SubmixDataViewEntry;
 				}
 			});
-
-			++Index;
 		}
 
 		return true;
