[Audio Insights] Mixer Source Tab: Improved plot data resolution and plot scrolling + misc. code optimizations to handle better hi-freq data + reduced num plots from 32 to 16 + removed 2nd plot graph.

diff --git a/Engine/Plugins/AudioInsights/Source/AudioInsights/Private/AudioInsightsTraceModule.cpp b/Engine/Plugins/AudioInsights/Source/AudioInsights/Private/AudioInsightsTraceModule.cpp
--- a/Engine/Plugins/AudioInsights/Source/AudioInsights/Private/AudioInsightsTraceModule.cpp
+++ b/Engine/Plugins/AudioInsights/Source/AudioInsights/Private/AudioInsightsTraceModule.cpp
@@ -69,6 +69,8 @@ namespace UE::Audio::Insights
 			InSession.AddProvider(Pair.Key, nullptr, Pair.Value);
 			InSession.AddAnalyzer(Pair.Value->ConstructAnalyzer());
 		}
+
+		FirstTimeStamp = -TNumericLimits<double>::Min();
 	}
 
 	void FTraceModule::StartTraceAnalysis() const
diff --git a/Engine/Plugins/AudioInsights/Source/AudioInsights/Private/AudioInsightsTraceModule.h b/Engine/Plugins/AudioInsights/Source/AudioInsights/Private/AudioInsightsTraceModule.h
--- a/Engine/Plugins/AudioInsights/Source/AudioInsights/Private/AudioInsightsTraceModule.h
+++ b/Engine/Plugins/AudioInsights/Source/AudioInsights/Private/AudioInsightsTraceModule.h
@@ -37,6 +37,9 @@ namespace UE::Audio::Insights
 		bool IsTraceAnalysisActive() const;
 		void StopTraceAnalysis() const;
 
+		void SetFirstTimeStamp(double InFirstTimeStamp) { FirstTimeStamp = InFirstTimeStamp; }
+		double GetFirstTimeStamp() const { return FirstTimeStamp; }
+
 	private:
 		static const FName GetName();
 
@@ -45,5 +48,7 @@ namespace UE::Audio::Insights
 
 		TSharedRef<FTraceChannelManager> ChannelManager;
 		TMap<FName, TSharedPtr<FTraceProviderBase>> TraceProviders;
+
+		double FirstTimeStamp = -TNumericLimits<double>::Min();
 	};
 } // namespace UE::Audio::Insights
diff --git a/Engine/Plugins/AudioInsights/Source/AudioInsights/Private/Messages/MixerSourceTraceMessages.h b/Engine/Plugins/AudioInsights/Source/AudioInsights/Private/Messages/MixerSourceTraceMessages.h
--- a/Engine/Plugins/AudioInsights/Source/AudioInsights/Private/Messages/MixerSourceTraceMessages.h
+++ b/Engine/Plugins/AudioInsights/Source/AudioInsights/Private/Messages/MixerSourceTraceMessages.h
@@ -2,6 +2,7 @@
 #pragma once
 
 #include "AudioDefines.h"
+#include "DSP/Dsp.h"
 #include "Messages/AnalyzerMessageQueue.h"
 #include "Trace/Analyzer.h"
 #include "Views/TableDashboardViewFactory.h"
@@ -66,20 +67,33 @@ namespace UE::Audio::Insights
 	DEFINE_MIXERSOURCE_PARAM_MESSAGE(FMixerSourceVolumeMessage, Volume, float, 1.0f)
 #undef DEFINE_MIXERSOURCE_PARAM_MESSAGE
 
+	using FDataPoint = TPair<double, float>; // (Timestamp, Value)
+
 	class FMixerSourceDashboardEntry : public FSoundAssetDashboardEntry
 	{
 	public:
-		FMixerSourceDashboardEntry() = default;
+		FMixerSourceDashboardEntry()
+		{
+			constexpr uint32 DataPointsCapacity = 256;
+
+			VolumeDataPoints.SetCapacity(DataPointsCapacity);
+			PitchDataPoints.SetCapacity(DataPointsCapacity);
+			LPFFreqDataPoints.SetCapacity(DataPointsCapacity);
+			HPFFreqDataPoints.SetCapacity(DataPointsCapacity);
+			EnvelopeDataPoints.SetCapacity(DataPointsCapacity);
+			DistanceAttenuationDataPoints.SetCapacity(DataPointsCapacity);
+		}
+
 		virtual ~FMixerSourceDashboardEntry() = default;
 
 		int32 SourceId = INDEX_NONE;
 
-		float Volume = 1.0f;
-		float Pitch = 1.0f;
-		float LPFFreq = MIN_FILTER_FREQUENCY;
-		float HPFFreq = MAX_FILTER_FREQUENCY;
-		float Envelope = 0.0f;
-		float DistanceAttenuation = 0.0f;
+		::Audio::TCircularAudioBuffer<FDataPoint> VolumeDataPoints;
+		::Audio::TCircularAudioBuffer<FDataPoint> PitchDataPoints;
+		::Audio::TCircularAudioBuffer<FDataPoint> LPFFreqDataPoints;
+		::Audio::TCircularAudioBuffer<FDataPoint> HPFFreqDataPoints;
+		::Audio::TCircularAudioBuffer<FDataPoint> EnvelopeDataPoints;
+		::Audio::TCircularAudioBuffer<FDataPoint> DistanceAttenuationDataPoints;
 	};
 
 	class FMixerSourceMessages
diff --git a/Engine/Plugins/AudioInsights/Source/AudioInsights/Private/Providers/MixerSourceTraceProvider.cpp b/Engine/Plugins/AudioInsights/Source/AudioInsights/Private/Providers/MixerSourceTraceProvider.cpp
--- a/Engine/Plugins/AudioInsights/Source/AudioInsights/Private/Providers/MixerSourceTraceProvider.cpp
+++ b/Engine/Plugins/AudioInsights/Source/AudioInsights/Private/Providers/MixerSourceTraceProvider.cpp
@@ -1,6 +1,7 @@
 // Copyright Epic Games, Inc. All Rights Reserved.
 #include "Providers/MixerSourceTraceProvider.h"
 
+#include "AudioInsightsModule.h"
 #include "Trace/Analyzer.h"
 #include "TraceServices/Model/AnalysisSession.h"
 #include "TraceServices/ModuleService.h"
@@ -42,43 +43,83 @@ namespace UE::Audio::Insights
 			EntryRef.SourceId = Msg.SourceId;
 		});
 
+		TArray<int32, TInlineAllocator<64>> EntriesWithPoppedDataPoints;
+
 		ProcessMessageQueue<FMixerSourceVolumeMessage>(TraceMessages.VolumeMessages, BumpEntryFunc,
-		[](const FMixerSourceVolumeMessage& Msg, TSharedPtr<FMixerSourceDashboardEntry>* OutEntry)
+		[&EntriesWithPoppedDataPoints](const FMixerSourceVolumeMessage& Msg, TSharedPtr<FMixerSourceDashboardEntry>* OutEntry)
 		{
-			(*OutEntry)->Volume = Msg.Volume;
+			if (!EntriesWithPoppedDataPoints.Contains(Msg.PlayOrder))
+			{
+				(*OutEntry)->VolumeDataPoints.Pop((*OutEntry)->VolumeDataPoints.Num());
+				EntriesWithPoppedDataPoints.Add(Msg.PlayOrder);
+			}
+
+			(*OutEntry)->VolumeDataPoints.Push({ Msg.Timestamp, Msg.Volume });
 		});
 
+		EntriesWithPoppedDataPoints.Reset();
 		ProcessMessageQueue<FMixerSourcePitchMessage>(TraceMessages.PitchMessages, BumpEntryFunc,
-		[](const FMixerSourcePitchMessage& Msg, TSharedPtr<FMixerSourceDashboardEntry>* OutEntry)
+		[&EntriesWithPoppedDataPoints](const FMixerSourcePitchMessage& Msg, TSharedPtr<FMixerSourceDashboardEntry>* OutEntry)
 		{
-			(*OutEntry)->Pitch = Msg.Pitch;
+			if (!EntriesWithPoppedDataPoints.Contains(Msg.PlayOrder))
+			{
+				(*OutEntry)->PitchDataPoints.Pop((*OutEntry)->PitchDataPoints.Num());
+				EntriesWithPoppedDataPoints.Add(Msg.PlayOrder);
+			}
+
+			(*OutEntry)->PitchDataPoints.Push({ Msg.Timestamp, Msg.Pitch });
 		});
 
+		EntriesWithPoppedDataPoints.Reset();
 		ProcessMessageQueue<FMixerSourceLPFFreqMessage>(TraceMessages.LPFFreqMessages, BumpEntryFunc,
-		[](const FMixerSourceLPFFreqMessage& Msg, TSharedPtr<FMixerSourceDashboardEntry>* OutEntry)
+		[&EntriesWithPoppedDataPoints](const FMixerSourceLPFFreqMessage& Msg, TSharedPtr<FMixerSourceDashboardEntry>* OutEntry)
 		{
-			(*OutEntry)->LPFFreq = Msg.LPFFrequency;
+			if (!EntriesWithPoppedDataPoints.Contains(Msg.PlayOrder))
+			{
+				(*OutEntry)->LPFFreqDataPoints.Pop((*OutEntry)->LPFFreqDataPoints.Num());
+				EntriesWithPoppedDataPoints.Add(Msg.PlayOrder);
+			}
+
+			(*OutEntry)->LPFFreqDataPoints.Push({ Msg.Timestamp, Msg.LPFFrequency });
 		});
 
+		EntriesWithPoppedDataPoints.Reset();
 		ProcessMessageQueue<FMixerSourceHPFFreqMessage>(TraceMessages.HPFFreqMessages, BumpEntryFunc,
-		[](const FMixerSourceHPFFreqMessage& Msg, TSharedPtr<FMixerSourceDashboardEntry>* OutEntry)
+		[&EntriesWithPoppedDataPoints](const FMixerSourceHPFFreqMessage& Msg, TSharedPtr<FMixerSourceDashboardEntry>* OutEntry)
 		{
-			(*OutEntry)->HPFFreq = Msg.HPFFrequency;
+			if (!EntriesWithPoppedDataPoints.Contains(Msg.PlayOrder))
+			{
+				(*OutEntry)->HPFFreqDataPoints.Pop((*OutEntry)->HPFFreqDataPoints.Num());
+				EntriesWithPoppedDataPoints.Add(Msg.PlayOrder);
+			}
+
+			(*OutEntry)->HPFFreqDataPoints.Push({ Msg.Timestamp, Msg.HPFFrequency });
 		});
 
-		MaxEnvsMap.Reset();
+		EntriesWithPoppedDataPoints.Reset();
 		ProcessMessageQueue<FMixerSourceEnvelopeMessage>(TraceMessages.EnvelopeMessages, BumpEntryFunc,
-		[this](const FMixerSourceEnvelopeMessage& Msg, TSharedPtr<FMixerSourceDashboardEntry>* OutEntry)
+		[&EntriesWithPoppedDataPoints](const FMixerSourceEnvelopeMessage& Msg, TSharedPtr<FMixerSourceDashboardEntry>* OutEntry)
 		{
-			float& MaxEnv = MaxEnvsMap.FindOrAdd(Msg.PlayOrder);
-			MaxEnv = FMath::Max(MaxEnv, Msg.Envelope);
-			(*OutEntry)->Envelope = MaxEnv;
+			if (!EntriesWithPoppedDataPoints.Contains(Msg.PlayOrder))
+			{
+				(*OutEntry)->EnvelopeDataPoints.Pop((*OutEntry)->EnvelopeDataPoints.Num());
+				EntriesWithPoppedDataPoints.Add(Msg.PlayOrder);
+			}
+
+			(*OutEntry)->EnvelopeDataPoints.Push({ Msg.Timestamp, Msg.Envelope });
 		});
 
+		EntriesWithPoppedDataPoints.Reset();
 		ProcessMessageQueue<FMixerSourceDistanceAttenuationMessage>(TraceMessages.DistanceAttenuationMessages, BumpEntryFunc,
-		[](const FMixerSourceDistanceAttenuationMessage& Msg, TSharedPtr<FMixerSourceDashboardEntry>* OutEntry)
+		[&EntriesWithPoppedDataPoints](const FMixerSourceDistanceAttenuationMessage& Msg, TSharedPtr<FMixerSourceDashboardEntry>* OutEntry)
 		{
-			(*OutEntry)->DistanceAttenuation = Msg.DistanceAttenuation;
+			if (!EntriesWithPoppedDataPoints.Contains(Msg.PlayOrder))
+			{
+				(*OutEntry)->DistanceAttenuationDataPoints.Pop((*OutEntry)->DistanceAttenuationDataPoints.Num());
+				EntriesWithPoppedDataPoints.Add(Msg.PlayOrder);
+			}
+
+			(*OutEntry)->DistanceAttenuationDataPoints.Push({ Msg.Timestamp, Msg.DistanceAttenuation });
 		});
 
 		auto GetEntry = [this](const FMixerSourceMessageBase& Msg)
@@ -124,6 +165,14 @@ namespace UE::Audio::Insights
 			{
 				LLM_SCOPE_BYNAME(TEXT("Insights/FMixerSourceTraceAnalyzer"));
 
+				FTraceModule& TraceModule = FAudioInsightsModule::GetChecked().GetTraceModule();
+
+				const double CurrentTime = FPlatformTime::Seconds();
+				if (TraceModule.GetFirstTimeStamp() < 0.0)
+				{
+					TraceModule.SetFirstTimeStamp(CurrentTime);
+				}
+
 				FMixerSourceMessages& Messages = GetProvider<FMixerSourceTraceProvider>().TraceMessages;
 				switch (RouteId)
 				{
diff --git a/Engine/Plugins/AudioInsights/Source/AudioInsights/Private/Providers/MixerSourceTraceProvider.h b/Engine/Plugins/AudioInsights/Source/AudioInsights/Private/Providers/MixerSourceTraceProvider.h
--- a/Engine/Plugins/AudioInsights/Source/AudioInsights/Private/Providers/MixerSourceTraceProvider.h
+++ b/Engine/Plugins/AudioInsights/Source/AudioInsights/Private/Providers/MixerSourceTraceProvider.h
@@ -28,8 +28,6 @@ namespace UE::Audio::Insights
 		static FName GetName_Static();
 
 	private:
-		TMap<int32, float> MaxEnvsMap;
-
 		FMixerSourceMessages TraceMessages;
 
 		TSet<FTraceChannelHandle> Channels;
diff --git a/Engine/Plugins/AudioInsights/Source/AudioInsights/Private/Views/MixerSourceDashboardViewFactory.cpp b/Engine/Plugins/AudioInsights/Source/AudioInsights/Private/Views/MixerSourceDashboardViewFactory.cpp
--- a/Engine/Plugins/AudioInsights/Source/AudioInsights/Private/Views/MixerSourceDashboardViewFactory.cpp
+++ b/Engine/Plugins/AudioInsights/Source/AudioInsights/Private/Views/MixerSourceDashboardViewFactory.cpp
@@ -6,6 +6,7 @@
 #include "AudioInsightsDashboardAssetCommands.h"
 #include "AudioInsightsModule.h"
 #include "AudioInsightsStyle.h"
+#include "DSP/Dsp.h"
 #include "Editor.h"
 #include "Internationalization/Text.h"
 #include "IPropertyTypeCustomization.h"
@@ -27,11 +28,24 @@ namespace UE::Audio::Insights
 			return static_cast<const FMixerSourceDashboardEntry&>(InData);
 		};
 
+		float GetLastEntryArrayValue(const ::Audio::TCircularAudioBuffer<FDataPoint>& EnvelopeDataPoints)
+		{
+			float LastValue = 0.0f;
+
+			if (EnvelopeDataPoints.Num() > 0)
+			{
+				const ::Audio::DisjointedArrayView<const FDataPoint> EnvelopeDataPointsDisjointedArrayView = EnvelopeDataPoints.PeekInPlace(EnvelopeDataPoints.Num());
+				LastValue = EnvelopeDataPointsDisjointedArrayView.FirstBuffer.Last().Value;
+			}
+
+			return LastValue;
+		};
+
 		static const FText PlotColumnSelectDescription = LOCTEXT("AudioDashboard_MixerSources_SelectPlotColumnDescription", "Select a column from the table to plot.");
 	} // namespace MixerSourcePrivate
 
 	const double FMixerSourceDashboardViewFactory::MaxPlotHistorySeconds = 5.0;
-	const int32 FMixerSourceDashboardViewFactory::MaxPlotSources = 32;
+	const int32 FMixerSourceDashboardViewFactory::MaxPlotSources = 16;
 
 	FMixerSourceDashboardViewFactory::FMixerSourceDashboardViewFactory()
 	{
@@ -66,7 +80,6 @@ namespace UE::Audio::Insights
 	{
 		auto CreateColumnData = []()
 		{
-
 			return TMap<FName, FTraceTableDashboardViewFactory::FColumnData>
 			{
 				{
@@ -91,7 +104,11 @@ namespace UE::Audio::Insights
 					"Amplitude",
 					{
 						LOCTEXT("Source_EnvColumnDisplayName", "Amp (Peak)"),
-						[](const IDashboardDataViewEntry& InData) { return FText::AsNumber(MixerSourcePrivate::CastEntry(InData).Envelope, FSlateStyle::Get().GetAmpFloatFormat()); },
+						[](const IDashboardDataViewEntry& InData)
+						{
+							const ::Audio::TCircularAudioBuffer<FDataPoint>& EnvelopeDataPoints = MixerSourcePrivate::CastEntry(InData).EnvelopeDataPoints;
+							return FText::AsNumber(MixerSourcePrivate::GetLastEntryArrayValue(EnvelopeDataPoints), FSlateStyle::Get().GetAmpFloatFormat());
+						},
 						false /* bDefaultHidden */,
 						0.12f /* FillWidth */
 					}
@@ -100,7 +117,11 @@ namespace UE::Audio::Insights
 					"Volume",
 					{
 						LOCTEXT("Source_VolumeColumnDisplayName", "Volume"),
-						[](const IDashboardDataViewEntry& InData) { return FText::AsNumber(MixerSourcePrivate::CastEntry(InData).Volume, FSlateStyle::Get().GetAmpFloatFormat()); },
+						[](const IDashboardDataViewEntry& InData)
+						{
+							const ::Audio::TCircularAudioBuffer<FDataPoint>& VolumeDataPoints = MixerSourcePrivate::CastEntry(InData).VolumeDataPoints;
+							return FText::AsNumber(MixerSourcePrivate::GetLastEntryArrayValue(VolumeDataPoints), FSlateStyle::Get().GetAmpFloatFormat());
+						},
 						false /* bDefaultHidden */,
 						0.07f /* FillWidth */
 					}
@@ -109,7 +130,11 @@ namespace UE::Audio::Insights
 					"DistanceAttenuation",
 					{
 						LOCTEXT("Source_AttenuationColumnDisplayName", "Distance Attenuation"),
-						[](const IDashboardDataViewEntry& InData) { return FText::AsNumber(MixerSourcePrivate::CastEntry(InData).DistanceAttenuation, FSlateStyle::Get().GetAmpFloatFormat()); },
+						[](const IDashboardDataViewEntry& InData)
+						{
+							const ::Audio::TCircularAudioBuffer<FDataPoint>& DistanceAttenuationDataPoints = MixerSourcePrivate::CastEntry(InData).DistanceAttenuationDataPoints;
+							return FText::AsNumber(MixerSourcePrivate::GetLastEntryArrayValue(DistanceAttenuationDataPoints), FSlateStyle::Get().GetAmpFloatFormat());
+						},
 						true  /* bDefaultHidden */,
 						0.15f /* FillWidth */
 					}
@@ -118,7 +143,11 @@ namespace UE::Audio::Insights
 					"Pitch",
 					{
 						LOCTEXT("Source_PitchColumnDisplayName", "Pitch"),
-						[](const IDashboardDataViewEntry& InData) { return FText::AsNumber(MixerSourcePrivate::CastEntry(InData).Pitch, FSlateStyle::Get().GetPitchFloatFormat()); },
+						[](const IDashboardDataViewEntry& InData)
+						{
+							const ::Audio::TCircularAudioBuffer<FDataPoint>& PitchDataPoints = MixerSourcePrivate::CastEntry(InData).PitchDataPoints;
+							return FText::AsNumber(MixerSourcePrivate::GetLastEntryArrayValue(PitchDataPoints), FSlateStyle::Get().GetPitchFloatFormat());
+						},
 						false /* bDefaultHidden */,
 						0.06f /* FillWidth */
 					}
@@ -127,7 +156,11 @@ namespace UE::Audio::Insights
 					"LPF",
 					{
 						LOCTEXT("Source_LPFColumnDisplayName", "LPF Freq (Hz)"),
-						[](const IDashboardDataViewEntry& InData) { return FText::AsNumber(MixerSourcePrivate::CastEntry(InData).LPFFreq, FSlateStyle::Get().GetFreqFloatFormat()); },
+						[](const IDashboardDataViewEntry& InData)
+						{
+							const ::Audio::TCircularAudioBuffer<FDataPoint>& LPFFreqDataPoints = MixerSourcePrivate::CastEntry(InData).LPFFreqDataPoints;
+							return FText::AsNumber(MixerSourcePrivate::GetLastEntryArrayValue(LPFFreqDataPoints), FSlateStyle::Get().GetFreqFloatFormat());
+						},
 						true  /* bDefaultHidden */,
 						0.1f /* FillWidth */
 					}
@@ -136,7 +169,11 @@ namespace UE::Audio::Insights
 					"HPF",
 					{
 						LOCTEXT("Source_HPFColumnDisplayName", "HPF Freq (Hz)"),
-						[](const IDashboardDataViewEntry& InData) { return FText::AsNumber(MixerSourcePrivate::CastEntry(InData).HPFFreq, FSlateStyle::Get().GetFreqFloatFormat()); },
+						[](const IDashboardDataViewEntry& InData)
+						{
+							const ::Audio::TCircularAudioBuffer<FDataPoint>& HPFFreqDataPoints = MixerSourcePrivate::CastEntry(InData).HPFFreqDataPoints;
+							return FText::AsNumber(MixerSourcePrivate::GetLastEntryArrayValue(HPFFreqDataPoints), FSlateStyle::Get().GetFreqFloatFormat());
+						},
 						true  /* bDefaultHidden */,
 						0.1f /* FillWidth */
 					}
@@ -204,7 +241,7 @@ namespace UE::Audio::Insights
 					const FMixerSourceDashboardEntry& AData = MixerSourcePrivate::CastEntry(*A.Get());
 					const FMixerSourceDashboardEntry& BData = MixerSourcePrivate::CastEntry(*B.Get());
 
-					return AData.Envelope < BData.Envelope;
+					return MixerSourcePrivate::GetLastEntryArrayValue(AData.EnvelopeDataPoints) < MixerSourcePrivate::GetLastEntryArrayValue(BData.EnvelopeDataPoints);
 				});
 			}
 			else if (SortMode == EColumnSortMode::Descending)
@@ -214,7 +251,7 @@ namespace UE::Audio::Insights
 					const FMixerSourceDashboardEntry& AData = MixerSourcePrivate::CastEntry(*A.Get());
 					const FMixerSourceDashboardEntry& BData = MixerSourcePrivate::CastEntry(*B.Get());
 
-					return BData.Envelope < AData.Envelope;
+					return MixerSourcePrivate::GetLastEntryArrayValue(BData.EnvelopeDataPoints) < MixerSourcePrivate::GetLastEntryArrayValue(AData.EnvelopeDataPoints);
 				});
 			}
 		}
@@ -227,7 +264,7 @@ namespace UE::Audio::Insights
 					const FMixerSourceDashboardEntry& AData = MixerSourcePrivate::CastEntry(*A.Get());
 					const FMixerSourceDashboardEntry& BData = MixerSourcePrivate::CastEntry(*B.Get());
 
-					return AData.Volume < BData.Volume;
+					return MixerSourcePrivate::GetLastEntryArrayValue(AData.VolumeDataPoints) < MixerSourcePrivate::GetLastEntryArrayValue(BData.VolumeDataPoints);
 				});
 			}
 			else if (SortMode == EColumnSortMode::Descending)
@@ -237,7 +274,7 @@ namespace UE::Audio::Insights
 					const FMixerSourceDashboardEntry& AData = MixerSourcePrivate::CastEntry(*A.Get());
 					const FMixerSourceDashboardEntry& BData = MixerSourcePrivate::CastEntry(*B.Get());
 
-					return BData.Volume < AData.Volume;
+					return MixerSourcePrivate::GetLastEntryArrayValue(BData.VolumeDataPoints) < MixerSourcePrivate::GetLastEntryArrayValue(AData.VolumeDataPoints);
 				});
 			}
 		}
@@ -250,7 +287,7 @@ namespace UE::Audio::Insights
 					const FMixerSourceDashboardEntry& AData = MixerSourcePrivate::CastEntry(*A.Get());
 					const FMixerSourceDashboardEntry& BData = MixerSourcePrivate::CastEntry(*B.Get());
 
-					return AData.DistanceAttenuation < BData.DistanceAttenuation;
+					return MixerSourcePrivate::GetLastEntryArrayValue(AData.DistanceAttenuationDataPoints) < MixerSourcePrivate::GetLastEntryArrayValue(BData.DistanceAttenuationDataPoints);
 				});
 			}
 			else if (SortMode == EColumnSortMode::Descending)
@@ -260,7 +297,7 @@ namespace UE::Audio::Insights
 					const FMixerSourceDashboardEntry& AData = MixerSourcePrivate::CastEntry(*A.Get());
 					const FMixerSourceDashboardEntry& BData = MixerSourcePrivate::CastEntry(*B.Get());
 
-					return BData.DistanceAttenuation < AData.DistanceAttenuation;
+					return MixerSourcePrivate::GetLastEntryArrayValue(BData.DistanceAttenuationDataPoints) < MixerSourcePrivate::GetLastEntryArrayValue(AData.DistanceAttenuationDataPoints);
 				});
 			}
 		}
@@ -273,7 +310,7 @@ namespace UE::Audio::Insights
 					const FMixerSourceDashboardEntry& AData = MixerSourcePrivate::CastEntry(*A.Get());
 					const FMixerSourceDashboardEntry& BData = MixerSourcePrivate::CastEntry(*B.Get());
 
-					return AData.Pitch < BData.Pitch;
+					return MixerSourcePrivate::GetLastEntryArrayValue(AData.PitchDataPoints) < MixerSourcePrivate::GetLastEntryArrayValue(BData.PitchDataPoints);
 				});
 			}
 			else if (SortMode == EColumnSortMode::Descending)
@@ -283,7 +320,7 @@ namespace UE::Audio::Insights
 					const FMixerSourceDashboardEntry& AData = MixerSourcePrivate::CastEntry(*A.Get());
 					const FMixerSourceDashboardEntry& BData = MixerSourcePrivate::CastEntry(*B.Get());
 
-					return BData.Pitch < AData.Pitch;
+					return MixerSourcePrivate::GetLastEntryArrayValue(BData.PitchDataPoints) < MixerSourcePrivate::GetLastEntryArrayValue(AData.PitchDataPoints);
 				});
 			}
 		}
@@ -296,7 +333,7 @@ namespace UE::Audio::Insights
 					const FMixerSourceDashboardEntry& AData = MixerSourcePrivate::CastEntry(*A.Get());
 					const FMixerSourceDashboardEntry& BData = MixerSourcePrivate::CastEntry(*B.Get());
 
-					return AData.LPFFreq < BData.LPFFreq;
+					return MixerSourcePrivate::GetLastEntryArrayValue(AData.LPFFreqDataPoints) < MixerSourcePrivate::GetLastEntryArrayValue(BData.LPFFreqDataPoints);
 				});
 			}
 			else if (SortMode == EColumnSortMode::Descending)
@@ -306,7 +343,7 @@ namespace UE::Audio::Insights
 					const FMixerSourceDashboardEntry& AData = MixerSourcePrivate::CastEntry(*A.Get());
 					const FMixerSourceDashboardEntry& BData = MixerSourcePrivate::CastEntry(*B.Get());
 
-					return BData.LPFFreq < AData.LPFFreq;
+					return MixerSourcePrivate::GetLastEntryArrayValue(BData.LPFFreqDataPoints) < MixerSourcePrivate::GetLastEntryArrayValue(AData.LPFFreqDataPoints);
 				});
 			}
 		}
@@ -319,7 +356,7 @@ namespace UE::Audio::Insights
 					const FMixerSourceDashboardEntry& AData = MixerSourcePrivate::CastEntry(*A.Get());
 					const FMixerSourceDashboardEntry& BData = MixerSourcePrivate::CastEntry(*B.Get());
 
-					return AData.HPFFreq < BData.HPFFreq;
+					return MixerSourcePrivate::GetLastEntryArrayValue(AData.HPFFreqDataPoints) < MixerSourcePrivate::GetLastEntryArrayValue(BData.HPFFreqDataPoints);
 				});
 			}
 			else if (SortMode == EColumnSortMode::Descending)
@@ -329,7 +366,7 @@ namespace UE::Audio::Insights
 					const FMixerSourceDashboardEntry& AData = MixerSourcePrivate::CastEntry(*A.Get());
 					const FMixerSourceDashboardEntry& BData = MixerSourcePrivate::CastEntry(*B.Get());
 
-					return BData.HPFFreq < AData.HPFFreq;
+					return MixerSourcePrivate::GetLastEntryArrayValue(BData.HPFFreqDataPoints) < MixerSourcePrivate::GetLastEntryArrayValue(AData.HPFFreqDataPoints);
 				});
 			}
 		}
@@ -349,9 +386,26 @@ namespace UE::Audio::Insights
 		CurrentTimestamp = 0;
 	}
 
+	void FMixerSourceDashboardViewFactory::OnPIEStarted(bool bSimulating)
+	{
+		PIEState = EPIEState::Running;
+	}
+
 	void FMixerSourceDashboardViewFactory::OnPIEStopped(bool bSimulating)
 	{
 		ResetPlots();
+
+		PIEState = EPIEState::Stopped;
+	}
+
+	void FMixerSourceDashboardViewFactory::OnPIEPaused(bool bSimulating)
+	{
+		PIEState = EPIEState::Paused;
+	}
+
+	void FMixerSourceDashboardViewFactory::OnPIEResumed(bool bSimulating)
+	{
+		PIEState = EPIEState::Running;
 	}
 
 #if AUDIO_INSIGHTS_SHOW_SOURCE_CONTEXT_MENU
@@ -583,18 +637,9 @@ namespace UE::Audio::Insights
 			// Only add new sources if there are less than the max 
 			const bool bCanAddNewSources = PlotWidgetMetadataPerCurve->Num() < MaxPlotSources;
 
-			const double PointTime = SourceDataPoint.Timestamp;
-			BeginTimestamp = FMath::Min(BeginTimestamp, PointTime);
-			CurrentTimestamp = FMath::Max(CurrentTimestamp, PointTime);
-			const double DataPointTime = SourceDataPoint.Timestamp - BeginTimestamp;
-
 			// For each column, get the array for this data point's source id and add the value to that data array
-			for (auto Iter = PlotWidgetCurveIdToPointDataMapPerColumn.CreateIterator(); Iter; ++Iter)
+			for (const auto& [ColumnName, DataMap] : PlotWidgetCurveIdToPointDataMapPerColumn)
 			{
-				const FName& ColumnName = Iter.Key();
-				auto DataFunc = GetPlotColumnDataFunc(ColumnName);
-				TSharedPtr<FPointDataPerCurveMap>& DataMap = Iter.Value();
-
 				// Add new data point array
 				if (bCanAddNewSources && !DataMap->Contains(SourceId))
 				{
@@ -605,8 +650,20 @@ namespace UE::Audio::Insights
 				TArray<FPlotCurvePoint>* DataPoints = DataMap->Find(SourceId);
 				if (DataPoints)
 				{
-					const float Value = (DataFunc)(SourceDataPoint);
-					DataPoints->Emplace(DataPointTime, Value);
+					auto DataFunc = GetPlotColumnDataFunc(ColumnName);
+					const ::Audio::TCircularAudioBuffer<FDataPoint>& TimeStampedValues = (DataFunc)(SourceDataPoint);
+
+					const ::Audio::DisjointedArrayView<const FDataPoint> TimeStampedValuesDisjointedArrayView = TimeStampedValues.PeekInPlace(TimeStampedValues.Num());
+					
+					for (const auto& [Timestamp, Value] : TimeStampedValuesDisjointedArrayView.FirstBuffer)
+					{
+						BeginTimestamp   = FMath::Min(BeginTimestamp, Timestamp);
+						CurrentTimestamp = FMath::Max(CurrentTimestamp, Timestamp);
+
+						const double DataPointTime = Timestamp - BeginTimestamp;
+
+						DataPoints->Emplace(DataPointTime, Value);
+					}
 				}
 			}
 
@@ -630,33 +687,31 @@ namespace UE::Audio::Insights
 			}
 		}
 
-		// Remove points that are older than max history limit from the most recent timestamp
-		const static auto RemoveOldCurvePoints = [this](TSharedPtr<FPointDataPerCurveMap> PlotWidgetPointDataPerCurve)
+		// Remove old points and set curve data for each widget 
+		const double PlotDrawLimitTimestamp = CurrentTimestamp - BeginTimestamp - (MaxPlotHistorySeconds + 0.2 /* extra grace time to avoid curve cuts being displayed */);
+
+		for (int32 WidgetIndex = 0; WidgetIndex < NumPlotWidgets; ++WidgetIndex)
 		{
-			for (auto Iter = PlotWidgetPointDataPerCurve->CreateIterator(); Iter; ++Iter)
+			const FName& SelectedPlotColumn = SelectedPlotColumnNames[WidgetIndex];
+
+			const TSharedPtr<FPointDataPerCurveMap> CurveDataMapPtr = *PlotWidgetCurveIdToPointDataMapPerColumn.Find(SelectedPlotColumn);
+			if (CurveDataMapPtr.IsValid())
 			{
-				TArray<FPlotCurvePoint>& CurvePoints = Iter->Value;
-				for (int32 i = CurvePoints.Num() - 1; i >= 0; --i)
+				// Remove points that are older than max history limit from the most recent timestamp
+				for (auto& [CurveId, CurvePoints] : *CurveDataMapPtr)
 				{
-					const FPlotCurvePoint& Point = CurvePoints[i];
-					if (Point.Key + BeginTimestamp < CurrentTimestamp - MaxPlotHistorySeconds)
+					const int32 FoundIndex = CurvePoints.IndexOfByPredicate([&PlotDrawLimitTimestamp](const FDataPoint& InDataPoint)
 					{
-						CurvePoints.RemoveAt(i);
+						return InDataPoint.Key >= PlotDrawLimitTimestamp;
+					});
+
+					if (FoundIndex > 0)
+					{
+						CurvePoints.RemoveAt(0, FoundIndex, EAllowShrinking::No);
 					}
 				}
-			}
-		};
-
-		// Remove old points and set curve data for each widget 
-		for (int32 WidgetIndex = 0; WidgetIndex < NumPlotWidgets; ++WidgetIndex)
-		{
-			TSharedPtr<SAudioCurveView> PlotWidget = PlotWidgets[WidgetIndex];
-			const FName& SelectedPlotColumn = SelectedPlotColumnNames[WidgetIndex];
 
-			if (TSharedPtr<FPointDataPerCurveMap> CurveData = *PlotWidgetCurveIdToPointDataMapPerColumn.Find(SelectedPlotColumn))
-			{
-				RemoveOldCurvePoints(CurveData);
-				PlotWidget->SetCurvesPointData(CurveData);
+				PlotWidgets[WidgetIndex]->SetCurvesPointData(CurveDataMapPtr);
 			}
 		}
 	}
@@ -670,7 +725,7 @@ namespace UE::Audio::Insights
 				{
 					"Amplitude",
 					{
-						[](const IDashboardDataViewEntry& InData) { return MixerSourcePrivate::CastEntry(InData).Envelope; },
+						[](const IDashboardDataViewEntry& InData) -> const ::Audio::TCircularAudioBuffer<FDataPoint>& { return MixerSourcePrivate::CastEntry(InData).EnvelopeDataPoints; },
 						FSlateStyle::Get().GetAmpFloatFormat()
 					}
 				},
@@ -678,35 +733,35 @@ namespace UE::Audio::Insights
 					"Volume",
 					{
 
-						[](const IDashboardDataViewEntry& InData) { return MixerSourcePrivate::CastEntry(InData).Volume; },
+						[](const IDashboardDataViewEntry& InData) -> const ::Audio::TCircularAudioBuffer<FDataPoint>& { return MixerSourcePrivate::CastEntry(InData).VolumeDataPoints; },
 						FSlateStyle::Get().GetAmpFloatFormat()
 					}
 				},
 				{
 					"DistanceAttenuation",
 					{
-						[](const IDashboardDataViewEntry& InData) { return MixerSourcePrivate::CastEntry(InData).DistanceAttenuation; }, 
+						[](const IDashboardDataViewEntry& InData) -> const ::Audio::TCircularAudioBuffer<FDataPoint>& { return MixerSourcePrivate::CastEntry(InData).DistanceAttenuationDataPoints; },
 						FSlateStyle::Get().GetAmpFloatFormat()
 					}
 				},
 				{
 					"Pitch",
 					{
-						[](const IDashboardDataViewEntry& InData) { return MixerSourcePrivate::CastEntry(InData).Pitch; },
+						[](const IDashboardDataViewEntry& InData) -> const ::Audio::TCircularAudioBuffer<FDataPoint>& { return MixerSourcePrivate::CastEntry(InData).PitchDataPoints; },
 						FSlateStyle::Get().GetPitchFloatFormat()
 					}
 				},
 				{
 					"LPF",
 					{
-						[](const IDashboardDataViewEntry& InData) { return MixerSourcePrivate::CastEntry(InData).LPFFreq; },
+						[](const IDashboardDataViewEntry& InData) -> const ::Audio::TCircularAudioBuffer<FDataPoint>& { return MixerSourcePrivate::CastEntry(InData).LPFFreqDataPoints; },
 							FSlateStyle::Get().GetFreqFloatFormat() 
 					}
 				},
 				{
 					"HPF",
 					{
-						[](const IDashboardDataViewEntry& InData) { return MixerSourcePrivate::CastEntry(InData).HPFFreq; },
+						[](const IDashboardDataViewEntry& InData) -> const ::Audio::TCircularAudioBuffer<FDataPoint>& { return MixerSourcePrivate::CastEntry(InData).HPFFreqDataPoints; },
 						FSlateStyle::Get().GetFreqFloatFormat()
 					}
 				}
@@ -725,7 +780,7 @@ namespace UE::Audio::Insights
 		return nullptr;
 	}
 
-	const TFunctionRef<float(const IDashboardDataViewEntry& InData)> FMixerSourceDashboardViewFactory::GetPlotColumnDataFunc(const FName& ColumnName)
+	const TFunctionRef<const ::Audio::TCircularAudioBuffer<FDataPoint>&(const IDashboardDataViewEntry& InData)> FMixerSourceDashboardViewFactory::GetPlotColumnDataFunc(const FName& ColumnName)
 	{
 		return GetPlotColumnInfo().Find(ColumnName)->DataFunc;
 	}
@@ -763,6 +818,23 @@ namespace UE::Audio::Insights
 		// Create plot widgets
 		auto GetViewRange = [this]()
 		{
+			if (PIEState == EPIEState::Stopped || BeginTimestamp == TNumericLimits<double>::Max())
+			{
+				return TRange<double>(0, MaxPlotHistorySeconds);
+			}
+
+			if (PIEState == EPIEState::Running)
+			{
+				const FTraceModule& TraceModule = FAudioInsightsModule::GetChecked().GetTraceModule();
+				const double FirstTimestamp = TraceModule.GetFirstTimeStamp();
+
+				const double CurrentTime = FPlatformTime::Seconds();
+				const double RelativeCurrentTime = CurrentTime - FirstTimestamp;
+				const double TimestampsDiff = RelativeCurrentTime - (CurrentTimestamp - BeginTimestamp);
+
+				CurrentTimestamp += TimestampsDiff;
+			}
+
 			return TRange<double>(FMath::Max(0, CurrentTimestamp - MaxPlotHistorySeconds - BeginTimestamp), CurrentTimestamp - BeginTimestamp);
 		};
 
@@ -772,7 +844,8 @@ namespace UE::Audio::Insights
 			for (int32 WidgetNum = 0; WidgetNum < NumPlotWidgets; ++WidgetNum)
 			{
 				SAssignNew(PlotWidgets[WidgetNum], SAudioCurveView)
-					.ViewRange_Lambda(GetViewRange);
+				.ViewRange_Lambda(GetViewRange)
+				.PixelSnappingMethod(EWidgetPixelSnapping::Disabled);
 			}
 		}
 
@@ -816,6 +889,8 @@ namespace UE::Audio::Insights
 				SNew(SSimpleTimeSlider)
 				.ViewRange_Lambda(GetViewRange)
 				.ClampRangeHighlightSize(0.0f) // Hide clamp range
+				.ScrubPosition_Lambda([]() { return -1.0; }) // Hide scrub
+				.PixelSnappingMethod(EWidgetPixelSnapping::Disabled)
 			]
 			+ SVerticalBox::Slot()
 			.AutoHeight()
@@ -824,22 +899,9 @@ namespace UE::Audio::Insights
 				CreatePlotColumnComboBoxWidget(0)
 			]
 			+ SVerticalBox::Slot()
-			.AutoHeight()
 			.HAlign(HAlign_Fill)
 			[
 				PlotWidgets[0].ToSharedRef()
-			]
-			+ SVerticalBox::Slot()
-			.AutoHeight()
-			.HAlign(HAlign_Fill)
-			[
-				CreatePlotColumnComboBoxWidget(1)
-			]
-			+ SVerticalBox::Slot()
-			.AutoHeight()
-			.HAlign(HAlign_Fill)
-			[
-				PlotWidgets[1].ToSharedRef()
 			];
 	}
 
@@ -944,7 +1006,11 @@ namespace UE::Audio::Insights
 	TSharedRef<SWidget> FMixerSourceDashboardViewFactory::MakeWidget()
 	{
 		FDashboardFactory::OnActiveAudioDeviceChanged.AddSP(this, &FMixerSourceDashboardViewFactory::ClearMutesAndSolos);
+
+		FEditorDelegates::PostPIEStarted.AddSP(this, &FMixerSourceDashboardViewFactory::OnPIEStarted);
 		FEditorDelegates::EndPIE.AddSP(this, &FMixerSourceDashboardViewFactory::OnPIEStopped);
+		FEditorDelegates::PausePIE.AddSP(this, &FMixerSourceDashboardViewFactory::OnPIEPaused);
+		FEditorDelegates::ResumePIE.AddSP(this, &FMixerSourceDashboardViewFactory::OnPIEResumed);
 
 		TSharedRef<SWidget> MuteSoloWidget = MakeMuteSoloWidget();
 		TSharedRef<SWidget> TableDashboardWidget = FTraceTableDashboardViewFactory::MakeWidget();
diff --git a/Engine/Plugins/AudioInsights/Source/AudioInsights/Private/Views/MixerSourceDashboardViewFactory.h b/Engine/Plugins/AudioInsights/Source/AudioInsights/Private/Views/MixerSourceDashboardViewFactory.h
--- a/Engine/Plugins/AudioInsights/Source/AudioInsights/Private/Views/MixerSourceDashboardViewFactory.h
+++ b/Engine/Plugins/AudioInsights/Source/AudioInsights/Private/Views/MixerSourceDashboardViewFactory.h
@@ -1,6 +1,7 @@
 // Copyright Epic Games, Inc. All Rights Reserved.
 #pragma once
 
+#include "Messages/MixerSourceTraceMessages.h"
 #include "SAudioCurveView.h"
 #include "Views/TableDashboardViewFactory.h"
 #include "Widgets/Input/SCheckBox.h"
@@ -17,7 +18,7 @@ namespace UE::Audio::Insights
 		FMixerSourceDashboardViewFactory();
 		struct FPlotColumnInfo
 		{
-			const TFunctionRef<float(const IDashboardDataViewEntry& InData)> DataFunc;
+			const TFunctionRef<const ::Audio::TCircularAudioBuffer<FDataPoint>&(const IDashboardDataViewEntry& InData)> DataFunc;
 			const FNumberFormattingOptions* FormatOptions;
 		};
 
@@ -47,13 +48,16 @@ namespace UE::Audio::Insights
 
 	private:
 		void ResetPlots();
+		void OnPIEStarted(bool bSimulating);
 		void OnPIEStopped(bool bSimulating);
+		void OnPIEPaused(bool bSimulating);
+		void OnPIEResumed(bool bSimulating);
 		void UpdatePlotsWidgetsData();
 		void UpdateSoloMuteState();
 
 		// Column information used by plot widgets, keyed by column name. These keys should be a subset of the keys in GetColumns(). 
 		const TMap<FName, FPlotColumnInfo>& GetPlotColumnInfo();
-		const TFunctionRef<float(const IDashboardDataViewEntry& InData)> GetPlotColumnDataFunc(const FName& ColumnName);
+		const TFunctionRef<const ::Audio::TCircularAudioBuffer<FDataPoint>&(const IDashboardDataViewEntry& InData)> GetPlotColumnDataFunc(const FName& ColumnName);
 		const FNumberFormattingOptions* GetPlotColumnNumberFormat(const FName& ColumnName);
 		const FText GetPlotColumnDisplayName(const FName& ColumnName);
 
@@ -90,10 +94,19 @@ namespace UE::Audio::Insights
 		// Column names for plot selector widget 
 		TArray<FName> ColumnNames;
 		
+		enum class EPIEState : uint8
+		{
+			Running,
+			Stopped,
+			Paused
+		};
+
+		EPIEState PIEState = EPIEState::Stopped;
+
 		double BeginTimestamp = TNumericLimits<double>::Max();
-		double CurrentTimestamp = 0;
+		double CurrentTimestamp = 0.0;
 
-		const static int32 NumPlotWidgets = 2;
+		const static int32 NumPlotWidgets = 1;
 		TArray<FName> SelectedPlotColumnNames;
 		TArray<TSharedPtr<SAudioCurveView>> PlotWidgets;
 
diff --git a/Engine/Plugins/AudioInsights/Source/AudioInsights/Private/Views/SAudioCurveView.cpp b/Engine/Plugins/AudioInsights/Source/AudioInsights/Private/Views/SAudioCurveView.cpp
--- a/Engine/Plugins/AudioInsights/Source/AudioInsights/Private/Views/SAudioCurveView.cpp
+++ b/Engine/Plugins/AudioInsights/Source/AudioInsights/Private/Views/SAudioCurveView.cpp
@@ -252,7 +252,7 @@ int32 SAudioCurveView::PaintGridLines(const FGeometry& AllottedGeometry, const F
 				GridPoints,
 				LineDrawEffects,
 				GridLineColor.Get(),
-				false
+				true
 			);
 			VerticalLineValue += HorizontalAxisIncrement.Get();
 		}
@@ -308,11 +308,6 @@ int32 SAudioCurveView::PaintGridLines(const FGeometry& AllottedGeometry, const F
 
 int32 SAudioCurveView::PaintCurves(const FGeometry& AllottedGeometry, const FSlateRect& MyCullingRect, FSlateWindowElementList& OutDrawElements, int32 LayerId, const FWidgetStyle& InWidgetStyle, bool bParentEnabled) const 
 {
-	static const float LargeFrameTime = 0.25;  // ideally, we could check the recorded frame data for actual frame length
-	
-	const SSimpleTimeSlider::FScrubRangeToScreen RangeToScreen(ViewRange.Get(), AllottedGeometry.GetLocalSize());
-	const FVector2f Size = AllottedGeometry.GetLocalSize();
-
 	// Skip drawing if curve data is not initialized yet
 	if (!PointDataPerCurve || !MetadataPerCurve)
 	{
@@ -325,7 +320,11 @@ int32 SAudioCurveView::PaintCurves(const FGeometry& AllottedGeometry, const FSla
 		return LayerId;
 	}
 
+	const SSimpleTimeSlider::FScrubRangeToScreen RangeToScreen(ViewRange.Get(), AllottedGeometry.GetLocalSize());
 	LayerId = PaintGridLines(AllottedGeometry, MyCullingRect, OutDrawElements, LayerId, InWidgetStyle, bParentEnabled, RangeToScreen);
+	
+	static constexpr float LargeFrameTime = 0.5f; // ideally, we could check the recorded frame data for actual frame length
+	const FVector2f Size = AllottedGeometry.GetLocalSize();
 
 	// Create and draw points per curve
 	for (auto Iter = MetadataPerCurve->CreateConstIterator(); Iter; ++Iter)
@@ -345,7 +344,8 @@ int32 SAudioCurveView::PaintCurves(const FGeometry& AllottedGeometry, const FSla
 		for (int32 i = 0; i < CurvePoints->Num(); i++)
 		{
 			const FCurvePoint& Point = (*CurvePoints)[i];
-			if (Point.Value - PrevX > LargeFrameTime && Points.Num()>1)
+
+			if (Point.Value - PrevX > LargeFrameTime && Points.Num() > 1)
 			{
 				// break the line list - data has stopped and started again
 				FSlateDrawElement::MakeLines(
@@ -355,15 +355,16 @@ int32 SAudioCurveView::PaintCurves(const FGeometry& AllottedGeometry, const FSla
 					Points,
 					LineDrawEffects,
 					CurveMetadata.CurveColor,
-					false
+					true
 				);
 				Points.SetNum(0, EAllowShrinking::No);
 			}
 
 			const float X = RangeToScreen.InputToLocalX(Point.Key);
-			PrevX = Point.Key;
 			const float Y = ValueToLocalY(Size, Point.Value);
-			Points.Add(FVector2f(X, Y));
+			Points.Emplace(X, Y);
+
+			PrevX = Point.Key;
 		}
 
 		FSlateDrawElement::MakeLines(
@@ -373,9 +374,10 @@ int32 SAudioCurveView::PaintCurves(const FGeometry& AllottedGeometry, const FSla
 			Points,
 			LineDrawEffects,
 			CurveMetadata.CurveColor,
-			false
+			true
 		);
 	}
+
 	return LayerId;
 }
 
