[Audio Insights] SAudioCurveView: Fixed bug where plot lines got cut sometimes (we were doing Point.Value - PrevX instead of Point.Key - PrevX, see diff for more info) + misc. code adjustments for readability.

diff --git a/Engine/Plugins/AudioInsights/Source/AudioInsights/Private/Views/SAudioCurveView.cpp b/Engine/Plugins/AudioInsights/Source/AudioInsights/Private/Views/SAudioCurveView.cpp
--- a/Engine/Plugins/AudioInsights/Source/AudioInsights/Private/Views/SAudioCurveView.cpp
+++ b/Engine/Plugins/AudioInsights/Source/AudioInsights/Private/Views/SAudioCurveView.cpp
@@ -309,7 +309,7 @@ int32 SAudioCurveView::PaintGridLines(const FGeometry& AllottedGeometry, const F
 int32 SAudioCurveView::PaintCurves(const FGeometry& AllottedGeometry, const FSlateRect& MyCullingRect, FSlateWindowElementList& OutDrawElements, int32 LayerId, const FWidgetStyle& InWidgetStyle, bool bParentEnabled) const 
 {
 	// Skip drawing if curve data is not initialized yet
-	if (!PointDataPerCurve || !MetadataPerCurve)
+	if (!PointDataPerCurve.IsValid() || !MetadataPerCurve.IsValid())
 	{
 		return LayerId;
 	}
@@ -330,34 +330,35 @@ int32 SAudioCurveView::PaintCurves(const FGeometry& AllottedGeometry, const FSla
 	for (auto Iter = MetadataPerCurve->CreateConstIterator(); Iter; ++Iter)
 	{
 		const FCurveMetadata& CurveMetadata = Iter->Value;
-		const int32 CurveId = CurveMetadata.CurveId;
-		const TArray<FCurvePoint>* CurvePoints = PointDataPerCurve->Find(CurveId);
-		if (!CurvePoints || CurvePoints->Num() <= 0)
+		const TArray<FCurvePoint>* CurvePointsPtr = PointDataPerCurve->Find(CurveMetadata.CurveId);
+
+		if (!CurvePointsPtr || CurvePointsPtr->IsEmpty())
 		{
 			continue;
 		}
 
+		const TArray<FCurvePoint>& CurvePoints = (*CurvePointsPtr);
+
 		TArray<FVector2f> Points;
-		Points.Reserve(CurvePoints->Num());
+		Points.Reserve(CurvePoints.Num());
 
-		float PrevX = (*CurvePoints)[0].Key;
-		for (int32 i = 0; i < CurvePoints->Num(); i++)
+		float PrevX = CurvePoints[0].Key;
+		for (const FCurvePoint& Point : CurvePoints)
 		{
-			const FCurvePoint& Point = (*CurvePoints)[i];
-
-			if (Point.Value - PrevX > LargeFrameTime && Points.Num() > 1)
+			if (Point.Key - PrevX > LargeFrameTime && Points.Num() > 1)
 			{
 				// break the line list - data has stopped and started again
 				FSlateDrawElement::MakeLines(
 					OutDrawElements,
-					LayerId++,
+					++LayerId,
 					AllottedGeometry.ToPaintGeometry(),
 					Points,
 					LineDrawEffects,
 					CurveMetadata.CurveColor,
 					true
 				);
-				Points.SetNum(0, EAllowShrinking::No);
+
+				Points.Reset();
 			}
 
 			const float X = RangeToScreen.InputToLocalX(Point.Key);
