AudioInsights - Fix possible crash encountered during commandlet execution

diff --git a/Engine/Plugins/AudioInsights/Source/AudioInsights/Private/Providers/SoundSubmixProvider.cpp b/Engine/Plugins/AudioInsights/Source/AudioInsights/Private/Providers/SoundSubmixProvider.cpp
--- a/Engine/Plugins/AudioInsights/Source/AudioInsights/Private/Providers/SoundSubmixProvider.cpp
+++ b/Engine/Plugins/AudioInsights/Source/AudioInsights/Private/Providers/SoundSubmixProvider.cpp
@@ -144,13 +144,16 @@ namespace UE::Audio::Insights
 
 		for (const TSharedPtr<FSoundSubmixAssetDashboardEntry>& SubmixDataViewEntry : SubmixDataViewEntries)
 		{
-			UpdateDeviceEntry(DeviceId, SubmixDataViewEntry->SoundSubmix->GetUniqueID(), [&SubmixDataViewEntry](TSharedPtr<FSoundSubmixAssetDashboardEntry>& Entry)
+			if (SubmixDataViewEntry.IsValid() && SubmixDataViewEntry->SoundSubmix.IsValid())
 			{
-				if (!Entry.IsValid())
+				UpdateDeviceEntry(DeviceId, SubmixDataViewEntry->SoundSubmix->GetUniqueID(), [&SubmixDataViewEntry](TSharedPtr<FSoundSubmixAssetDashboardEntry>& Entry)
 				{
-					Entry = SubmixDataViewEntry;
-				}
-			});
+					if (!Entry.IsValid())
+					{
+						Entry = SubmixDataViewEntry;
+					}
+				});
+			}
 		}
 
 		return true;
