[Audio Insights] FAudioInsightsModule: Changed IsRunningCookCommandlet to IsRunningCommandlet to make sure AudioInsights is not running while other commandlets are running.

diff --git a/Engine/Plugins/AudioInsights/Source/AudioInsights/Private/AudioInsightsModule.cpp b/Engine/Plugins/AudioInsights/Source/AudioInsights/Private/AudioInsightsModule.cpp
--- a/Engine/Plugins/AudioInsights/Source/AudioInsights/Private/AudioInsightsModule.cpp
+++ b/Engine/Plugins/AudioInsights/Source/AudioInsights/Private/AudioInsightsModule.cpp
@@ -32,8 +32,8 @@ namespace UE::Audio::Insights
 {
 	void FAudioInsightsModule::StartupModule()
 	{
-		// Don't run providers in cook commandlet to avoid additional, unnecessary overhead as audio insights is dormant.
-		if (!IsRunningCookCommandlet())
+		// Don't run providers in any commandlet to avoid additional, unnecessary overhead as audio insights is dormant.
+		if (!IsRunningCommandlet())
 		{
 			IModularFeatures::Get().RegisterModularFeature(TraceServices::ModuleFeatureName, &TraceModule);
 
@@ -56,7 +56,7 @@ namespace UE::Audio::Insights
 
 	void FAudioInsightsModule::ShutdownModule()
 	{
-		if (!IsRunningCookCommandlet())
+		if (!IsRunningCommandlet())
 		{
 			DashboardFactory.Reset();
 			IModularFeatures::Get().UnregisterModularFeature(TraceServices::ModuleFeatureName, &TraceModule);
