[Audio Insights] Audio Insights for Unreal Insights (part I): refactor previous to define a functional Unreal Insights component.

diff --git a/Engine/Plugins/AudioInsights/AudioInsights.uplugin b/Engine/Plugins/AudioInsights/AudioInsights.uplugin
--- a/Engine/Plugins/AudioInsights/AudioInsights.uplugin
+++ b/Engine/Plugins/AudioInsights/AudioInsights.uplugin
@@ -14,10 +14,25 @@
 	"IsBetaVersion": false,
 	"IsExperimentalVersion": true,
 	"Installed": false,
+	"SupportedPrograms": [ "UnrealInsights" ],
 	"Modules":
 	[
 		{
 			"Name": "AudioInsights",
+			"Type": "EditorAndProgram",
+			"LoadingPhase": "Default",
+			"ProgramAllowList":
+			[
+				"UnrealInsights"
+			],
+			"PlatformAllowList":
+			[
+				"Win64",
+				"Mac"
+			]
+		},
+		{
+			"Name": "AudioInsightsEditor",
 			"Type": "Editor",
 			"LoadingPhase": "Default",
 			"PlatformAllowList":
@@ -32,10 +47,6 @@
 		{
 			"Name": "AudioWidgets",
 			"Enabled": true
-		},
-		{
-			"Name": "GameplayInsights",
-			"Enabled": true
 		}
 	]
 }
\ No newline at end of file
diff --git a/Engine/Plugins/AudioInsights/Source/AudioInsights/AudioInsights.Build.cs b/Engine/Plugins/AudioInsights/Source/AudioInsights/AudioInsights.Build.cs
--- a/Engine/Plugins/AudioInsights/Source/AudioInsights/AudioInsights.Build.cs
+++ b/Engine/Plugins/AudioInsights/Source/AudioInsights/AudioInsights.Build.cs
@@ -13,7 +13,9 @@ public class AudioInsights : ModuleRules
 			new string[]
 			{
 				"Core",
+				"TraceInsights",
 				"TraceServices",
+				"ToolWidgets",
 			}
 		);	
 		
@@ -21,26 +23,26 @@ public class AudioInsights : ModuleRules
 		(
 			new string[]
 			{
-				"CoreUObject",
-				"Engine",
-				"AudioMixer",
 				"AudioMixerCore",
-				"AudioWidgets",
-				"GameplayInsights",
+				"CoreUObject",
 				"InputCore",
-				"LevelEditor",
 				"OutputLog",
-				"SignalProcessing",
+				"SessionServices",
 				"Slate",
 				"SlateCore",
-				"ToolMenus",
-				"ToolWidgets",
-				"TraceLog",
 				"TraceAnalysis",
-				"TraceInsights",
-				"UnrealEd",
-				"WorkspaceMenuStructure"
+				"TraceLog",
 			}
-		);		
+		);
+
+		if (Target.Type == TargetType.Editor)
+		{
+			PrivateDependencyModuleNames.AddRange(
+			new string[]
+			{
+				"Engine",
+				"UnrealEd",
+			});
+		}
 	}
 }
diff --git a/Engine/Plugins/AudioInsights/Source/AudioInsights/Private/Analyzers/AudioMeterAnalyzer.cpp b/Engine/Plugins/AudioInsights/Source/AudioInsights/Private/Analyzers/AudioMeterAnalyzer.cpp
deleted file mode 100644
diff --git a/Engine/Plugins/AudioInsights/Source/AudioInsights/Private/Analyzers/AudioMeterAnalyzer.h b/Engine/Plugins/AudioInsights/Source/AudioInsights/Private/Analyzers/AudioMeterAnalyzer.h
deleted file mode 100644
diff --git a/Engine/Plugins/AudioInsights/Source/AudioInsights/Private/Analyzers/AudioMeterSubmixAnalyzer.cpp b/Engine/Plugins/AudioInsights/Source/AudioInsights/Private/Analyzers/AudioMeterSubmixAnalyzer.cpp
deleted file mode 100644
diff --git a/Engine/Plugins/AudioInsights/Source/AudioInsights/Private/Analyzers/AudioMeterSubmixAnalyzer.h b/Engine/Plugins/AudioInsights/Source/AudioInsights/Private/Analyzers/AudioMeterSubmixAnalyzer.h
deleted file mode 100644
diff --git a/Engine/Plugins/AudioInsights/Source/AudioInsights/Private/Analyzers/AudioOscilloscopeAnalyzer.cpp b/Engine/Plugins/AudioInsights/Source/AudioInsights/Private/Analyzers/AudioOscilloscopeAnalyzer.cpp
deleted file mode 100644
diff --git a/Engine/Plugins/AudioInsights/Source/AudioInsights/Private/Analyzers/AudioOscilloscopeAnalyzer.h b/Engine/Plugins/AudioInsights/Source/AudioInsights/Private/Analyzers/AudioOscilloscopeAnalyzer.h
deleted file mode 100644
diff --git a/Engine/Plugins/AudioInsights/Source/AudioInsights/Private/AudioInsightsDashboardAssetCommands.cpp b/Engine/Plugins/AudioInsights/Source/AudioInsights/Private/AudioInsightsDashboardAssetCommands.cpp
--- a/Engine/Plugins/AudioInsights/Source/AudioInsights/Private/AudioInsightsDashboardAssetCommands.cpp
+++ b/Engine/Plugins/AudioInsights/Source/AudioInsights/Private/AudioInsightsDashboardAssetCommands.cpp
@@ -25,10 +25,6 @@ namespace UE::Audio::Insights
 
 		UI_COMMAND(Start, "Start Trace", "Starts the active trace session used by Audio Insights.", EUserInterfaceActionType::Button, FInputChord());
 		UI_COMMAND(Stop, "Stop Trace", "Stops the active trace session used by Audio Insights.", EUserInterfaceActionType::Button, FInputChord());
-
-		UI_COMMAND(Mute, "Mute (Toggle)", "Mutes the selected item in the Mixer Source View.", EUserInterfaceActionType::Button, FInputChord());
-		UI_COMMAND(Solo, "Solo (Toggle)", "Solo the selected item in the Mixer Source View.", EUserInterfaceActionType::Button, FInputChord());
-		UI_COMMAND(ClearMuteSolo, "Clear All Mute/Solo", "Clears any Mute/Solo state from any of the items in the Mixer Source View.", EUserInterfaceActionType::Button, FInputChord());
 	}
 
 	void FDashboardAssetCommands::AddAssetCommands(FToolBarBuilder& OutToolbarBuilder) const
@@ -72,21 +68,6 @@ namespace UE::Audio::Insights
 		return Stop;
 	}
 
-	TSharedPtr<const FUICommandInfo> FDashboardAssetCommands::GetMuteCommand() const
-	{
-		return Mute;
-	};
-
-	TSharedPtr<const FUICommandInfo> FDashboardAssetCommands::GetSoloCommand() const
-	{
-		return Solo;
-	}
-
-	TSharedPtr<const FUICommandInfo> FDashboardAssetCommands::GetClearMuteSoloCommand() const
-	{
-		return ClearMuteSolo;
-	}
-
 	FSlateIcon FDashboardAssetCommands::GetStartIcon() const
 	{
 		const FTraceModule& TraceModule = FAudioInsightsModule::GetChecked().GetTraceModule();
diff --git a/Engine/Plugins/AudioInsights/Source/AudioInsights/Private/AudioInsightsDashboardFactory.cpp b/Engine/Plugins/AudioInsights/Source/AudioInsights/Private/AudioInsightsDashboardFactory.cpp
--- a/Engine/Plugins/AudioInsights/Source/AudioInsights/Private/AudioInsightsDashboardFactory.cpp
+++ b/Engine/Plugins/AudioInsights/Source/AudioInsights/Private/AudioInsightsDashboardFactory.cpp
@@ -1,20 +1,14 @@
 // Copyright Epic Games, Inc. All Rights Reserved.
 #include "AudioInsightsDashboardFactory.h"
 
-#include "AudioDevice.h"
-#include "AudioDeviceManager.h"
 #include "AudioInsightsModule.h"
 #include "AudioInsightsStyle.h"
-#include "Editor.h"
-#include "Engine/World.h"
 #include "Framework/Docking/TabManager.h"
+#include "Framework/MultiBox/MultiBoxBuilder.h"
 #include "Internationalization/Text.h"
-#include "IPropertyTypeCustomization.h"
-#include "Kismet2/DebuggerCommands.h"
 #include "Templates/SharedPointer.h"
-#include "ToolMenu.h"
-#include "ToolMenus.h"
 #include "Widgets/Input/SCheckBox.h"
+#include "Widgets/Layout/SBox.h"
 
 #define LOCTEXT_NAMESPACE "AudioInsights"
 
@@ -24,173 +18,6 @@ namespace UE::Audio::Insights
 	namespace DashboardFactoryPrivate
 	{
 		static const FText ToolName = LOCTEXT("AudioDashboard_ToolName", "Audio Insights");
-
-		static const FName MainToolbarName = "MainToolbar";
-		static const FText MainToolbarDisplayName = LOCTEXT("AudioDashboard_MainToolbarDisplayName", "Dashboard Transport");
-
-		static const FText PreviewDeviceDisplayName = LOCTEXT("AudioDashboard_PreviewDevice", "[Preview Audio]");
-		static const FText DashboardWorldSelectDescription = LOCTEXT("AudioDashboard_SelectWorldDescription", "Select world(s) to monitor (worlds may share audio output).");
-
-		FText GetDebugNameFromDeviceId(::Audio::FDeviceId InDeviceId)
-		{
-			FString WorldName;
-			if (FAudioDeviceManager* DeviceManager = FAudioDeviceManager::Get())
-			{
-				TArray<UWorld*> DeviceWorlds = DeviceManager->GetWorldsUsingAudioDevice(InDeviceId);
-				for (const UWorld* World : DeviceWorlds)
-				{
-					if (!WorldName.IsEmpty())
-					{
-						WorldName += TEXT(", ");
-					}
-					WorldName += World->GetDebugDisplayName();
-				}
-			}
-
-			if (WorldName.IsEmpty())
-			{
-				return PreviewDeviceDisplayName;
-			}
-
-			return FText::FromString(WorldName);
-		}
-	} // namespace DashboardFactoryPrivate
-
-	FDashboardFactory::FDashboardFactory()
-	{
-	}
-
-	void FDashboardFactory::OnWorldRegisteredToAudioDevice(const UWorld* InWorld, ::Audio::FDeviceId InDeviceId)
-	{
-		if (InDeviceId != INDEX_NONE)
-		{
-			if (bStartWithPIE)
-			{
-				const FTraceModule& TraceModule = FAudioInsightsModule::GetChecked().GetTraceModule();
-				TraceModule.StartTraceAnalysis();
-
-				ActiveDeviceId = InDeviceId;
-			}
-		}
-
-		RefreshDeviceSelector();
-	}
-
-	void FDashboardFactory::OnPIEStarted(bool bSimulating)
-	{
-		if (bStartWithPIE)
-		{
-			const FTraceModule& TraceModule = FAudioInsightsModule::GetChecked().GetTraceModule();
-			TraceModule.StartTraceAnalysis();
-		}
-	}
-
-	void FDashboardFactory::OnPostPIEStarted(bool bSimulating)
-	{
-		OnActiveAudioDeviceChanged.Broadcast();
-	}
-
-	void FDashboardFactory::OnPIEStopped(bool bSimulating)
-	{
-		if (bStopWithPIE)
-		{
-			const FTraceModule& TraceModule = FAudioInsightsModule::GetChecked().GetTraceModule();
-			TraceModule.StopTraceAnalysis();
-		}
-
-		RefreshDeviceSelector();
-	}
-
-	void FDashboardFactory::OnWorldUnregisteredFromAudioDevice(const UWorld* InWorld, ::Audio::FDeviceId InDeviceId)
-	{
-		RefreshDeviceSelector();
-	}
-
-	void FDashboardFactory::OnDeviceDestroyed(::Audio::FDeviceId InDeviceId)
-	{
-		if (ActiveDeviceId == InDeviceId)
-		{
-			if (FAudioDeviceManager* DeviceManager = FAudioDeviceManager::Get())
-			{
-				ActiveDeviceId = DeviceManager->GetMainAudioDeviceID();
-			}
-		}
-
-		AudioDeviceIds.RemoveAll([InDeviceId](const TSharedPtr<::Audio::FDeviceId>& DeviceIdPtr)
-		{
-			return *DeviceIdPtr.Get() == InDeviceId;
-		});
-
-		if (AudioDeviceComboBox.IsValid())
-		{
-			AudioDeviceComboBox->RefreshOptions();
-		}
-
-		OnActiveAudioDeviceChanged.Broadcast();
-	}
-
-	void FDashboardFactory::RefreshDeviceSelector()
-	{
-		if (FAudioDeviceManager* DeviceManager = FAudioDeviceManager::Get())
-		{
-			if (!DeviceManager->IsValidAudioDevice(ActiveDeviceId))
-			{
-				ActiveDeviceId = DeviceManager->GetMainAudioDeviceID();
-			}
-		}
-
-		AudioDeviceIds.Empty();
-		if (const FAudioDeviceManager* DeviceManager = FAudioDeviceManager::Get())
-		{
-			DeviceManager->IterateOverAllDevices([this, &DeviceManager](::Audio::FDeviceId DeviceId, const FAudioDevice* AudioDevice)
-			{
-				AudioDeviceIds.Add(MakeShared<::Audio::FDeviceId>(DeviceId));
-			});
-		}
-
-		if (AudioDeviceComboBox.IsValid())
-		{
-			AudioDeviceComboBox->RefreshOptions();
-		}
-	}
-
-	void FDashboardFactory::ResetDelegates()
-	{
-		if (OnWorldRegisteredToAudioDeviceHandle.IsValid())
-		{
-			FAudioDeviceWorldDelegates::OnWorldRegisteredToAudioDevice.Remove(OnWorldRegisteredToAudioDeviceHandle);
-			OnWorldRegisteredToAudioDeviceHandle.Reset();
-		}
-
-		if (OnWorldUnregisteredFromAudioDeviceHandle.IsValid())
-		{
-			FAudioDeviceWorldDelegates::OnWorldUnregisteredWithAudioDevice.Remove(OnWorldUnregisteredFromAudioDeviceHandle);
-			OnWorldUnregisteredFromAudioDeviceHandle.Reset();
-		}
-
-		if (OnDeviceDestroyedHandle.IsValid())
-		{
-			FAudioDeviceManagerDelegates::OnAudioDeviceDestroyed.Remove(OnDeviceDestroyedHandle);
-			OnDeviceDestroyedHandle.Reset();
-		}
-
-		if (OnPIEStartedHandle.IsValid())
-		{
-			FEditorDelegates::PreBeginPIE.Remove(OnPIEStartedHandle);
-			OnPIEStartedHandle.Reset();
-		}
-
-		if (OnPostPIEStartedHandle.IsValid())
-		{
-			FEditorDelegates::PostPIEStarted.Remove(OnPostPIEStartedHandle);
-			OnPostPIEStartedHandle.Reset();
-		}
-
-		if (OnPIEStoppedHandle.IsValid())
-		{
-			FEditorDelegates::EndPIE.Remove(OnPIEStoppedHandle);
-			OnPIEStoppedHandle.Reset();
-		}
 	}
 
 	::Audio::FDeviceId FDashboardFactory::GetDeviceId() const
@@ -207,11 +34,9 @@ namespace UE::Audio::Insights
 
 		DashboardTabManager = FGlobalTabmanager::Get()->NewTabManager(DockTab);
 
-		InitDelegates();
 		TabLayout = GetDefaultTabLayout();
 
 		RegisterTabSpawners();
-		RefreshDeviceSelector();
 
 		const TSharedRef<SWidget> TabContent = SNew(SVerticalBox)
 			+ SVerticalBox::Slot()
@@ -221,11 +46,6 @@ namespace UE::Audio::Insights
 			]
 			+ SVerticalBox::Slot()
 			.AutoHeight()
-			[
-				MakeMainToolbarWidget()
-			]
-			+ SVerticalBox::Slot()
-			.AutoHeight()
 			[
 				SNew(SBox)
 				.HeightOverride(4.0f)
@@ -239,7 +59,6 @@ namespace UE::Audio::Insights
 
 		DockTab->SetOnTabClosed(SDockTab::FOnTabClosedCallback::CreateLambda([this](TSharedRef<SDockTab> TabClosed)
 		{
-			ResetDelegates();
 			UnregisterTabSpawners();
 		}));
 
@@ -333,184 +152,20 @@ namespace UE::Audio::Insights
 
 		return MenuBarBuilder.MakeWidget();
 	}
-
-	TSharedRef<SWidget> FDashboardFactory::MakeMainToolbarWidget()
-	{
-		using namespace DashboardFactoryPrivate;
-
-		static const FName PlayWorldToolBarName = "Kismet.DebuggingViewToolBar";
-		if (!UToolMenus::Get()->IsMenuRegistered(PlayWorldToolBarName))
-		{
-			UToolMenu* ToolBar = UToolMenus::Get()->RegisterMenu(PlayWorldToolBarName, NAME_None, EMultiBoxType::SlimHorizontalToolBar);
-			FToolMenuSection& Section = ToolBar->AddSection("Debug");
-			FPlayWorldCommands::BuildToolbar(Section);
-		}
-
-		static FSlateBrush TransportBackgroundColorBrush;
-		TransportBackgroundColorBrush.TintColor = FSlateColor(FLinearColor(0.018f, 0.018f, 0.018f, 1.0f));
-		TransportBackgroundColorBrush.DrawAs    = ESlateBrushDrawType::Box;
-
-		return SNew(SBorder)
-			.BorderImage(&TransportBackgroundColorBrush)
-			[
-				SNew(SHorizontalBox)
-				+ SHorizontalBox::Slot()
-				.HAlign(HAlign_Left)
-				.VAlign(VAlign_Center)
-				.AutoWidth()
-				[
-					SNew(SBorder)
-					.BorderImage(FAppStyle::Get().GetBrush("NoBorder"))
-					[
-						UToolMenus::Get()->GenerateWidget(PlayWorldToolBarName, { FPlayWorldCommands::GlobalPlayWorldActions })
-					]
-				]
-				+ SHorizontalBox::Slot()
-				.HAlign(HAlign_Right)
-				.VAlign(VAlign_Center)
-				.AutoWidth()
-				.Padding(2.0f, 0.0f)
-				[
-					SNew(STextBlock)
-					.Text(LOCTEXT("StartOnPIE_DisplayName", "Start with PIE:"))
-					.Font(IPropertyTypeCustomizationUtils::GetRegularFont())
-				]
-				+ SHorizontalBox::Slot()
-				.HAlign(HAlign_Left)
-				.VAlign(VAlign_Center)
-				.AutoWidth()
-				.Padding(2.0f, 0.0f)
-				[
-					SNew(SCheckBox)
-					.IsChecked_Lambda([this]() { return bStartWithPIE ? ECheckBoxState::Checked : ECheckBoxState::Unchecked; })
-					.OnCheckStateChanged_Lambda([this](ECheckBoxState NewState) { bStartWithPIE = NewState == ECheckBoxState::Checked; })
-				]
-				+ SHorizontalBox::Slot()
-				.HAlign(HAlign_Right)
-				.VAlign(VAlign_Center)
-				.AutoWidth()
-				.Padding(2.0f, 0.0f)
-				[
-					SNew(STextBlock)
-					.Text(LOCTEXT("StopOnPIE_DisplayName", "Stop with PIE:"))
-					.Font(IPropertyTypeCustomizationUtils::GetRegularFont())
-				]
-				+ SHorizontalBox::Slot()
-				.HAlign(HAlign_Left)
-				.VAlign(VAlign_Center)
-				.AutoWidth()
-				.Padding(2.0f, 0.0f)
-				[
-					SNew(SCheckBox)
-					.IsChecked_Lambda([this]() { return bStopWithPIE ? ECheckBoxState::Checked : ECheckBoxState::Unchecked; })
-					.OnCheckStateChanged_Lambda([this](ECheckBoxState NewState) { bStopWithPIE = NewState == ECheckBoxState::Checked; })
-				]
-				+ SHorizontalBox::Slot()
-				.HAlign(HAlign_Right)
-				.VAlign(VAlign_Center)
-				.Padding(2.0f, 0.0f)
-				[
-					SNew(STextBlock)
-					.Text(LOCTEXT("SelectDashboardWorld_DisplayName", "World Filter:"))
-					.ToolTipText(DashboardWorldSelectDescription)
-					.Font(IPropertyTypeCustomizationUtils::GetRegularFont())
-				]
-				+ SHorizontalBox::Slot()
-				.HAlign(HAlign_Left)
-				.VAlign(VAlign_Center)
-				.AutoWidth()
-				.Padding(2.0f, 0.0f)
-				[
-					SAssignNew(AudioDeviceComboBox, SComboBox<TSharedPtr<::Audio::FDeviceId>>)
-					.ToolTipText(DashboardWorldSelectDescription)
-					.OptionsSource(&AudioDeviceIds)
-					.OnGenerateWidget_Lambda([](const TSharedPtr<::Audio::FDeviceId>& WidgetDeviceId)
-					{
-						FText NameText = GetDebugNameFromDeviceId(*WidgetDeviceId);
-						return SNew(STextBlock)
-							.Text(NameText)
-							.Font(IPropertyTypeCustomizationUtils::GetRegularFont());
-					})
-					.OnSelectionChanged_Lambda([this](TSharedPtr<::Audio::FDeviceId> NewDeviceId, ESelectInfo::Type)
-					{
-						if (NewDeviceId.IsValid())
-						{
-							ActiveDeviceId = *NewDeviceId;
-							RefreshDeviceSelector();
-
-							OnActiveAudioDeviceChanged.Broadcast();
-						}
-					})
-					[
-						SNew(STextBlock)
-						.Font(IPropertyTypeCustomizationUtils::GetRegularFont())
-						.Text_Lambda([this]()
-						{
-							return GetDebugNameFromDeviceId(ActiveDeviceId);
-						})
-					]
-				]
-			];
-	}
-
-	void FDashboardFactory::InitDelegates()
-	{
-		if (!OnWorldRegisteredToAudioDeviceHandle.IsValid())
-		{
-			OnWorldRegisteredToAudioDeviceHandle = FAudioDeviceWorldDelegates::OnWorldRegisteredToAudioDevice.AddSP(this, &FDashboardFactory::OnWorldRegisteredToAudioDevice);
-		}
-
-		if (!OnWorldUnregisteredFromAudioDeviceHandle.IsValid())
-		{
-			OnWorldUnregisteredFromAudioDeviceHandle = FAudioDeviceWorldDelegates::OnWorldUnregisteredWithAudioDevice.AddSP(this, &FDashboardFactory::OnWorldUnregisteredFromAudioDevice);
-		}
-
-		if (!OnDeviceDestroyedHandle.IsValid())
-		{
-			OnDeviceDestroyedHandle = FAudioDeviceManagerDelegates::OnAudioDeviceDestroyed.AddSP(this, &FDashboardFactory::OnDeviceDestroyed);
-		}
-
-		if (!OnPIEStartedHandle.IsValid())
-		{
-			OnPIEStartedHandle = FEditorDelegates::PreBeginPIE.AddSP(this, &FDashboardFactory::OnPIEStarted);
-		}
-
-		if (!OnPostPIEStartedHandle.IsValid())
-		{
-			OnPostPIEStartedHandle = FEditorDelegates::PostPIEStarted.AddSP(this, &FDashboardFactory::OnPostPIEStarted);
-		}
-
-		if (!OnPIEStoppedHandle.IsValid())
-		{
-			OnPIEStoppedHandle = FEditorDelegates::EndPIE.AddSP(this, &FDashboardFactory::OnPIEStopped);
-		}
-	}
-
+	
 	TSharedPtr<FTabManager::FLayout> FDashboardFactory::GetDefaultTabLayout()
 	{
 		using namespace DashboardFactoryPrivate;
 
-		TSharedRef<FTabManager::FStack> ViewportTabStack = FTabManager::NewStack();
 		TSharedRef<FTabManager::FStack> LogTabStack = FTabManager::NewStack();
 		TSharedRef<FTabManager::FStack> AnalysisTabStack = FTabManager::NewStack();
-		TSharedRef<FTabManager::FStack> AudioMetersTabStack = FTabManager::NewStack();
-		TSharedRef<FTabManager::FStack> AudioMeterTabStack = FTabManager::NewStack();
-		TSharedRef<FTabManager::FStack> OscilloscopeTabStack = FTabManager::NewStack();
 
-		for (const auto& KVP : DashboardViewFactories)
+		for (const auto& [FactoryName, Factory] : DashboardViewFactories)
 		{
-			const FName& FactoryName = KVP.Key;
-			const TSharedPtr<IDashboardViewFactory>& Factory = KVP.Value;
-
 			const EDefaultDashboardTabStack DefaultTabStack = Factory->GetDefaultTabStack();
+
 			switch (DefaultTabStack)
 			{
-				case EDefaultDashboardTabStack::Viewport:
-				{
-					ViewportTabStack->AddTab(FactoryName, ETabState::OpenedTab);
-				}
-				break;
-
 				case EDefaultDashboardTabStack::Log:
 				{
 					LogTabStack->AddTab(FactoryName, ETabState::OpenedTab);
@@ -522,25 +177,7 @@ namespace UE::Audio::Insights
 					AnalysisTabStack->AddTab(FactoryName, ETabState::OpenedTab);
 				}
 				break;
-
-				case EDefaultDashboardTabStack::AudioMeters:
-				{
-					AudioMetersTabStack->AddTab(FactoryName, ETabState::OpenedTab);
-				}
-				break;
-
-				case EDefaultDashboardTabStack::AudioMeter:
-				{
-					AudioMeterTabStack->AddTab(FactoryName, ETabState::OpenedTab);
-				}
-				break;
-
-				case EDefaultDashboardTabStack::Oscilloscope:
-				{
-					OscilloscopeTabStack->AddTab(FactoryName, ETabState::OpenedTab);
-				}
-				break;
-
+				
 				default:
 					break;
 			}
@@ -555,72 +192,17 @@ namespace UE::Audio::Insights
 			->SetOrientation(Orient_Vertical)
 			->Split
 			(
-				// Left column
 				FTabManager::NewSplitter()
 				->SetOrientation(Orient_Horizontal)
 				->Split
 				(
-					// Top
-					FTabManager::NewSplitter()
-					->SetOrientation(Orient_Vertical)
-					->SetSizeCoefficient(0.25f) // Column width
-					->Split
-					(
-						ViewportTabStack
-						->SetSizeCoefficient(0.5f)
-					)
-					// Bottom
-					->Split
-					(
-						LogTabStack
-						->SetSizeCoefficient(0.5f)
-					)
-				)
-				
-				// Middle column
-				->Split
-				(
-					// Top
-					FTabManager::NewSplitter()
-					->SetOrientation(Orient_Vertical)
-					->SetSizeCoefficient(0.6f) // Column width
-					->Split
-					(
-						FTabManager::NewSplitter()
-						->SetOrientation(Orient_Horizontal)
-						->Split
-						(
-							AnalysisTabStack
-							->SetSizeCoefficient(0.58f)
-						)
-					)
-					// Bottom
-					->Split
-					(
-						AudioMetersTabStack
-						->SetSizeCoefficient(0.42f)
-					)
+					LogTabStack
+					->SetSizeCoefficient(0.25f)
 				)
-				// Right column
 				->Split
 				(
-					// Top
-					FTabManager::NewSplitter()
-					->SetOrientation(Orient_Vertical)
-					->SetSizeCoefficient(0.15f) // Column width
-					->Split
-					(
-						AudioMeterTabStack
-						->SetSizeCoefficient(0.7f)
-						->SetHideTabWell(true)
-					)
-					// Bottom
-					->Split
-					(
-						OscilloscopeTabStack
-						->SetSizeCoefficient(0.3f)
-						->SetHideTabWell(true)
-					)
+					AnalysisTabStack
+					->SetSizeCoefficient(0.75f)
 				)
 			)
 		);
diff --git a/Engine/Plugins/AudioInsights/Source/AudioInsights/Private/AudioInsightsDashboardFactory.h b/Engine/Plugins/AudioInsights/Source/AudioInsights/Private/AudioInsightsDashboardFactory.h
deleted file mode 100644
diff --git a/Engine/Plugins/AudioInsights/Source/AudioInsights/Private/AudioInsightsModule.cpp b/Engine/Plugins/AudioInsights/Source/AudioInsights/Private/AudioInsightsModule.cpp
--- a/Engine/Plugins/AudioInsights/Source/AudioInsights/Private/AudioInsightsModule.cpp
+++ b/Engine/Plugins/AudioInsights/Source/AudioInsights/Private/AudioInsightsModule.cpp
@@ -4,25 +4,20 @@
 #include "AudioInsightsDashboardAssetCommands.h"
 #include "AudioInsightsDashboardFactory.h"
 #include "AudioInsightsLog.h"
-#include "AudioInsightsStyle.h"
 #include "AudioInsightsTraceModule.h"
 #include "Features/IModularFeatures.h"
 #include "Framework/Docking/TabManager.h"
+#include "Insights/IUnrealInsightsModule.h"
 #include "Modules/ModuleManager.h"
 #include "Templates/SharedPointer.h"
 #include "TraceServices/ModuleService.h"
 #include "UObject/NameTypes.h"
-#include "Views/AudioBusesDashboardViewFactory.h"
-#include "Views/AudioMetersDashboardViewFactory.h"
+
+#if !WITH_EDITOR
 #include "Views/LogDashboardViewFactory.h"
 #include "Views/MixerSourceDashboardViewFactory.h"
-#include "Views/OutputMeterDashboardViewFactory.h"
-#include "Views/OutputOscilloscopeDashboardViewFactory.h"
-#include "Views/SubmixesDashboardViewFactory.h"
-#include "Views/ViewportDashboardViewFactory.h"
 #include "Views/VirtualLoopDashboardViewFactory.h"
-#include "WorkspaceMenuStructure.h"
-#include "WorkspaceMenuStructureModule.h"
+#endif // !WITH_EDITOR
 
 #define LOCTEXT_NAMESPACE "AudioInsights"
 DEFINE_LOG_CATEGORY(LogAudioInsights);
@@ -37,20 +32,27 @@ namespace UE::Audio::Insights
 		{
 			IModularFeatures::Get().RegisterModularFeature(TraceServices::ModuleFeatureName, &TraceModule);
 
-			RegisterMenus();
-
 			DashboardFactory = MakeShared<FDashboardFactory>();
-			DashboardFactory->RegisterViewFactory(MakeShared<FViewportDashboardViewFactory>());
+			
+			FDashboardAssetCommands::Register();
+
+#if !WITH_EDITOR
 			DashboardFactory->RegisterViewFactory(MakeShared<FLogDashboardViewFactory>());
 			DashboardFactory->RegisterViewFactory(MakeShared<FMixerSourceDashboardViewFactory>());
 			DashboardFactory->RegisterViewFactory(MakeShared<FVirtualLoopDashboardViewFactory>());
-			DashboardFactory->RegisterViewFactory(MakeShared<FSubmixesDashboardViewFactory>());
-			DashboardFactory->RegisterViewFactory(MakeShared<FAudioBusesDashboardViewFactory>());
-			DashboardFactory->RegisterViewFactory(MakeShared<FAudioMetersDashboardViewFactory>());
-			DashboardFactory->RegisterViewFactory(MakeShared<FOutputMeterDashboardViewFactory>());
-			DashboardFactory->RegisterViewFactory(MakeShared<FOutputOscilloscopeDashboardViewFactory>());
-
-			FDashboardAssetCommands::Register();
+#endif // !WITH_EDITOR
+
+			FCoreDelegates::OnFEngineLoopInitComplete.AddLambda([this]
+			{
+				LLM_SCOPE_BYNAME(TEXT("Insights/AudioInsights"));
+				IUnrealInsightsModule& UnrealInsightsModule = FModuleManager::LoadModuleChecked<IUnrealInsightsModule>("TraceInsights");
+				if (!UnrealInsightsModule.GetStoreClient())
+				{
+					UE_LOG(LogCore, Display, TEXT("AudioInsights module auto-connecting to local trace server..."));
+					UnrealInsightsModule.ConnectToStore(TEXT("127.0.0.1"));
+					UnrealInsightsModule.CreateSessionViewer(false);
+				}
+			});
 		}
 	}
 
@@ -104,16 +106,6 @@ namespace UE::Audio::Insights
 	{
 		return DashboardFactory->MakeDockTabWidget(Args);
 	}
-
-	void FAudioInsightsModule::RegisterMenus()
-	{
-		const IWorkspaceMenuStructure& MenuStructure = WorkspaceMenu::GetMenuStructure();
-		FGlobalTabmanager::Get()->RegisterNomadTabSpawner("AudioInsights", FOnSpawnTab::CreateRaw(this, &FAudioInsightsModule::CreateDashboardTabWidget))
-			.SetDisplayName(LOCTEXT("OpenDashboard_TabDisplayName", "Audio Insights"))
-			.SetTooltipText(LOCTEXT("OpenDashboard_TabTooltip", "Opens Audio Insights, an extensible suite of tools and visualizers which enable monitoring and debugging audio in the Unreal Engine."))
-			.SetGroup(MenuStructure.GetToolsCategory())
-			.SetIcon(FSlateStyle::Get().CreateIcon("AudioInsights.Icon.Dashboard"));
-	};
 } // namespace UE::Audio::Insights
 #undef LOCTEXT_NAMESPACE // AudioInsights
 
diff --git a/Engine/Plugins/AudioInsights/Source/AudioInsights/Private/AudioInsightsModule.h b/Engine/Plugins/AudioInsights/Source/AudioInsights/Private/AudioInsightsModule.h
deleted file mode 100644
diff --git a/Engine/Plugins/AudioInsights/Source/AudioInsights/Private/AudioInsightsStyle.h b/Engine/Plugins/AudioInsights/Source/AudioInsights/Private/AudioInsightsStyle.h
deleted file mode 100644
diff --git a/Engine/Plugins/AudioInsights/Source/AudioInsights/Private/AudioInsightsTraceModule.h b/Engine/Plugins/AudioInsights/Source/AudioInsights/Private/AudioInsightsTraceModule.h
deleted file mode 100644
diff --git a/Engine/Plugins/AudioInsights/Source/AudioInsights/Private/Messages/MixerSourceTraceMessages.h b/Engine/Plugins/AudioInsights/Source/AudioInsights/Private/Messages/MixerSourceTraceMessages.h
deleted file mode 100644
diff --git a/Engine/Plugins/AudioInsights/Source/AudioInsights/Private/Messages/VirtualLoopTraceMessages.h b/Engine/Plugins/AudioInsights/Source/AudioInsights/Private/Messages/VirtualLoopTraceMessages.h
deleted file mode 100644
diff --git a/Engine/Plugins/AudioInsights/Source/AudioInsights/Private/Providers/AudioBusProvider.cpp b/Engine/Plugins/AudioInsights/Source/AudioInsights/Private/Providers/AudioBusProvider.cpp
deleted file mode 100644
diff --git a/Engine/Plugins/AudioInsights/Source/AudioInsights/Private/Providers/AudioBusProvider.h b/Engine/Plugins/AudioInsights/Source/AudioInsights/Private/Providers/AudioBusProvider.h
deleted file mode 100644
diff --git a/Engine/Plugins/AudioInsights/Source/AudioInsights/Private/Providers/MixerSourceTraceProvider.cpp b/Engine/Plugins/AudioInsights/Source/AudioInsights/Private/Providers/MixerSourceTraceProvider.cpp
--- a/Engine/Plugins/AudioInsights/Source/AudioInsights/Private/Providers/MixerSourceTraceProvider.cpp
+++ b/Engine/Plugins/AudioInsights/Source/AudioInsights/Private/Providers/MixerSourceTraceProvider.cpp
@@ -159,18 +159,25 @@ namespace UE::Audio::Insights
 				Builder.RouteEvent(RouteId_Start, "Audio", "MixerSourceStart");
 				Builder.RouteEvent(RouteId_Stop, "Audio", "MixerSourceStop");
 				Builder.RouteEvent(RouteId_Volume, "Audio", "MixerSourceVolume");
+
+				bIsFirstTimestampSet = false;
 			}
 
 			virtual bool OnEvent(uint16 RouteId, EStyle Style, const FOnEventContext& Context) override
 			{
 				LLM_SCOPE_BYNAME(TEXT("Insights/FMixerSourceTraceAnalyzer"));
 
-				FTraceModule& TraceModule = FAudioInsightsModule::GetChecked().GetTraceModule();
-
-				const double CurrentTime = FPlatformTime::Seconds();
-				if (TraceModule.GetFirstTimeStamp() < 0.0)
+				if (!bIsFirstTimestampSet)
 				{
-					TraceModule.SetFirstTimeStamp(CurrentTime);
+					FTraceModule& TraceModule = FAudioInsightsModule::GetChecked().GetTraceModule();
+
+					if (TraceModule.GetFirstTimeStamp() < 0.0)
+					{
+						const double CurrentTime = FPlatformTime::Seconds();
+						TraceModule.SetFirstTimeStamp(CurrentTime);
+
+						bIsFirstTimestampSet = true;
+					}
 				}
 
 				FMixerSourceMessages& Messages = GetProvider<FMixerSourceTraceProvider>().TraceMessages;
@@ -239,6 +246,8 @@ namespace UE::Audio::Insights
 				RouteId_Stop,
 				RouteId_Volume,
 			};
+
+			bool bIsFirstTimestampSet = false;
 		};
 
 		return new FMixerSourceTraceAnalyzer(AsShared());
diff --git a/Engine/Plugins/AudioInsights/Source/AudioInsights/Private/Providers/SoundSubmixProvider.cpp b/Engine/Plugins/AudioInsights/Source/AudioInsights/Private/Providers/SoundSubmixProvider.cpp
deleted file mode 100644
diff --git a/Engine/Plugins/AudioInsights/Source/AudioInsights/Private/Providers/SoundSubmixProvider.h b/Engine/Plugins/AudioInsights/Source/AudioInsights/Private/Providers/SoundSubmixProvider.h
deleted file mode 100644
diff --git a/Engine/Plugins/AudioInsights/Source/AudioInsights/Private/Views/AudioBusAssetDashboardEntry.cpp b/Engine/Plugins/AudioInsights/Source/AudioInsights/Private/Views/AudioBusAssetDashboardEntry.cpp
deleted file mode 100644
diff --git a/Engine/Plugins/AudioInsights/Source/AudioInsights/Private/Views/AudioBusAssetDashboardEntry.h b/Engine/Plugins/AudioInsights/Source/AudioInsights/Private/Views/AudioBusAssetDashboardEntry.h
deleted file mode 100644
diff --git a/Engine/Plugins/AudioInsights/Source/AudioInsights/Private/Views/AudioBusesDashboardViewFactory.cpp b/Engine/Plugins/AudioInsights/Source/AudioInsights/Private/Views/AudioBusesDashboardViewFactory.cpp
deleted file mode 100644
diff --git a/Engine/Plugins/AudioInsights/Source/AudioInsights/Private/Views/AudioBusesDashboardViewFactory.h b/Engine/Plugins/AudioInsights/Source/AudioInsights/Private/Views/AudioBusesDashboardViewFactory.h
deleted file mode 100644
diff --git a/Engine/Plugins/AudioInsights/Source/AudioInsights/Private/Views/AudioMeterView.cpp b/Engine/Plugins/AudioInsights/Source/AudioInsights/Private/Views/AudioMeterView.cpp
deleted file mode 100644
diff --git a/Engine/Plugins/AudioInsights/Source/AudioInsights/Private/Views/AudioMeterView.h b/Engine/Plugins/AudioInsights/Source/AudioInsights/Private/Views/AudioMeterView.h
deleted file mode 100644
diff --git a/Engine/Plugins/AudioInsights/Source/AudioInsights/Private/Views/AudioMetersDashboardViewFactory.cpp b/Engine/Plugins/AudioInsights/Source/AudioInsights/Private/Views/AudioMetersDashboardViewFactory.cpp
deleted file mode 100644
diff --git a/Engine/Plugins/AudioInsights/Source/AudioInsights/Private/Views/AudioMetersDashboardViewFactory.h b/Engine/Plugins/AudioInsights/Source/AudioInsights/Private/Views/AudioMetersDashboardViewFactory.h
deleted file mode 100644
diff --git a/Engine/Plugins/AudioInsights/Source/AudioInsights/Private/Views/LogDashboardViewFactory.h b/Engine/Plugins/AudioInsights/Source/AudioInsights/Private/Views/LogDashboardViewFactory.h
deleted file mode 100644
diff --git a/Engine/Plugins/AudioInsights/Source/AudioInsights/Private/Views/MixerSourceDashboardViewFactory.cpp b/Engine/Plugins/AudioInsights/Source/AudioInsights/Private/Views/MixerSourceDashboardViewFactory.cpp
--- a/Engine/Plugins/AudioInsights/Source/AudioInsights/Private/Views/MixerSourceDashboardViewFactory.cpp
+++ b/Engine/Plugins/AudioInsights/Source/AudioInsights/Private/Views/MixerSourceDashboardViewFactory.cpp
@@ -1,21 +1,21 @@
 // Copyright Epic Games, Inc. All Rights Reserved.
 #include "Views/MixerSourceDashboardViewFactory.h"
 
-#include "Audio/AudioDebug.h"
-#include "AudioDeviceManager.h"
-#include "AudioInsightsDashboardAssetCommands.h"
 #include "AudioInsightsModule.h"
 #include "AudioInsightsStyle.h"
 #include "DSP/Dsp.h"
-#include "Editor.h"
 #include "Internationalization/Text.h"
-#include "IPropertyTypeCustomization.h"
 #include "Providers/MixerSourceTraceProvider.h"
 #include "SSimpleTimeSlider.h"
 #include "Templates/SharedPointer.h"
 #include "UObject/SoftObjectPath.h"
+#include "Widgets/Input/SComboBox.h"
 #include "Widgets/Layout/SSplitter.h"
 
+#if WITH_EDITOR
+#include "Editor.h"
+#endif // WITH_EDITOR
+
 #define LOCTEXT_NAMESPACE "AudioInsights"
 
 
@@ -386,6 +386,7 @@ namespace UE::Audio::Insights
 		CurrentTimestamp = 0;
 	}
 
+#if WITH_EDITOR
 	void FMixerSourceDashboardViewFactory::OnPIEStarted(bool bSimulating)
 	{
 		PIEState = EPIEState::Running;
@@ -407,217 +408,36 @@ namespace UE::Audio::Insights
 	{
 		PIEState = EPIEState::Running;
 	}
+#endif // WITH_EDITOR
 
-#if AUDIO_INSIGHTS_SHOW_SOURCE_CONTEXT_MENU
-	TSharedPtr<SWidget> FMixerSourceDashboardViewFactory::OnConstructContextMenu()
-	{
-		const FDashboardAssetCommands& Commands = FDashboardAssetCommands::Get();
-
-		TSharedPtr<FUICommandList> CommandList = MakeShared<FUICommandList>();
-		CommandList->MapAction(Commands.GetMuteCommand(), FExecuteAction::CreateRaw(this, &FMixerSourceDashboardViewFactory::MuteSound));
-		CommandList->MapAction(Commands.GetSoloCommand(), FExecuteAction::CreateRaw(this, &FMixerSourceDashboardViewFactory::SoloSound));
-		CommandList->MapAction(Commands.GetClearMuteSoloCommand(), FExecuteAction::CreateRaw(this, &FMixerSourceDashboardViewFactory::ClearMutesAndSolos));
-
-		constexpr bool bShouldCloseWindowAfterMenuSelection = true;
-		FMenuBuilder MenuBuilder(bShouldCloseWindowAfterMenuSelection, CommandList);
-
-		MenuBuilder.BeginSection("SoundActions", LOCTEXT("SoundActions_Header", "Sound Actions"));
-		{
-			MenuBuilder.AddMenuEntry(Commands.GetMuteCommand());
-			MenuBuilder.AddMenuEntry(Commands.GetSoloCommand());
-			MenuBuilder.AddMenuEntry(Commands.GetClearMuteSoloCommand());
-		}
-		MenuBuilder.EndSection();
-
-		return MenuBuilder.MakeWidget();
-	}
-#endif // AUDIO_INSIGHTS_SHOW_SOURCE_CONTEXT_MENU 
-
-	FSlateColor FMixerSourceDashboardViewFactory::GetRowColor(const TSharedPtr<IDashboardDataViewEntry>& InRowDataPtr)
-	{
-		FColor RowTextColor(255, 255, 255);
-
-#if ENABLE_AUDIO_DEBUG
-		if (const FAudioDeviceManager* AudioDeviceManager = FAudioDeviceManager::Get())
-		{
-			const FSoundAssetDashboardEntry& SoundAssetDashboardEntry = *StaticCastSharedPtr<FSoundAssetDashboardEntry>(InRowDataPtr).Get();
-			const FName SoundAssetName { SoundAssetDashboardEntry.Name };
-			const bool bIsSolo = AudioDeviceManager->GetDebugger().IsSoloSoundWave(SoundAssetName);
-			if (bIsSolo)
-			{
-				RowTextColor = FColor(255, 255, 0);
-			}
-			else
-			{
-				const bool bIsMute = AudioDeviceManager->GetDebugger().IsMuteSoundWave(SoundAssetName);
-				if (bIsMute)
-				{
-					RowTextColor = FColor(255, 0, 0);
-				}
-			}
-		}
-#endif // ENABLE_AUDIO_DEBUG
-
-		return FSlateColor(RowTextColor);
-	}
-
+#if WITH_EDITOR
 	void FMixerSourceDashboardViewFactory::ToggleMuteForAllItems(ECheckBoxState NewState)
 	{
-#if ENABLE_AUDIO_DEBUG
 		if (MuteState != NewState)
 		{
 			MuteState = NewState;
-			UpdateSoloMuteState();
+			UpdateMuteSoloState();
 		}
-#endif
 	}
 
 	void FMixerSourceDashboardViewFactory::ToggleSoloForAllItems(ECheckBoxState NewState)
 	{
-#if ENABLE_AUDIO_DEBUG
 		if (SoloState != NewState)
 		{
 			SoloState = NewState;
-			UpdateSoloMuteState();
+			UpdateMuteSoloState();
 		}
-#endif
 	}
 
-	void FMixerSourceDashboardViewFactory::MuteSound()
+	void FMixerSourceDashboardViewFactory::UpdateMuteSoloState()
 	{
-#if ENABLE_AUDIO_DEBUG
-		if (FilteredEntriesListView.IsValid())
-		{
-			if (FAudioDeviceManager* AudioDeviceManager = FAudioDeviceManager::Get())
-			{
-				const TArray<TSharedPtr<IDashboardDataViewEntry>> SelectedItems = FilteredEntriesListView->GetSelectedItems();
-
-				for (const TSharedPtr<IDashboardDataViewEntry>& SelectedItem : SelectedItems)
-				{
-					if (SelectedItem.IsValid())
-					{
-						const FSoundAssetDashboardEntry& SoundAssetDashboardEntry = *StaticCastSharedPtr<FSoundAssetDashboardEntry>(SelectedItem).Get();
-						const FName SoundAssetDisplayName { SoundAssetDashboardEntry.Name };
-						AudioDeviceManager->GetDebugger().ToggleMuteSoundWave(SoundAssetDisplayName);
-					}
-				}
-
-				// Handle general Mute button state
-				bool bIsAnySoundMuted = false;
-
-				const TArrayView<const TSharedPtr<IDashboardDataViewEntry>> TableItems = FilteredEntriesListView->GetItems();
-
-				for (const TSharedPtr<IDashboardDataViewEntry>& Item : TableItems)
-				{
-					if (!Item.IsValid())
-					{
-						continue;
-					}
-
-					const FSoundAssetDashboardEntry& SoundAssetDashboardEntry = *StaticCastSharedPtr<FSoundAssetDashboardEntry>(Item).Get();
-					const FName SoundAssetDisplayName { SoundAssetDashboardEntry.Name };
-					if (AudioDeviceManager->GetDebugger().IsMuteSoundWave(SoundAssetDisplayName))
-					{
-						bIsAnySoundMuted = true;
-						break;
-					}
-				}
-
-				MuteToggleButton->SetIsChecked(bIsAnySoundMuted ? ECheckBoxState::Checked : ECheckBoxState::Unchecked);
-			}
-		}
-#endif
-	}
-
-	void FMixerSourceDashboardViewFactory::SoloSound()
-	{
-#if ENABLE_AUDIO_DEBUG
-		if (FilteredEntriesListView.IsValid())
-		{
-			if (FAudioDeviceManager* AudioDeviceManager = FAudioDeviceManager::Get())
-			{
-				const TArray<TSharedPtr<IDashboardDataViewEntry>> SelectedItems = FilteredEntriesListView->GetSelectedItems();
-
-				for (const TSharedPtr<IDashboardDataViewEntry>& SelectedItem : SelectedItems)
-				{
-					if (SelectedItem.IsValid())
-					{
-						const FSoundAssetDashboardEntry& SoundAssetDashboardEntry = *StaticCastSharedPtr<FSoundAssetDashboardEntry>(SelectedItem).Get();
-						const FName SoundAssetDisplayName { SoundAssetDashboardEntry.Name };
-						AudioDeviceManager->GetDebugger().ToggleSoloSoundWave(SoundAssetDisplayName);
-					}
-				}
-
-				// Handle general Solo button state
-				bool bIsAnySoundSoloed = false;
-
-				const TArrayView<const TSharedPtr<IDashboardDataViewEntry>> TableItems = FilteredEntriesListView->GetItems();
-
-				for (const TSharedPtr<IDashboardDataViewEntry>& Item : TableItems)
-				{
-					if (!Item.IsValid())
-					{
-						continue;
-					}
-
-					const FSoundAssetDashboardEntry& SoundAssetDashboardEntry = *StaticCastSharedPtr<FSoundAssetDashboardEntry>(Item).Get();
-					const FName SoundAssetDisplayName { SoundAssetDashboardEntry.Name };
-					if (AudioDeviceManager->GetDebugger().IsSoloSoundWave(SoundAssetDisplayName))
-					{
-						bIsAnySoundSoloed = true;
-						break;
-					}
-				}
-
-				SoloToggleButton->SetIsChecked(bIsAnySoundSoloed ? ECheckBoxState::Checked : ECheckBoxState::Unchecked);
-			}
-		}
-#endif
-	}
-
-	void FMixerSourceDashboardViewFactory::ClearMutesAndSolos()
-	{
-#if ENABLE_AUDIO_DEBUG
-		if (FAudioDeviceManager* AudioDeviceManager = FAudioDeviceManager::Get())
-		{
-			AudioDeviceManager->GetDebugger().ClearMutesAndSolos();
-
-			MuteToggleButton->SetIsChecked(ECheckBoxState::Unchecked);
-			SoloToggleButton->SetIsChecked(ECheckBoxState::Unchecked);
-		}
-#endif
-	}
-
-	void FMixerSourceDashboardViewFactory::UpdateSoloMuteState()
-	{
-#if ENABLE_AUDIO_DEBUG
-		if (FAudioDeviceManager* AudioDeviceManager = FAudioDeviceManager::Get())
-		{
-			FName CurrentFilterStringName = FName{ CurrentFilterString };
-			if (MuteState == ECheckBoxState::Checked && !CurrentFilterString.IsEmpty())
-			{
-				AudioDeviceManager->GetDebugger().ToggleMuteSoundWave(CurrentFilterStringName, true);
-			}
-			else
-			{
-				AudioDeviceManager->GetDebugger().ToggleMuteSoundWave(NAME_None, true);
-			}
-
-			if (SoloState == ECheckBoxState::Checked && !CurrentFilterString.IsEmpty())
-			{
-				AudioDeviceManager->GetDebugger().ToggleSoloSoundWave(CurrentFilterStringName, true);
-			}
-			else
-			{
-				AudioDeviceManager->GetDebugger().ToggleSoloSoundWave(NAME_None, true);
-			}
-		}
-#endif
+		OnUpdateMuteSoloState.Broadcast(MuteState, SoloState, CurrentFilterString);
 	}
+#endif // WITH_EDITOR
 
 	void FMixerSourceDashboardViewFactory::UpdatePlotsWidgetsData()
 	{
-		if (DataViewEntries.Num() <= 0)
+		if (!PlotWidgetMetadataPerCurve.IsValid() || DataViewEntries.Num() <= 0)
 		{
 			return;
 		}
@@ -755,7 +575,7 @@ namespace UE::Audio::Insights
 					"LPF",
 					{
 						[](const IDashboardDataViewEntry& InData) -> const ::Audio::TCircularAudioBuffer<FDataPoint>& { return MixerSourcePrivate::CastEntry(InData).LPFFreqDataPoints; },
-							FSlateStyle::Get().GetFreqFloatFormat() 
+						FSlateStyle::Get().GetFreqFloatFormat()
 					}
 				},
 				{
@@ -858,7 +678,7 @@ namespace UE::Audio::Insights
 				.OnGenerateWidget_Lambda([this](const FName& ColumnName)
 				{
 					return SNew(STextBlock)
-					.Font(IPropertyTypeCustomizationUtils::GetRegularFont())
+					.Font(FAppStyle::GetFontStyle(TEXT("PropertyWindow.NormalFont")))
 					.Text(GetPlotColumnDisplayName(ColumnName));
 				})
 				.OnSelectionChanged_Lambda([this, PlotWidgetIndex](FName NewColumnName, ESelectInfo::Type)
@@ -872,7 +692,7 @@ namespace UE::Audio::Insights
 				})
 				[
 					SNew(STextBlock)
-					.Font(IPropertyTypeCustomizationUtils::GetRegularFont())
+					.Font(FAppStyle::GetFontStyle(TEXT("PropertyWindow.NormalFont")))
 					.Text_Lambda([this, PlotWidgetIndex]()
 					{
 						return GetPlotColumnDisplayName(SelectedPlotColumnNames[PlotWidgetIndex]);
@@ -905,6 +725,7 @@ namespace UE::Audio::Insights
 			];
 	}
 
+#if WITH_EDITOR
 	TSharedRef<SWidget> FMixerSourceDashboardViewFactory::MakeMuteSoloWidget()
 	{
 		// Mute/Solo labels generation
@@ -1002,30 +823,32 @@ namespace UE::Audio::Insights
 				]
 			];
 	}
+#endif // WITH_EDITOR
 
 	TSharedRef<SWidget> FMixerSourceDashboardViewFactory::MakeWidget()
 	{
-		FDashboardFactory::OnActiveAudioDeviceChanged.AddSP(this, &FMixerSourceDashboardViewFactory::ClearMutesAndSolos);
-
+#if WITH_EDITOR
 		FEditorDelegates::PostPIEStarted.AddSP(this, &FMixerSourceDashboardViewFactory::OnPIEStarted);
 		FEditorDelegates::EndPIE.AddSP(this, &FMixerSourceDashboardViewFactory::OnPIEStopped);
 		FEditorDelegates::PausePIE.AddSP(this, &FMixerSourceDashboardViewFactory::OnPIEPaused);
 		FEditorDelegates::ResumePIE.AddSP(this, &FMixerSourceDashboardViewFactory::OnPIEResumed);
+#endif // WITH_EDITOR
 
-		TSharedRef<SWidget> MuteSoloWidget = MakeMuteSoloWidget();
 		TSharedRef<SWidget> TableDashboardWidget = FTraceTableDashboardViewFactory::MakeWidget();
 		TSharedRef<SWidget> PlotsWidget = MakePlotsWidget();
 
 		return SNew(SVerticalBox)
+#if WITH_EDITOR
 			+ SVerticalBox::Slot()
 			.AutoHeight()
 			.HAlign(HAlign_Fill)
+			.Padding(0.0f, 0.0f, 0.0f, 6.0f)
 			[
-				MuteSoloWidget
+				MakeMuteSoloWidget()
 			]
+#endif // WITH_EDITOR
 			+ SVerticalBox::Slot()
 			.HAlign(HAlign_Fill)
-			.Padding(0.0f, 6.0f, 0.0f, 0.0f)
 			[
 				// Dashboard and plots area
 				SNew(SSplitter)
@@ -1059,37 +882,15 @@ namespace UE::Audio::Insights
 
 		UpdatePlotsWidgetsData();
 
-#if ENABLE_AUDIO_DEBUG
+#if WITH_EDITOR
 		// Update the mute and solo states if the filter string changes
 		if (CurrentFilterString != FilterString)
 		{
 			CurrentFilterString = FilterString;
-			UpdateSoloMuteState();
+			UpdateMuteSoloState();
 		}
-#endif
-	}
-
-#if WITH_EDITOR
-	bool FMixerSourceDashboardViewFactory::IsDebugDrawEnabled() const
-	{
-		return false;
-	}
-
-	void FMixerSourceDashboardViewFactory::DebugDraw(float InElapsed, const IDashboardDataViewEntry& InEntry, ::Audio::FDeviceId DeviceId) const
-	{
-		// TODO: Get source position if 3d so debug draw works
-// 		const FMixerSourceDashboardEntry& LoopData = static_cast<const FMixerSourceDashboardEntry&>(InEntry);
-// 		const FRotator& Rotator = LoopData.Rotator;
-// 		const FVector& Location = LoopData.Location;
-// 		const FString Description = FString::Printf(TEXT("%s [Virt: %.2fs]"), *LoopData.Name, LoopData.TimeVirtualized);
-// 
-// 		const TArray<UWorld*> Worlds = FAudioDeviceManager::Get()->GetWorldsUsingAudioDevice(DeviceId);
-// 		for (UWorld* World : Worlds)
-// 		{
-// 			DrawDebugSphere(World, Location, 30.0f, 8, FColor::Magenta, false, InElapsed, SDPG_Foreground);
-// 			DrawDebugString(World, Location + FVector(0, 0, 32), *Description, nullptr, FColor::Magenta, InElapsed, false, 1.0f);
-// 		}
-	}
 #endif // WITH_EDITOR
+	}
 } // namespace UE::Audio::Insights
+
 #undef LOCTEXT_NAMESPACE
diff --git a/Engine/Plugins/AudioInsights/Source/AudioInsights/Private/Views/MixerSourceDashboardViewFactory.h b/Engine/Plugins/AudioInsights/Source/AudioInsights/Private/Views/MixerSourceDashboardViewFactory.h
deleted file mode 100644
diff --git a/Engine/Plugins/AudioInsights/Source/AudioInsights/Private/Views/OutputMeterDashboardViewFactory.cpp b/Engine/Plugins/AudioInsights/Source/AudioInsights/Private/Views/OutputMeterDashboardViewFactory.cpp
deleted file mode 100644
diff --git a/Engine/Plugins/AudioInsights/Source/AudioInsights/Private/Views/OutputMeterDashboardViewFactory.h b/Engine/Plugins/AudioInsights/Source/AudioInsights/Private/Views/OutputMeterDashboardViewFactory.h
deleted file mode 100644
diff --git a/Engine/Plugins/AudioInsights/Source/AudioInsights/Private/Views/OutputOscilloscopeDashboardViewFactory.cpp b/Engine/Plugins/AudioInsights/Source/AudioInsights/Private/Views/OutputOscilloscopeDashboardViewFactory.cpp
deleted file mode 100644
diff --git a/Engine/Plugins/AudioInsights/Source/AudioInsights/Private/Views/OutputOscilloscopeDashboardViewFactory.h b/Engine/Plugins/AudioInsights/Source/AudioInsights/Private/Views/OutputOscilloscopeDashboardViewFactory.h
deleted file mode 100644
diff --git a/Engine/Plugins/AudioInsights/Source/AudioInsights/Private/Views/SAudioCurveView.cpp b/Engine/Plugins/AudioInsights/Source/AudioInsights/Private/Views/SAudioCurveView.cpp
--- a/Engine/Plugins/AudioInsights/Source/AudioInsights/Private/Views/SAudioCurveView.cpp
+++ b/Engine/Plugins/AudioInsights/Source/AudioInsights/Private/Views/SAudioCurveView.cpp
@@ -1,9 +1,10 @@
 // Copyright Epic Games, Inc. All Rights Reserved.
 
-#include "SAudioCurveView.h"
+#include "Views/SAudioCurveView.h"
+
 #include "Algo/MaxElement.h"
 #include "Algo/MinElement.h"
-#include "Audio.h"
+#include "AudioInsightsLog.h"
 #include "Fonts/FontMeasure.h"
 #include "Framework/Application/SlateApplication.h"
 #include "Rendering/DrawElements.h"
@@ -316,7 +317,7 @@ int32 SAudioCurveView::PaintCurves(const FGeometry& AllottedGeometry, const FSla
 
 	if (PointDataPerCurve->Num() != MetadataPerCurve->Num())
 	{
-		UE_LOG(LogAudio, Warning, TEXT("Invalid audio curve view data. Metadata and point curve data nums do not match."))
+		UE_LOG(LogAudioInsights, Warning, TEXT("Invalid audio curve view data. Metadata and point curve data nums do not match."))
 		return LayerId;
 	}
 
diff --git a/Engine/Plugins/AudioInsights/Source/AudioInsights/Private/Views/SAudioCurveView.h b/Engine/Plugins/AudioInsights/Source/AudioInsights/Private/Views/SAudioCurveView.h
deleted file mode 100644
diff --git a/Engine/Plugins/AudioInsights/Source/AudioInsights/Private/Views/SoundSubmixAssetDashboardEntry.cpp b/Engine/Plugins/AudioInsights/Source/AudioInsights/Private/Views/SoundSubmixAssetDashboardEntry.cpp
deleted file mode 100644
diff --git a/Engine/Plugins/AudioInsights/Source/AudioInsights/Private/Views/SoundSubmixAssetDashboardEntry.h b/Engine/Plugins/AudioInsights/Source/AudioInsights/Private/Views/SoundSubmixAssetDashboardEntry.h
deleted file mode 100644
diff --git a/Engine/Plugins/AudioInsights/Source/AudioInsights/Private/Views/SubmixesDashboardViewFactory.cpp b/Engine/Plugins/AudioInsights/Source/AudioInsights/Private/Views/SubmixesDashboardViewFactory.cpp
deleted file mode 100644
diff --git a/Engine/Plugins/AudioInsights/Source/AudioInsights/Private/Views/SubmixesDashboardViewFactory.h b/Engine/Plugins/AudioInsights/Source/AudioInsights/Private/Views/SubmixesDashboardViewFactory.h
deleted file mode 100644
diff --git a/Engine/Plugins/AudioInsights/Source/AudioInsights/Private/Views/TableDashboardViewFactory.cpp b/Engine/Plugins/AudioInsights/Source/AudioInsights/Private/Views/TableDashboardViewFactory.cpp
--- a/Engine/Plugins/AudioInsights/Source/AudioInsights/Private/Views/TableDashboardViewFactory.cpp
+++ b/Engine/Plugins/AudioInsights/Source/AudioInsights/Private/Views/TableDashboardViewFactory.cpp
@@ -2,19 +2,20 @@
 #include "Views/TableDashboardViewFactory.h"
 
 #include "Algo/Transform.h"
-#include "AudioDeviceManager.h"
-#include "AudioInsightsDashboardAssetCommands.h"
 #include "AudioInsightsTraceProviderBase.h"
 #include "Containers/Array.h"
-#include "DrawDebugHelpers.h"
-#include "Editor.h"
 #include "Framework/Commands/UICommandList.h"
-#include "Sound/SoundCue.h"
-#include "Sound/SoundNodeAttenuation.h"
-#include "Subsystems/AssetEditorSubsystem.h"
 #include "Templates/SharedPointer.h"
 #include "Widgets/SBoxPanel.h"
 
+#if WITH_EDITOR
+#include "AudioDeviceManager.h"
+#include "AudioInsightsDashboardAssetCommands.h"
+#include "Editor.h"
+#include "Editor/EditorEngine.h"
+#include "Subsystems/AssetEditorSubsystem.h"
+#endif // WITH_EDITOR
+
 #define LOCTEXT_NAMESPACE "AudioInsights"
 
 namespace UE::Audio::Insights
@@ -232,28 +233,17 @@ namespace UE::Audio::Insights
 		}
 
 #if WITH_EDITOR
-		const bool bDrawDebug = IsDebugDrawEnabled();
-		if (bDrawDebug)
+		if (IsDebugDrawEnabled())
 		{
-			FAudioDeviceManager* Manager = FAudioDeviceManager::Get();
-			if (!Manager)
-			{
-				return;
-			}
-
-			if (!FilteredEntriesListView.IsValid())
+			if (FilteredEntriesListView.IsValid())
 			{
-				return;
-			}
+				TArray<TSharedPtr<IDashboardDataViewEntry>> SelectedItems = FilteredEntriesListView->GetSelectedItems();
 
-			TArray<TSharedPtr<IDashboardDataViewEntry>> SelectedItems = FilteredEntriesListView->GetSelectedItems();
-			for (const TSharedPtr<IDashboardDataViewEntry>& SelectedEntry : SelectedItems)
-			{
-				if (SelectedEntry.IsValid())
+				if (FAudioDeviceManager* AudioDeviceManager = FAudioDeviceManager::Get())
 				{
-					Manager->IterateOverAllDevices([this, &SelectedEntry, InElapsed](::Audio::FDeviceId DeviceId, FAudioDevice* Device)
+					AudioDeviceManager->IterateOverAllDevices([this, &SelectedItems, InElapsed](::Audio::FDeviceId DeviceId, FAudioDevice* Device)
 					{
-						DebugDraw(InElapsed, *SelectedEntry.Get(), DeviceId);
+						DebugDraw(InElapsed, SelectedItems, DeviceId);
 					});
 				}
 			}
@@ -313,6 +303,7 @@ namespace UE::Audio::Insights
 			.ColorAndOpacity(GetRowColor(RowDataPtr))
 			.OnDoubleClicked_Lambda([this, RowDataPtr](const FGeometry& MyGeometry, const FPointerEvent& PointerEvent)
 			{
+#if WITH_EDITOR
 				if (GEditor)
 				{
 					TSharedPtr<IObjectDashboardEntry> ObjectData = StaticCastSharedPtr<IObjectDashboardEntry>(RowDataPtr);
@@ -326,7 +317,7 @@ namespace UE::Audio::Insights
 						}
 					}
 				}
-
+#endif // WITH_EDITOR
 				return FReply::Unhandled();
 			})
 		];
@@ -334,11 +325,13 @@ namespace UE::Audio::Insights
 
 	FReply FTraceObjectTableDashboardViewFactory::OnDataRowKeyInput(const FGeometry& Geometry, const FKeyEvent& KeyEvent) const
 	{
-		if (FilteredEntriesListView.IsValid())
+#if WITH_EDITOR
+		if (GEditor && FilteredEntriesListView.IsValid())
 		{
-			TArray<TSharedPtr<IDashboardDataViewEntry>> SelectedItems = FilteredEntriesListView->GetSelectedItems();
 			if (KeyEvent.GetKey() == EKeys::Enter)
 			{
+				TArray<TSharedPtr<IDashboardDataViewEntry>> SelectedItems = FilteredEntriesListView->GetSelectedItems();
+
 				for (const TSharedPtr<IDashboardDataViewEntry>& SelectedItem : SelectedItems)
 				{
 					if (SelectedItem.IsValid())
@@ -357,6 +350,7 @@ namespace UE::Audio::Insights
 				return FReply::Handled();
 			}
 		}
+#endif // WITH_EDITOR
 		return FReply::Unhandled();
 	}
 
@@ -365,6 +359,7 @@ namespace UE::Audio::Insights
 		if (!DashboardWidget.IsValid())
 		{
 			DashboardWidget = SNew(SVerticalBox)
+#if WITH_EDITOR
 			+ SVerticalBox::Slot()
 			.AutoHeight()
 			.HAlign(HAlign_Fill)
@@ -376,6 +371,7 @@ namespace UE::Audio::Insights
 					MakeAssetMenuBar()
 				]
 			]
+#endif // WITH_EDITOR
 			+ SVerticalBox::Slot()
 			.AutoHeight()
 			.HAlign(HAlign_Fill)
@@ -392,6 +388,7 @@ namespace UE::Audio::Insights
 		return DashboardWidget->AsShared();
 	}
 
+#if WITH_EDITOR
 	TSharedRef<SWidget> FTraceObjectTableDashboardViewFactory::MakeAssetMenuBar() const
 	{
 		const FDashboardAssetCommands& Commands = FDashboardAssetCommands::Get();
@@ -464,129 +461,7 @@ namespace UE::Audio::Insights
 
 		return false;
 	}
-
-	void FSoundAttenuationVisualizer::Draw(float InDeltaTime, const FTransform& InTransform, const UObject& InObject, const UWorld& InWorld) const
-	{
-		if (LastObjectId != InObject.GetUniqueID())
-		{
-			ShapeDetailsMap.Reset();
-			if (const USoundCue* SoundCue = Cast<const USoundCue>(&InObject))
-			{
-				TArray<const USoundNodeAttenuation*> AttenuationNodes;
-				SoundCue->RecursiveFindAttenuation(SoundCue->FirstNode, AttenuationNodes);
-				for (const USoundNodeAttenuation* Node : AttenuationNodes)
-				{
-					if (Node)
-					{
-						if (const FSoundAttenuationSettings* AttenuationSettingsToApply = Node->GetAttenuationSettingsToApply())
-						{
-							AttenuationSettingsToApply->CollectAttenuationShapesForVisualization(ShapeDetailsMap);
-						}
-					}
-				}
-			}
-			else if (const USoundBase* SoundBase = Cast<USoundBase>(&InObject))
-			{
-				if (const FSoundAttenuationSettings* Settings = SoundBase->GetAttenuationSettingsToApply())
-				{
-					Settings->CollectAttenuationShapesForVisualization(ShapeDetailsMap);
-				}
-			}
-			else
-			{
-				return;
-			}
-		}
-		LastObjectId = InObject.GetUniqueID();
-
-		for (const TPair<EAttenuationShape::Type, FBaseAttenuationSettings::AttenuationShapeDetails>& Pair : ShapeDetailsMap)
-		{
-			const FBaseAttenuationSettings::AttenuationShapeDetails& ShapeDetails = Pair.Value;
-			switch (Pair.Key)
-			{
-				case EAttenuationShape::Sphere:
-				{
-					if (ShapeDetails.Falloff > 0.f)
-					{
-						DrawDebugSphere(&InWorld, InTransform.GetTranslation(), ShapeDetails.Extents.X + ShapeDetails.Falloff, 10, Color);
-						DrawDebugSphere(&InWorld, InTransform.GetTranslation(), ShapeDetails.Extents.X, 10, Color);
-					}
-					else
-					{
-						DrawDebugSphere(&InWorld, InTransform.GetTranslation(), ShapeDetails.Extents.X, 10, Color);
-					}
-					break;
-				}
-
-				case EAttenuationShape::Box:
-				{
-					if (ShapeDetails.Falloff > 0.f)
-					{
-						DrawDebugBox(&InWorld, InTransform.GetTranslation(), ShapeDetails.Extents + FVector(ShapeDetails.Falloff), InTransform.GetRotation(), Color);
-						DrawDebugBox(&InWorld, InTransform.GetTranslation(), ShapeDetails.Extents, InTransform.GetRotation(), Color);
-					}
-					else
-					{
-						DrawDebugBox(&InWorld, InTransform.GetTranslation(), ShapeDetails.Extents, InTransform.GetRotation(), Color);
-					}
-					break;
-				}
-
-				case EAttenuationShape::Capsule:
-				{
-					if (ShapeDetails.Falloff > 0.f)
-					{
-						DrawDebugCapsule(&InWorld, InTransform.GetTranslation(), ShapeDetails.Extents.X + ShapeDetails.Falloff, ShapeDetails.Extents.Y + ShapeDetails.Falloff, InTransform.GetRotation(), Color);
-						DrawDebugCapsule(&InWorld, InTransform.GetTranslation(), ShapeDetails.Extents.X, ShapeDetails.Extents.Y, InTransform.GetRotation(), Color);
-					}
-					else
-					{
-						DrawDebugCapsule(&InWorld, InTransform.GetTranslation(), ShapeDetails.Extents.X, ShapeDetails.Extents.Y, InTransform.GetRotation(), Color);
-					}
-					break;
-				}
-
-				case EAttenuationShape::Cone:
-				{
-					const FVector Origin = InTransform.GetTranslation() - (InTransform.GetUnitAxis(EAxis::X) * ShapeDetails.ConeOffset);
-
-					if (ShapeDetails.Falloff > 0.f || ShapeDetails.Extents.Z > 0.f)
-					{
-						const float OuterAngle = FMath::DegreesToRadians(ShapeDetails.Extents.Y + ShapeDetails.Extents.Z);
-						const float InnerAngle = FMath::DegreesToRadians(ShapeDetails.Extents.Y);
-						DrawDebugCone(&InWorld, Origin, InTransform.GetUnitAxis(EAxis::X), ShapeDetails.Extents.X + ShapeDetails.Falloff + ShapeDetails.ConeOffset, OuterAngle, OuterAngle, 10, Color);
-						DrawDebugCone(&InWorld, Origin, InTransform.GetUnitAxis(EAxis::X), ShapeDetails.Extents.X + ShapeDetails.ConeOffset, InnerAngle, InnerAngle, 10, Color);
-					}
-					else
-					{
-						const float Angle = FMath::DegreesToRadians(ShapeDetails.Extents.Y);
-						DrawDebugCone(&InWorld, Origin, InTransform.GetUnitAxis(EAxis::X), ShapeDetails.Extents.X + ShapeDetails.ConeOffset, Angle, Angle, 10, Color);
-					}
-
-					if (!FMath::IsNearlyZero(ShapeDetails.ConeSphereRadius, UE_KINDA_SMALL_NUMBER))
-					{
-						if (ShapeDetails.ConeSphereFalloff > 0.f)
-						{
-
-							DrawDebugSphere(&InWorld, Origin, ShapeDetails.ConeSphereRadius + ShapeDetails.ConeSphereFalloff, 10, Color);
-							DrawDebugSphere(&InWorld, Origin, ShapeDetails.ConeSphereRadius, 10, Color);
-						}
-						else
-						{
-							DrawDebugSphere(&InWorld, Origin, ShapeDetails.ConeSphereRadius, 10, Color);
-						}
-					}
-
-					break;
-				}
-
-				default:
-				{
-					break;
-				}
-			}
-		}
-	}
+#endif // WITH_EDITOR
 } // namespace UE::Audio::Insights
 
 #undef LOCTEXT_NAMESPACE
diff --git a/Engine/Plugins/AudioInsights/Source/AudioInsights/Private/Views/ViewportDashboardViewFactory.h b/Engine/Plugins/AudioInsights/Source/AudioInsights/Private/Views/ViewportDashboardViewFactory.h
deleted file mode 100644
diff --git a/Engine/Plugins/AudioInsights/Source/AudioInsights/Private/Views/VirtualLoopDashboardViewFactory.cpp b/Engine/Plugins/AudioInsights/Source/AudioInsights/Private/Views/VirtualLoopDashboardViewFactory.cpp
--- a/Engine/Plugins/AudioInsights/Source/AudioInsights/Private/Views/VirtualLoopDashboardViewFactory.cpp
+++ b/Engine/Plugins/AudioInsights/Source/AudioInsights/Private/Views/VirtualLoopDashboardViewFactory.cpp
@@ -1,16 +1,18 @@
 // Copyright Epic Games, Inc. All Rights Reserved.
 #include "Views/VirtualLoopDashboardViewFactory.h"
 
-#include "Audio/AudioDebug.h"
 #include "AudioDefines.h"
-#include "AudioDeviceManager.h"
 #include "AudioInsightsModule.h"
 #include "AudioInsightsStyle.h"
-#include "DrawDebugHelpers.h"
 #include "Internationalization/Text.h"
 #include "Providers/VirtualLoopTraceProvider.h"
 #include "Templates/SharedPointer.h"
 
+#if WITH_EDITOR
+#include "Audio/AudioDebug.h"
+#include "AudioDeviceManager.h"
+#endif // WITH_EDITOR
+
 #define LOCTEXT_NAMESPACE "AudioInsights"
 
 
@@ -25,7 +27,6 @@ namespace UE::Audio::Insights
 	} // namespace VirtualLoopPrivate
 
 	FVirtualLoopDashboardViewFactory::FVirtualLoopDashboardViewFactory()
-		: AttenuationVisualizer(FColor::Blue)
 	{
 		const FTraceModule& TraceModule = FAudioInsightsModule::GetChecked().GetTraceModule();
 		Providers = TArray<TSharedPtr<FTraceProviderBase>>
@@ -258,27 +259,9 @@ namespace UE::Audio::Insights
 		return false;
 	}
 
-	void FVirtualLoopDashboardViewFactory::DebugDraw(float InElapsed, const IDashboardDataViewEntry& InEntry, ::Audio::FDeviceId DeviceId) const
+	void FVirtualLoopDashboardViewFactory::DebugDraw(float InElapsed, const TArray<TSharedPtr<IDashboardDataViewEntry>>& InSelectedItems, ::Audio::FDeviceId InAudioDeviceId) const
 	{
-		const FVirtualLoopDashboardEntry& LoopData = VirtualLoopPrivate::CastEntry(InEntry);
-		const FRotator& Rotator = LoopData.Rotator;
-		const FVector& Location = LoopData.Location;
-		const FString Description = FString::Printf(TEXT("%s [Virt: %.2fs]"), *LoopData.Name, LoopData.TimeVirtualized);
-
-		const TArray<UWorld*> Worlds = FAudioDeviceManager::Get()->GetWorldsUsingAudioDevice(DeviceId);
-		for (UWorld* World : Worlds)
-		{
-			DrawDebugSphere(World, Location, 30.0f, 8, AttenuationVisualizer.Color, false, InElapsed, SDPG_Foreground);
-			DrawDebugString(World, Location + FVector(0, 0, 32), *Description, nullptr, AttenuationVisualizer.Color, InElapsed, false, 1.0f);
-
-			if (const UObject* Object = LoopData.GetObject())
-			{
-				FTransform Transform;
-				Transform.SetLocation(Location);
-				Transform.SetRotation(FQuat(Rotator));
-				AttenuationVisualizer.Draw(InElapsed, Transform, *Object, *World);
-			}
-		}
+		OnDebugDrawEntries.Broadcast(InElapsed, InSelectedItems, InAudioDeviceId);
 	}
 #endif // WITH_EDITOR
 } // namespace UE::Audio::Insights
diff --git a/Engine/Plugins/AudioInsights/Source/AudioInsights/Private/Views/VirtualLoopDashboardViewFactory.h b/Engine/Plugins/AudioInsights/Source/AudioInsights/Private/Views/VirtualLoopDashboardViewFactory.h
deleted file mode 100644
diff --git a/Engine/Plugins/AudioInsights/Source/AudioInsights/Public/AudioInsightsDashboardAssetCommands.h b/Engine/Plugins/AudioInsights/Source/AudioInsights/Public/AudioInsightsDashboardAssetCommands.h
--- a/Engine/Plugins/AudioInsights/Source/AudioInsights/Public/AudioInsightsDashboardAssetCommands.h
+++ b/Engine/Plugins/AudioInsights/Source/AudioInsights/Public/AudioInsightsDashboardAssetCommands.h
@@ -24,10 +24,6 @@ namespace UE::Audio::Insights
 		virtual TSharedPtr<const FUICommandInfo> GetStartCommand() const;
 		virtual TSharedPtr<const FUICommandInfo> GetStopCommand() const;
 
-		virtual TSharedPtr<const FUICommandInfo> GetMuteCommand() const;
-		virtual TSharedPtr<const FUICommandInfo> GetSoloCommand() const;
-		virtual TSharedPtr<const FUICommandInfo> GetClearMuteSoloCommand() const;
-
 	private:
 		FSlateIcon GetStartIcon() const;
 		FSlateIcon GetStopIcon() const;
@@ -43,14 +39,5 @@ namespace UE::Audio::Insights
 
 		/** Stops the trace session & disables the required audio channels off. */
 		TSharedPtr<FUICommandInfo> Stop;
-
-		/** Mutes the selected item in the Mixer Source View. */
-		TSharedPtr<FUICommandInfo> Mute;
-
-		/** Solo the selected item in the Mixer Source View. */
-		TSharedPtr<FUICommandInfo> Solo;
-
-		/** Clears any Mute/Solo state from any of the items in the Mixer Source View. */
-		TSharedPtr<FUICommandInfo> ClearMuteSolo;
 	};
 } // namespace UE::Audio::Insights
diff --git a/Engine/Plugins/AudioInsights/Source/AudioInsights/Public/AudioInsightsDashboardFactory.h b/Engine/Plugins/AudioInsights/Source/AudioInsights/Public/AudioInsightsDashboardFactory.h
new file mode 100644
--- /dev/null
+++ b/Engine/Plugins/AudioInsights/Source/AudioInsights/Public/AudioInsightsDashboardFactory.h
@@ -0,0 +1,42 @@
+// Copyright Epic Games, Inc. All Rights Reserved.
+#pragma once
+
+#include "AudioDefines.h"
+#include "Containers/Map.h"
+#include "Framework/Docking/TabManager.h"
+#include "IAudioInsightsDashboardFactory.h"
+#include "Templates/SharedPointer.h"
+#include "Views/DashboardViewFactory.h"
+#include "Widgets/Docking/SDockTab.h"
+
+namespace UE::Audio::Insights
+{
+	class FDashboardFactory : public IDashboardFactory, public TSharedFromThis<FDashboardFactory>
+	{
+	public:
+		FDashboardFactory() = default;
+		virtual ~FDashboardFactory() = default;
+
+		TSharedRef<SDockTab> MakeDockTabWidget(const FSpawnTabArgs& Args);
+
+		virtual void RegisterViewFactory(TSharedRef<IDashboardViewFactory> InFactory) override;
+		virtual void UnregisterViewFactory(FName InDashboardName) override;
+		virtual ::Audio::FDeviceId GetDeviceId() const override;
+
+	private:
+		TSharedRef<SWidget> MakeMenuBarWidget();
+
+		TSharedPtr<FTabManager::FLayout> GetDefaultTabLayout();
+
+		void RegisterTabSpawners();
+		void UnregisterTabSpawners();
+
+		TSharedPtr<FTabManager> DashboardTabManager;
+		TSharedPtr<FWorkspaceItem> DashboardWorkspace;
+		TSharedPtr<FTabManager::FLayout> TabLayout;
+
+		static constexpr ::Audio::FDeviceId ActiveDeviceId = 1; // The default audio device id in standalone
+
+		TMap<FName, TSharedPtr<IDashboardViewFactory>> DashboardViewFactories;
+	};
+} // namespace UE::Audio::Insights
diff --git a/Engine/Plugins/AudioInsights/Source/AudioInsights/Public/AudioInsightsModule.h b/Engine/Plugins/AudioInsights/Source/AudioInsights/Public/AudioInsightsModule.h
new file mode 100644
--- /dev/null
+++ b/Engine/Plugins/AudioInsights/Source/AudioInsights/Public/AudioInsightsModule.h
@@ -0,0 +1,40 @@
+// Copyright Epic Games, Inc. All Rights Reserved.
+#pragma once 
+
+#include "AudioInsightsDashboardFactory.h"
+#include "AudioInsightsTraceModule.h"
+#include "Framework/Docking/TabManager.h"
+#include "IAudioInsightsModule.h"
+#include "Templates/SharedPointer.h"
+#include "Views/DashboardViewFactory.h"
+#include "Widgets/Docking/SDockTab.h"
+
+
+namespace UE::Audio::Insights
+{
+	class FAudioInsightsModule final : public IAudioInsightsModule
+	{
+	public:
+		FAudioInsightsModule() = default;
+
+		virtual ~FAudioInsightsModule() = default;
+
+		virtual void StartupModule() override;
+		virtual void ShutdownModule() override;
+		virtual void RegisterDashboardViewFactory(TSharedRef<IDashboardViewFactory> InDashboardFactory) override;
+		virtual void UnregisterDashboardViewFactory(FName InName) override;
+		virtual ::Audio::FDeviceId GetDeviceId() const override;
+
+		AUDIOINSIGHTS_API static FAudioInsightsModule& GetChecked();
+		AUDIOINSIGHTS_API FTraceModule& GetTraceModule();
+
+		TSharedRef<FDashboardFactory> GetDashboardFactory();
+		const TSharedRef<FDashboardFactory> GetDashboardFactory() const;
+
+		TSharedRef<SDockTab> CreateDashboardTabWidget(const FSpawnTabArgs& Args);
+
+	private:
+		TSharedPtr<FDashboardFactory> DashboardFactory;
+		FTraceModule TraceModule;
+	};
+} // namespace UE::Audio::Insights
diff --git a/Engine/Plugins/AudioInsights/Source/AudioInsights/Public/AudioInsightsStyle.h b/Engine/Plugins/AudioInsights/Source/AudioInsights/Public/AudioInsightsStyle.h
new file mode 100644
--- /dev/null
+++ b/Engine/Plugins/AudioInsights/Source/AudioInsights/Public/AudioInsightsStyle.h
@@ -0,0 +1,153 @@
+// Copyright Epic Games, Inc. All Rights Reserved.
+#pragma once
+
+#include "Misc/Paths.h"
+#include "Styling/AppStyle.h"
+#include "Styling/SlateStyle.h"
+#include "Styling/SlateStyleMacros.h"
+#include "Styling/SlateStyleRegistry.h"
+#include "Textures/SlateIcon.h"
+
+#define LOCTEXT_NAMESPACE "AudioInsights"
+
+namespace UE::Audio::Insights
+{
+	class AUDIOINSIGHTS_API FSlateStyle final : public FSlateStyleSet
+	{
+	public:
+		static FSlateStyle& Get()
+		{
+			static FSlateStyle InsightsStyle;
+			return InsightsStyle;
+		}
+
+		static FName GetStyleName()
+		{
+			const FLazyName StyleName = "AudioInsights";
+			return StyleName.Resolve();
+		}
+
+		const FNumberFormattingOptions* GetAmpFloatFormat() const
+		{
+			static FNumberFormattingOptions FloatFormat;
+			FloatFormat.MinimumIntegralDigits = 1;
+			FloatFormat.MaximumIntegralDigits = 1;
+			FloatFormat.MinimumFractionalDigits = 3;
+			FloatFormat.MaximumFractionalDigits = 3;
+			return &FloatFormat;
+		};
+
+		const FNumberFormattingOptions* GetDefaultFloatFormat() const
+		{
+			static FNumberFormattingOptions FloatFormat;
+			FloatFormat.MinimumIntegralDigits = 1;
+			FloatFormat.MinimumFractionalDigits = 4;
+			FloatFormat.MaximumFractionalDigits = 4;
+			return &FloatFormat;
+		};
+
+		const FNumberFormattingOptions* GetFreqFloatFormat() const
+		{
+			static FNumberFormattingOptions FloatFormat;
+			FloatFormat.MinimumIntegralDigits = 1;
+			FloatFormat.MaximumIntegralDigits = 5;
+			FloatFormat.MinimumFractionalDigits = 0;
+			FloatFormat.MaximumFractionalDigits = 2;
+			return &FloatFormat;
+		};
+
+		const FNumberFormattingOptions* GetPitchFloatFormat() const
+		{
+			static FNumberFormattingOptions FloatFormat;
+			FloatFormat.MinimumIntegralDigits = 1;
+			FloatFormat.MaximumIntegralDigits = 3;
+			FloatFormat.MinimumFractionalDigits = 3;
+			FloatFormat.MaximumFractionalDigits = 3;
+			return &FloatFormat;
+		};
+
+		const FNumberFormattingOptions* GetTimeFormat() const
+		{
+			static FNumberFormattingOptions FloatFormat;
+			FloatFormat.MinimumIntegralDigits = 1;
+			FloatFormat.MinimumFractionalDigits = 3;
+			FloatFormat.MaximumFractionalDigits = 3;
+			return &FloatFormat;
+		};
+
+		FText FormatSecondsAsTime(float InTimeSec) const
+		{
+			return FText::Format(LOCTEXT("TimeInSecondsFormat", "{0}s"), FText::AsNumber(InTimeSec, GetTimeFormat()));
+		}
+
+		FText FormatMillisecondsAsTime(float InTimeMS) const
+		{
+			return FText::Format(LOCTEXT("TimeInMillisecondsFormat", "{0}ms"), FText::AsNumber(InTimeMS, GetTimeFormat()));
+		}
+
+		FSlateIcon CreateIcon(FName InName)  const
+		{
+			return { GetStyleName(), InName };
+		}
+
+		const FSlateBrush& GetBrushEnsured(FName InName)  const
+		{
+			const ISlateStyle* AudioInsightsStyle = FSlateStyleRegistry::FindSlateStyle(GetStyleName());
+			if (ensureMsgf(AudioInsightsStyle, TEXT("Missing slate style '%s'"), *GetStyleName().ToString()))
+			{
+				const FSlateBrush* Brush = AudioInsightsStyle->GetBrush(InName);
+				if (ensureMsgf(Brush, TEXT("Missing brush '%s'"), *InName.ToString()))
+				{
+					return *Brush;
+				}
+			}
+
+			if (const FSlateBrush* NoBrush = FAppStyle::GetBrush("NoBrush"))
+			{
+				return *NoBrush;
+			}
+
+			return *DefaultBrush;
+		}
+
+		FSlateStyle()
+			: FSlateStyleSet(GetStyleName())
+		{
+			SetParentStyleName(FAppStyle::GetAppStyleSetName());
+
+			const FString PluginsDir = FPaths::EnginePluginsDir();
+
+			SetContentRoot(PluginsDir / TEXT("AudioInsights/Content"));
+			SetCoreContentRoot(PluginsDir / TEXT("Slate"));
+
+			Set("AudioInsights.Analyzers.BackgroundColor", FLinearColor(0.0075f, 0.0075f, 0.0075f, 1.0f));
+
+			const FVector2D Icon16(16.0f, 16.0f);
+			const FVector2D Icon20(20.0f, 20.0f);
+			const FVector2D Icon40(40.0f, 40.0f);
+			const FVector2D Icon64(64.0f, 64.0f);
+
+			Set("AudioInsights.Icon", new IMAGE_BRUSH_SVG(TEXT("Icons/audio_insights_icon"), Icon16));
+			Set("AudioInsights.Icon.Dashboard", new IMAGE_BRUSH_SVG(TEXT("Icons/audio_dashboard"), Icon16));
+			Set("AudioInsights.Icon.Event", new IMAGE_BRUSH_SVG(TEXT("Icons/audio_event"), Icon16));
+			Set("AudioInsights.Icon.Log", new IMAGE_BRUSH_SVG(TEXT("Icons/audio_log"), Icon16));
+			Set("AudioInsights.Icon.Sources", new IMAGE_BRUSH_SVG(TEXT("Icons/audio_sources"), Icon16));
+			Set("AudioInsights.Icon.Submix", new IMAGE_BRUSH_SVG(TEXT("Icons/audio_submix"), Icon16));
+			Set("AudioInsights.Icon.VirtualLoop", new IMAGE_BRUSH_SVG(TEXT("Icons/audio_virtualloop"), Icon16));
+			Set("AudioInsights.Icon.Viewport", new IMAGE_BRUSH_SVG(TEXT("Icons/viewport"), Icon16));
+			Set("AudioInsights.Icon.Open", new IMAGE_BRUSH_SVG(TEXT("Icons/open"), Icon20));
+			Set("AudioInsights.Icon.ContentBrowser", new IMAGE_BRUSH_SVG(TEXT("Icons/content_browser"), Icon20));
+			Set("AudioInsights.Icon.Start.Active", new IMAGE_BRUSH_SVG(TEXT("Icons/start_active"), Icon20));
+			Set("AudioInsights.Icon.Start.Inactive", new IMAGE_BRUSH_SVG(TEXT("Icons/start_inactive"), Icon20));
+			Set("AudioInsights.Icon.Record.Active", new IMAGE_BRUSH_SVG(TEXT("Icons/record_active"), Icon20));
+			Set("AudioInsights.Icon.Record.Inactive", new IMAGE_BRUSH_SVG(TEXT("Icons/record_inactive"), Icon20));
+			Set("AudioInsights.Icon.Stop.Active", new IMAGE_BRUSH_SVG(TEXT("Icons/stop_active"), Icon20));
+			Set("AudioInsights.Icon.Stop.Inactive", new IMAGE_BRUSH_SVG(TEXT("Icons/stop_inactive"), Icon20));
+
+			Set("AudioInsights.Thumbnail", new IMAGE_BRUSH_SVG(TEXT("Icons/audio_insights"), Icon16));
+
+			FSlateStyleRegistry::RegisterSlateStyle(*this);
+		}
+	};
+} // namespace UE::Audio::Insights
+#undef LOCTEXT_NAMESPACE
diff --git a/Engine/Plugins/AudioInsights/Source/AudioInsights/Public/AudioInsightsTraceModule.h b/Engine/Plugins/AudioInsights/Source/AudioInsights/Public/AudioInsightsTraceModule.h
new file mode 100644
--- /dev/null
+++ b/Engine/Plugins/AudioInsights/Source/AudioInsights/Public/AudioInsightsTraceModule.h
@@ -0,0 +1,54 @@
+// Copyright Epic Games, Inc. All Rights Reserved.
+#pragma once
+
+#include "AudioInsightsTraceChannelHandle.h"
+#include "AudioInsightsTraceProviderBase.h"
+#include "Containers/Map.h"
+#include "Templates/SharedPointer.h"
+#include "TraceServices/Model/AnalysisSession.h"
+#include "TraceServices/ModuleService.h"
+#include "UObject/NameTypes.h"
+
+namespace UE::Audio::Insights
+{
+	class AUDIOINSIGHTS_API FTraceModule : public TraceServices::IModule
+	{
+	public:
+		FTraceModule();
+		virtual ~FTraceModule() = default;
+
+		//~ Begin TraceServices::IModule interface
+		virtual void GetModuleInfo(TraceServices::FModuleInfo& OutModuleInfo) override;
+		virtual void OnAnalysisBegin(TraceServices::IAnalysisSession& Session) override;
+		virtual void GetLoggers(TArray<const TCHAR *>& OutLoggers) override;
+		virtual void GenerateReports(const TraceServices::IAnalysisSession& Session, const TCHAR* CmdLine, const TCHAR* OutputDirectory) override;
+		virtual const TCHAR* GetCommandLineArgument() override { return TEXT("audiotrace"); }
+		//~ End TraceServices::IModule interface
+
+		template <typename TraceProviderType>
+		TSharedPtr<TraceProviderType> FindAudioTraceProvider() const
+		{
+			return StaticCastSharedPtr<TraceProviderType>(TraceProviders.FindRef(TraceProviderType::GetName_Static()));
+		}
+
+		TSharedRef<FTraceChannelManager> GetChannelManager();
+
+		void StartTraceAnalysis() const;
+		bool IsTraceAnalysisActive() const;
+		void StopTraceAnalysis() const;
+
+		void SetFirstTimeStamp(double InFirstTimeStamp) { FirstTimeStamp = InFirstTimeStamp; }
+		double GetFirstTimeStamp() const { return FirstTimeStamp; }
+
+	private:
+		static const FName GetName();
+
+		static void DisableAllTraceChannels();
+		static void EnableAudioInsightsTraceChannels();
+
+		TSharedRef<FTraceChannelManager> ChannelManager;
+		TMap<FName, TSharedPtr<FTraceProviderBase>> TraceProviders;
+
+		double FirstTimeStamp = -TNumericLimits<double>::Min();
+	};
+} // namespace UE::Audio::Insights
diff --git a/Engine/Plugins/AudioInsights/Source/AudioInsights/Public/AudioInsightsTraceProviderBase.h b/Engine/Plugins/AudioInsights/Source/AudioInsights/Public/AudioInsightsTraceProviderBase.h
--- a/Engine/Plugins/AudioInsights/Source/AudioInsights/Public/AudioInsightsTraceProviderBase.h
+++ b/Engine/Plugins/AudioInsights/Source/AudioInsights/Public/AudioInsightsTraceProviderBase.h
@@ -120,10 +120,15 @@ namespace UE::Audio::Insights
 
 		const FDeviceData* FindFilteredDeviceData() const
 		{
-			const IAudioInsightsModule& InsightsModule = FModuleManager::GetModuleChecked<IAudioInsightsModule>(IAudioInsightsModule::GetName());
-			const ::Audio::FDeviceId DeviceId = InsightsModule.GetDeviceId();
+#if WITH_EDITOR
+			const IAudioInsightsModule& AudioInsightsEditorModule = FModuleManager::GetModuleChecked<IAudioInsightsModule>(IAudioInsightsModule::GetEditorName());
+			const ::Audio::FDeviceId AudioDeviceId = AudioInsightsEditorModule.GetDeviceId();
+#else
+			const IAudioInsightsModule& AudioInsightsModule = FModuleManager::GetModuleChecked<IAudioInsightsModule>(IAudioInsightsModule::GetName());
+			const ::Audio::FDeviceId AudioDeviceId = AudioInsightsModule.GetDeviceId();
+#endif // WITH_EDITOR
 
-			return DeviceDataMap.Find(DeviceId);
+			return DeviceDataMap.Find(AudioDeviceId);
 		}
 
 	protected:
@@ -174,10 +179,15 @@ namespace UE::Audio::Insights
 
 		FDeviceData* FindFilteredDeviceData()
 		{
-			const IAudioInsightsModule& InsightsModule = FModuleManager::GetModuleChecked<IAudioInsightsModule>(IAudioInsightsModule::GetName());
-			const ::Audio::FDeviceId DeviceId = InsightsModule.GetDeviceId();
+#if WITH_EDITOR
+			const IAudioInsightsModule& AudioInsightsEditorModule = FModuleManager::GetModuleChecked<IAudioInsightsModule>(IAudioInsightsModule::GetEditorName());
+			const ::Audio::FDeviceId AudioDeviceId = AudioInsightsEditorModule.GetDeviceId();
+#else
+			const IAudioInsightsModule& AudioInsightsModule = FModuleManager::GetModuleChecked<IAudioInsightsModule>(IAudioInsightsModule::GetName());
+			const ::Audio::FDeviceId AudioDeviceId = AudioInsightsModule.GetDeviceId();
+#endif // WITH_EDITOR
 
-			return DeviceDataMap.Find(DeviceId);
+			return DeviceDataMap.Find(AudioDeviceId);
 		}
 
 		bool RemoveDeviceEntry(::Audio::FDeviceId InDeviceId, const EntryKey& InKey)
diff --git a/Engine/Plugins/AudioInsights/Source/AudioInsights/Public/IAudioInsightsModule.h b/Engine/Plugins/AudioInsights/Source/AudioInsights/Public/IAudioInsightsModule.h
--- a/Engine/Plugins/AudioInsights/Source/AudioInsights/Public/IAudioInsightsModule.h
+++ b/Engine/Plugins/AudioInsights/Source/AudioInsights/Public/IAudioInsightsModule.h
@@ -24,4 +24,10 @@ public:
 		const FLazyName ModuleName = "AudioInsights";
 		return ModuleName.Resolve();
 	}
+
+	static FName GetEditorName()
+	{
+		const FLazyName EditorModuleName = "AudioInsightsEditor";
+		return EditorModuleName.Resolve();
+	}
 };
diff --git a/Engine/Plugins/AudioInsights/Source/AudioInsights/Public/Messages/MixerSourceTraceMessages.h b/Engine/Plugins/AudioInsights/Source/AudioInsights/Public/Messages/MixerSourceTraceMessages.h
new file mode 100644
--- /dev/null
+++ b/Engine/Plugins/AudioInsights/Source/AudioInsights/Public/Messages/MixerSourceTraceMessages.h
@@ -0,0 +1,111 @@
+// Copyright Epic Games, Inc. All Rights Reserved.
+#pragma once
+
+#include "AudioDefines.h"
+#include "DSP/Dsp.h"
+#include "Messages/AnalyzerMessageQueue.h"
+#include "Trace/Analyzer.h"
+#include "Views/TableDashboardViewFactory.h"
+
+
+namespace UE::Audio::Insights
+{
+	struct FMixerSourceMessageBase
+	{
+		FMixerSourceMessageBase() = default;
+		FMixerSourceMessageBase(const Trace::IAnalyzer::FOnEventContext& InContext)
+		{
+			const Trace::IAnalyzer::FEventData& EventData = InContext.EventData;
+			DeviceId = static_cast<::Audio::FDeviceId>(EventData.GetValue<uint32>("DeviceId"));
+			Timestamp = InContext.EventTime.AsSeconds(EventData.GetValue<uint64>("Timestamp"));
+			PlayOrder = EventData.GetValue<uint32>("PlayOrder");
+		}
+
+		::Audio::FDeviceId DeviceId = INDEX_NONE;
+		uint32 PlayOrder = INDEX_NONE;
+		uint64 ComponentId = TNumericLimits<uint64>::Max();
+		double Timestamp = 0.0;
+	};
+
+	using FMixerSourceStopMessage = FMixerSourceMessageBase;
+
+	struct FMixerSourceStartMessage : public FMixerSourceMessageBase
+	{
+		FMixerSourceStartMessage() = default;
+		FMixerSourceStartMessage(const Trace::IAnalyzer::FOnEventContext& InContext)
+			: FMixerSourceMessageBase(InContext)
+		{
+			const Trace::IAnalyzer::FEventData& EventData = InContext.EventData;
+			EventData.GetString("Name", Name);
+			ComponentId = EventData.GetValue<uint64>("ComponentId");
+			SourceId = EventData.GetValue<int32>("SourceId");
+		}
+
+		FString Name;
+		int32 SourceId = INDEX_NONE;
+		uint64 ComponentId = TNumericLimits<uint64>::Max();
+	};
+
+#define DEFINE_MIXERSOURCE_PARAM_MESSAGE(ClassName, ParamName, Type, Default)			\
+	struct ClassName : public FMixerSourceMessageBase									\
+	{																					\
+		ClassName() = default;															\
+		ClassName(const Trace::IAnalyzer::FOnEventContext& InContext)					\
+			: FMixerSourceMessageBase(InContext)										\
+		{																				\
+			const Trace::IAnalyzer::FEventData& EventData = InContext.EventData;		\
+			ParamName = EventData.GetValue<Type>(#ParamName);							\
+		}																				\
+		Type ParamName = Default;														\
+	};
+
+	DEFINE_MIXERSOURCE_PARAM_MESSAGE(FMixerSourceDistanceAttenuationMessage, DistanceAttenuation, float, 0.0f)
+	DEFINE_MIXERSOURCE_PARAM_MESSAGE(FMixerSourceEnvelopeMessage, Envelope, float, 0.0f)
+	DEFINE_MIXERSOURCE_PARAM_MESSAGE(FMixerSourceHPFFreqMessage, HPFFrequency, float, MIN_FILTER_FREQUENCY)
+	DEFINE_MIXERSOURCE_PARAM_MESSAGE(FMixerSourceLPFFreqMessage, LPFFrequency, float, MAX_FILTER_FREQUENCY)
+	DEFINE_MIXERSOURCE_PARAM_MESSAGE(FMixerSourcePitchMessage, Pitch, float, 1.0f)
+	DEFINE_MIXERSOURCE_PARAM_MESSAGE(FMixerSourceVolumeMessage, Volume, float, 1.0f)
+#undef DEFINE_MIXERSOURCE_PARAM_MESSAGE
+
+	using FDataPoint = TPair<double, float>; // (Timestamp, Value)
+
+	class FMixerSourceDashboardEntry : public FSoundAssetDashboardEntry
+	{
+	public:
+		FMixerSourceDashboardEntry()
+		{
+			constexpr uint32 DataPointsCapacity = 256;
+
+			VolumeDataPoints.SetCapacity(DataPointsCapacity);
+			PitchDataPoints.SetCapacity(DataPointsCapacity);
+			LPFFreqDataPoints.SetCapacity(DataPointsCapacity);
+			HPFFreqDataPoints.SetCapacity(DataPointsCapacity);
+			EnvelopeDataPoints.SetCapacity(DataPointsCapacity);
+			DistanceAttenuationDataPoints.SetCapacity(DataPointsCapacity);
+		}
+
+		virtual ~FMixerSourceDashboardEntry() = default;
+
+		int32 SourceId = INDEX_NONE;
+
+		::Audio::TCircularAudioBuffer<FDataPoint> VolumeDataPoints;
+		::Audio::TCircularAudioBuffer<FDataPoint> PitchDataPoints;
+		::Audio::TCircularAudioBuffer<FDataPoint> LPFFreqDataPoints;
+		::Audio::TCircularAudioBuffer<FDataPoint> HPFFreqDataPoints;
+		::Audio::TCircularAudioBuffer<FDataPoint> EnvelopeDataPoints;
+		::Audio::TCircularAudioBuffer<FDataPoint> DistanceAttenuationDataPoints;
+	};
+
+	class FMixerSourceMessages
+	{
+	public:
+		TAnalyzerMessageQueue<FMixerSourceDistanceAttenuationMessage> DistanceAttenuationMessages { 0.1 };
+		TAnalyzerMessageQueue<FMixerSourceEnvelopeMessage> EnvelopeMessages { 0.1 };
+		TAnalyzerMessageQueue<FMixerSourceHPFFreqMessage> HPFFreqMessages { 0.1 };
+		TAnalyzerMessageQueue<FMixerSourceLPFFreqMessage> LPFFreqMessages { 0.1 };
+		TAnalyzerMessageQueue<FMixerSourcePitchMessage> PitchMessages { 0.1 };
+		TAnalyzerMessageQueue<FMixerSourceStartMessage> StartMessages { 0.1 };
+		TAnalyzerMessageQueue<FMixerSourceStopMessage> StopMessages { 0.1 };
+		TAnalyzerMessageQueue<FMixerSourceVolumeMessage> VolumeMessages { 0.1 };
+	};
+} // namespace UE::Audio::Insights
diff --git a/Engine/Plugins/AudioInsights/Source/AudioInsights/Public/Messages/VirtualLoopTraceMessages.h b/Engine/Plugins/AudioInsights/Source/AudioInsights/Public/Messages/VirtualLoopTraceMessages.h
new file mode 100644
--- /dev/null
+++ b/Engine/Plugins/AudioInsights/Source/AudioInsights/Public/Messages/VirtualLoopTraceMessages.h
@@ -0,0 +1,111 @@
+// Copyright Epic Games, Inc. All Rights Reserved.
+#pragma once
+
+#include "AudioDefines.h"
+#include "Math/NumericLimits.h"
+#include "Messages/AnalyzerMessageQueue.h"
+#include "Trace/Analyzer.h"
+#include "Views/TableDashboardViewFactory.h"
+
+
+namespace UE::Audio::Insights
+{
+	struct FVirtualLoopMessageBase
+	{
+		FVirtualLoopMessageBase() = default;
+		FVirtualLoopMessageBase(const Trace::IAnalyzer::FOnEventContext& InContext)
+		{
+			const Trace::IAnalyzer::FEventData& EventData = InContext.EventData;
+			DeviceId = static_cast<::Audio::FDeviceId>(EventData.GetValue<uint32>("DeviceId"));
+			Timestamp = InContext.EventTime.AsSeconds(EventData.GetValue<uint64>("Timestamp"));
+			PlayOrder = EventData.GetValue<uint32>("PlayOrder");
+		}
+
+		::Audio::FDeviceId DeviceId = INDEX_NONE;
+		uint32 PlayOrder = INDEX_NONE;
+		double Timestamp = 0.0;
+	};
+
+	using FVirtualLoopRealizeMessage = FVirtualLoopMessageBase;
+	using FVirtualLoopStopMessage = FVirtualLoopMessageBase;
+
+	struct FVirtualLoopVirtualizeMessage : public FVirtualLoopMessageBase
+	{
+		FVirtualLoopVirtualizeMessage() = default;
+		FVirtualLoopVirtualizeMessage(const Trace::IAnalyzer::FOnEventContext& InContext)
+			: FVirtualLoopMessageBase(InContext)
+		{
+			const Trace::IAnalyzer::FEventData& EventData = InContext.EventData;
+			EventData.GetString("Name", Name);
+			ComponentId = EventData.GetValue<uint64>("ComponentId");
+		}
+
+		FString Name;
+		uint64 ComponentId = TNumericLimits<uint64>::Max();
+	};
+
+	struct FVirtualLoopUpdateMessage : public FVirtualLoopMessageBase
+	{
+		FVirtualLoopUpdateMessage() = default;
+		FVirtualLoopUpdateMessage(const Trace::IAnalyzer::FOnEventContext& InContext)
+			: FVirtualLoopMessageBase(InContext)
+		{
+			const Trace::IAnalyzer::FEventData& EventData = InContext.EventData;
+			TimeVirtualized = EventData.GetValue<float>("TimeVirtualized");
+			PlaybackTime = EventData.GetValue<float>("PlaybackTime");
+			UpdateInterval = EventData.GetValue<float>("UpdateInterval");
+
+			LocationX = EventData.GetValue<double>("LocationX");
+			LocationY = EventData.GetValue<double>("LocationY");
+			LocationZ = EventData.GetValue<double>("LocationZ");
+
+			RotatorPitch = EventData.GetValue<double>("RotatorPitch");
+			RotatorYaw = EventData.GetValue<double>("RotatorYaw");
+			RotatorRoll = EventData.GetValue<double>("RotatorRoll");
+		}
+
+		float TimeVirtualized = 0.0f;
+		float PlaybackTime = 0.0f;
+		float UpdateInterval = 0.0f;
+
+		double LocationX = 0.0;
+		double LocationY = 0.0;
+		double LocationZ = 0.0;
+
+		double RotatorPitch = 0.0;
+		double RotatorYaw = 0.0;
+		double RotatorRoll = 0.0;
+	};
+
+	class FVirtualLoopDashboardEntry : public FSoundAssetDashboardEntry
+	{
+	public:
+		FVirtualLoopDashboardEntry() = default;
+		virtual ~FVirtualLoopDashboardEntry() = default;
+
+		virtual bool IsValid() const override
+		{
+			return PlayOrder != static_cast<uint32>(INDEX_NONE);
+		}
+
+		uint32 PlayOrder = static_cast<uint32>(INDEX_NONE);
+		uint64 ComponentId = TNumericLimits<uint64>::Max();
+
+		float TimeVirtualized = 0.0f;
+		float PlaybackTime = 0.0f;
+		float UpdateInterval = 0.0f;
+
+		FVector Location = FVector::ZeroVector;
+		FRotator Rotator = FRotator::ZeroRotator;
+	};
+
+	class FVirtualLoopMessages
+	{
+		TAnalyzerMessageQueue<FVirtualLoopVirtualizeMessage> VirtualizeMessages { 5.0 };
+		TAnalyzerMessageQueue<FVirtualLoopRealizeMessage> RealizeMessages { 5.0 };
+		TAnalyzerMessageQueue<FVirtualLoopStopMessage> StopMessages { 5.0 };
+		TAnalyzerMessageQueue<FVirtualLoopUpdateMessage> UpdateMessages { 0.1 };
+
+		friend class FVirtualLoopTraceProvider;
+	};
+} // namespace UE::Audio::Insights
diff --git a/Engine/Plugins/AudioInsights/Source/AudioInsights/Public/Views/LogDashboardViewFactory.h b/Engine/Plugins/AudioInsights/Source/AudioInsights/Public/Views/LogDashboardViewFactory.h
new file mode 100644
--- /dev/null
+++ b/Engine/Plugins/AudioInsights/Source/AudioInsights/Public/Views/LogDashboardViewFactory.h
@@ -0,0 +1,49 @@
+// Copyright Epic Games, Inc. All Rights Reserved.
+#pragma once
+
+#include "Internationalization/Text.h"
+#include "Misc/OutputDevice.h"
+#include "OutputLogCreationParams.h"
+#include "Templates/SharedPointer.h"
+#include "Textures/SlateIcon.h"
+#include "UObject/NameTypes.h"
+#include "Views/DashboardViewFactory.h"
+#include "Widgets/SWidget.h"
+
+#define LOCTEXT_NAMESPACE "AudioInsights"
+
+
+namespace UE::Audio::Insights
+{
+	class AUDIOINSIGHTS_API FLogDashboardViewFactory : public IDashboardViewFactory
+	{
+	public:
+		FLogDashboardViewFactory();
+
+		virtual FName GetName() const override;
+		virtual FText GetDisplayName() const override;
+		virtual EDefaultDashboardTabStack GetDefaultTabStack() const override;
+		virtual FSlateIcon GetIcon() const override;
+		virtual TSharedRef<SWidget> MakeWidget() override;
+
+	private:
+		struct FLogCategoryCollector : public FOutputDevice
+		{
+			FLogCategoryCollector();
+			~FLogCategoryCollector();
+
+			FDefaultCategorySelectionMap GetCollectedCategories() const;
+			void RunAsync();
+
+			virtual bool IsMemoryOnly() const override;
+			virtual void Serialize(const TCHAR* InMsg, ELogVerbosity::Type Verbosity, const FName& InCategory) override;
+
+		private:
+			mutable FCriticalSection CollectionCritSec;
+			FDefaultCategorySelectionMap CollectedCategories;
+		};
+
+		FLogCategoryCollector CategoryCollector;
+	};
+} // namespace UE::Audio::Insights
+#undef LOCTEXT_NAMESPACE
diff --git a/Engine/Plugins/AudioInsights/Source/AudioInsights/Public/Views/MixerSourceDashboardViewFactory.h b/Engine/Plugins/AudioInsights/Source/AudioInsights/Public/Views/MixerSourceDashboardViewFactory.h
new file mode 100644
--- /dev/null
+++ b/Engine/Plugins/AudioInsights/Source/AudioInsights/Public/Views/MixerSourceDashboardViewFactory.h
@@ -0,0 +1,117 @@
+// Copyright Epic Games, Inc. All Rights Reserved.
+#pragma once
+
+#include "Messages/MixerSourceTraceMessages.h"
+#include "Views/SAudioCurveView.h"
+#include "Views/TableDashboardViewFactory.h"
+#include "Widgets/Input/SCheckBox.h"
+
+namespace UE::Audio::Insights
+{
+	class AUDIOINSIGHTS_API FMixerSourceDashboardViewFactory : public FTraceObjectTableDashboardViewFactory
+	{
+	public:
+		FMixerSourceDashboardViewFactory();
+		struct FPlotColumnInfo
+		{
+			const TFunctionRef<const ::Audio::TCircularAudioBuffer<FDataPoint>&(const IDashboardDataViewEntry& InData)> DataFunc;
+			const FNumberFormattingOptions* FormatOptions;
+		};
+
+		virtual ~FMixerSourceDashboardViewFactory() = default;
+
+		virtual FName GetName() const override;
+		virtual FText GetDisplayName() const override;
+		virtual FSlateIcon GetIcon() const override;
+		virtual EDefaultDashboardTabStack GetDefaultTabStack() const override;
+		virtual void ProcessEntries(FTraceTableDashboardViewFactory::EProcessReason Reason) override;
+		virtual TSharedRef<SWidget> MakeWidget() override;
+
+		// Maximum amount of data history kept for plots (in seconds)
+		static const double MaxPlotHistorySeconds;
+		// Maximum number of sources to plot at once 
+		static const int32 MaxPlotSources; 
+
+#if WITH_EDITOR
+		DECLARE_MULTICAST_DELEGATE_ThreeParams(FOnUpdateMuteSoloState, ECheckBoxState /*MuteState*/, ECheckBoxState /*SoloState*/, const FString& /*CurrentFilterString*/);
+		inline static FOnUpdateMuteSoloState OnUpdateMuteSoloState;
+#endif // WITH_EDITOR
+
+	protected:
+		virtual const TMap<FName, FTraceTableDashboardViewFactory::FColumnData>& GetColumns() const override;
+		virtual void SortTable() override;
+
+		TSharedRef<SWidget> MakePlotsWidget();
+
+	private:
+		void ResetPlots();
+
+#if WITH_EDITOR
+		void OnPIEStarted(bool bSimulating);
+		void OnPIEStopped(bool bSimulating);
+		void OnPIEPaused(bool bSimulating);
+		void OnPIEResumed(bool bSimulating);
+#endif // WITH_EDITOR
+
+		void UpdatePlotsWidgetsData();
+
+#if WITH_EDITOR
+		void UpdateMuteSoloState();
+#endif // WITH_EDITOR
+
+		// Column information used by plot widgets, keyed by column name. These keys should be a subset of the keys in GetColumns(). 
+		const TMap<FName, FPlotColumnInfo>& GetPlotColumnInfo();
+		const TFunctionRef<const ::Audio::TCircularAudioBuffer<FDataPoint>&(const IDashboardDataViewEntry& InData)> GetPlotColumnDataFunc(const FName& ColumnName);
+		const FNumberFormattingOptions* GetPlotColumnNumberFormat(const FName& ColumnName);
+		const FText GetPlotColumnDisplayName(const FName& ColumnName);
+
+#if WITH_EDITOR
+		TSharedRef<SWidget> MakeMuteSoloWidget();
+
+		void ToggleMuteForAllItems(ECheckBoxState NewState);
+		void ToggleSoloForAllItems(ECheckBoxState NewState);
+#endif // WITH EDITOR
+
+		using FPlotCurvePoint = SAudioCurveView::FCurvePoint;
+		// Map of source id to data point array 
+		using FPointDataPerCurveMap = TMap<int32, TArray<FPlotCurvePoint>>;
+		using FPlotCurveMetadata = SAudioCurveView::FCurveMetadata;
+
+		FCheckBoxStyle MuteToggleButtonStyle;
+		FCheckBoxStyle SoloToggleButtonStyle;
+
+		TSharedPtr<SCheckBox> MuteToggleButton;
+		TSharedPtr<SCheckBox> SoloToggleButton;
+
+		// Curve points per timestamp per source id per column name 
+		TMap<FName, TSharedPtr<FPointDataPerCurveMap>> PlotWidgetCurveIdToPointDataMapPerColumn;
+		// SourceId to metadata for the corresponding curve
+		TSharedPtr<TMap<int32, FPlotCurveMetadata>> PlotWidgetMetadataPerCurve;
+		
+		// Column names for plot selector widget 
+		TArray<FName> ColumnNames;
+		
+		enum class EPIEState : uint8
+		{
+			Running,
+			Stopped,
+			Paused
+		};
+
+		EPIEState PIEState = EPIEState::Stopped;
+
+		double BeginTimestamp = TNumericLimits<double>::Max();
+		double CurrentTimestamp = 0.0;
+
+		const static int32 NumPlotWidgets = 1;
+		TArray<FName> SelectedPlotColumnNames;
+		TArray<TSharedPtr<SAudioCurveView>> PlotWidgets;
+
+#if WITH_EDITOR
+		// State of the mute and solo buttons
+		ECheckBoxState MuteState = ECheckBoxState::Unchecked;
+		ECheckBoxState SoloState = ECheckBoxState::Unchecked;
+		FString CurrentFilterString;
+#endif // WITH_EDITOR
+	};
+} // namespace UE::Audio::Insights
diff --git a/Engine/Plugins/AudioInsights/Source/AudioInsights/Public/Views/SAudioCurveView.h b/Engine/Plugins/AudioInsights/Source/AudioInsights/Public/Views/SAudioCurveView.h
new file mode 100644
--- /dev/null
+++ b/Engine/Plugins/AudioInsights/Source/AudioInsights/Public/Views/SAudioCurveView.h
@@ -0,0 +1,110 @@
+// Copyright Epic Games, Inc. All Rights Reserved.
+
+#pragma once
+
+#include "SSimpleTimeSlider.h"
+#include "Widgets/SCompoundWidget.h"
+#include "Widgets/DeclarativeSyntaxSupport.h"
+
+class FPaintArgs;
+class FSlateWindowElementList;
+class SToolTip;
+
+// Loosely based on SCurveTimelineView
+// TODO move this to a more general place like the Audio Widgets plugin once it's further developed
+class SAudioCurveView : public SCompoundWidget
+{
+public:
+	// A curve point is a (double XValue, float YValue) pair
+	using FCurvePoint = TPair<double, float>;
+
+	struct FCurveMetadata
+	{
+		int32 CurveId; 
+		FLinearColor CurveColor;
+		FText DisplayName;
+	};
+
+	SLATE_BEGIN_ARGS(SAudioCurveView)
+		: _ViewRange(TRange<double>(0,5))
+		, _YMargin(0.05f)
+		, _HorizontalAxisIncrement(0.5)
+		, _GridLineColor(FLinearColor(0.5f, 0.5f, 0.5f, 0.25f))
+		, _AxesLabelColor(FLinearColor(0.5f, 0.5f, 0.75f))
+		, _DesiredSize(FVector2D(100.f,100.f))
+	{}
+    	/** View X axis range (in value space) */
+    	SLATE_ATTRIBUTE(TRange<double>, ViewRange);
+
+		/** Margin for Y axis, as a 0 - 0.5f proportion, for the space each of above and below the data range. (ex. 0.05 means a 5% margin on the top and bottom, with 90% of the widget's vertical size corresponding to the data range). */
+		SLATE_ATTRIBUTE(float, YMargin);
+
+		/** X axis increment for grid lines. */
+		SLATE_ATTRIBUTE(double, HorizontalAxisIncrement);
+
+		SLATE_ATTRIBUTE(FLinearColor, GridLineColor);
+
+		SLATE_ATTRIBUTE(FLinearColor, AxesLabelColor);
+
+		/** Desired widget size */
+		SLATE_ATTRIBUTE(FVector2D, DesiredSize);
+	
+	SLATE_END_ARGS()
+
+	/**
+	 * Construct the widget
+	 * 
+	 * @param InArgs   A declaration from which to construct the widget
+	 */
+	void Construct( const FArguments& InArgs );
+
+	void SetYValueFormattingOptions(const FNumberFormattingOptions InValueFormattingOptions);
+	void SetCurvesPointData(TSharedPtr<TMap<int32, TArray<FCurvePoint>>> InCurvesPointData);
+	void SetCurvesMetadata(TSharedPtr<TMap<int32, FCurveMetadata>> InMetadataPerCurve);
+
+	FText GetCurveToolTipXValueText() const { return CurveToolTipXValueText; }
+	FText GetCurveToolTipYValueText() const { return CurveToolTipYValueText; }
+	FText GetCurveToolTipDisplayNameText() const { return CurveToolTipDisplayNameText; }
+
+	/** Helper functions for converting between widget local Y position and a given data value (within DataRange). */
+	float ValueToLocalY(const FVector2f AllottedLocalSize, const float Value) const;
+	float LocalYToValue(const FVector2f AllottedLocalSize, const float LocalY) const;
+	
+protected:
+	// SWidget interface
+	virtual FVector2D ComputeDesiredSize(float) const override;
+	virtual int32 OnPaint( const FPaintArgs& Args, const FGeometry& AllottedGeometry, const FSlateRect& MyCullingRect, FSlateWindowElementList& OutDrawElements, int32 LayerId, const FWidgetStyle& InWidgetStyle, bool bParentEnabled ) const override;
+	virtual FReply OnMouseMove(const FGeometry& MyGeometry, const FPointerEvent& MouseEvent) override;
+	
+	int32 PaintCurves(const FGeometry& AllottedGeometry, const FSlateRect& MyCullingRect, FSlateWindowElementList& OutDrawElements, int32 LayerId, const FWidgetStyle& InWidgetStyle, bool bParentEnabled) const;
+	int32 PaintGridLines(const FGeometry& AllottedGeometry, const FSlateRect& MyCullingRect, FSlateWindowElementList& OutDrawElements, int32 LayerId, const FWidgetStyle& InWidgetStyle, bool bParentEnabled, const SSimpleTimeSlider::FScrubRangeToScreen RangeToScreen) const;
+	void UpdateCurveToolTip(const FGeometry& MyGeometry, const FPointerEvent& MouseEvent);
+	void UpdateYDataRange();
+	TSharedRef<SToolTip> CreateCurveTooltip();
+
+	TAttribute<TRange<double>> ViewRange;
+	TAttribute<float> YMargin;
+	TAttribute<double> HorizontalAxisIncrement;
+	TAttribute<FLinearColor> GridLineColor;
+	TAttribute<FLinearColor> AxesLabelColor;
+	TAttribute<FVector2D> DesiredSize;
+
+	// Point data and metadata, keyed by curve id
+	TSharedPtr<TMap<int32, TArray<FCurvePoint>>> PointDataPerCurve;
+	TSharedPtr<TMap<int32, FCurveMetadata>> MetadataPerCurve; 
+	// Y axis data range in value space
+	FVector2f YDataRange;
+
+	// Tooltip text
+	FText CurveToolTipXValueText;		
+	FText CurveToolTipYValueText;
+	FText CurveToolTipDisplayNameText;
+	
+	// Tooltip and axis text formatting
+	FNumberFormattingOptions XValueFormattingOptions;
+	FNumberFormattingOptions YValueFormattingOptions;
+
+	ESlateDrawEffect LineDrawEffects;
+	uint32 NumHorizontalGridLines;
+	FSlateFontInfo LabelFont;
+};
diff --git a/Engine/Plugins/AudioInsights/Source/AudioInsights/Public/Views/TableDashboardViewFactory.h b/Engine/Plugins/AudioInsights/Source/AudioInsights/Public/Views/TableDashboardViewFactory.h
--- a/Engine/Plugins/AudioInsights/Source/AudioInsights/Public/Views/TableDashboardViewFactory.h
+++ b/Engine/Plugins/AudioInsights/Source/AudioInsights/Public/Views/TableDashboardViewFactory.h
@@ -3,8 +3,6 @@
 
 #include "AudioDefines.h"
 #include "AudioInsightsDataSource.h"
-#include "Engine/Attenuation.h"
-#include "Engine/World.h"
 #include "Templates/SharedPointer.h"
 #include "Views/DashboardViewFactory.h"
 #include "Widgets/Input/SSearchBox.h"
@@ -131,7 +129,7 @@ namespace UE::Audio::Insights
 
 #if WITH_EDITOR
 		virtual bool IsDebugDrawEnabled() const { return false; }
-		virtual void DebugDraw(float InElapsed, const IDashboardDataViewEntry& InDashboardEntry, ::Audio::FDeviceId DeviceId) const { }
+		virtual void DebugDraw(float InElapsed, const TArray<TSharedPtr<IDashboardDataViewEntry>>& InSelectedItems, ::Audio::FDeviceId InAudioDeviceId) const { };
 #endif // WITH_EDITOR
 
 		struct FColumnData
@@ -166,7 +164,7 @@ namespace UE::Audio::Insights
 		TSharedPtr<SListView<TSharedPtr<IDashboardDataViewEntry>>> FilteredEntriesListView;
 
 		FName SortByColumn;
-		EColumnSortMode::Type SortMode;
+		EColumnSortMode::Type SortMode = EColumnSortMode::None;
 
 	private:
 		TSharedRef<SHeaderRow> MakeHeaderRowWidget();
@@ -190,29 +188,16 @@ namespace UE::Audio::Insights
 		virtual TSharedRef<SWidget> MakeWidget() override;
 
 	protected:
+#if WITH_EDITOR
 		virtual TSharedRef<SWidget> MakeAssetMenuBar() const;
+#endif // WITH_EDITOR
 
 	private:
+#if WITH_EDITOR
 		TArray<UObject*> GetSelectedEditableAssets() const;
 
 		bool OpenAsset() const;
 		bool BrowseToAsset() const;
-	};
-
-	class AUDIOINSIGHTS_API FSoundAttenuationVisualizer
-	{
-		mutable TMultiMap<EAttenuationShape::Type, FBaseAttenuationSettings::AttenuationShapeDetails> ShapeDetailsMap;
-		mutable uint32 LastObjectId = INDEX_NONE;
-
-	public:
-		const FColor Color { 155, 155, 255 };
-
-	public:
-		FSoundAttenuationVisualizer(const FColor& InColor)
-			: Color(InColor)
-		{
-		}
-
-		void Draw(float InDeltaTime, const FTransform& InTransform, const UObject& InObject, const UWorld& InWorld) const;
+#endif // WITH_EDITOR
 	};
 } // namespace UE::Audio::Insights
diff --git a/Engine/Plugins/AudioInsights/Source/AudioInsights/Public/Views/VirtualLoopDashboardViewFactory.h b/Engine/Plugins/AudioInsights/Source/AudioInsights/Public/Views/VirtualLoopDashboardViewFactory.h
new file mode 100644
--- /dev/null
+++ b/Engine/Plugins/AudioInsights/Source/AudioInsights/Public/Views/VirtualLoopDashboardViewFactory.h
@@ -0,0 +1,36 @@
+// Copyright Epic Games, Inc. All Rights Reserved.
+#pragma once
+
+#include "Views/TableDashboardViewFactory.h"
+
+
+namespace UE::Audio::Insights
+{
+	class AUDIOINSIGHTS_API FVirtualLoopDashboardViewFactory : public FTraceObjectTableDashboardViewFactory
+	{
+	public:
+		FVirtualLoopDashboardViewFactory();
+		virtual ~FVirtualLoopDashboardViewFactory() = default;
+
+		virtual FName GetName() const override;
+		virtual FText GetDisplayName() const override;
+		virtual FSlateIcon GetIcon() const override;
+		virtual EDefaultDashboardTabStack GetDefaultTabStack() const override;
+
+#if WITH_EDITOR
+		DECLARE_MULTICAST_DELEGATE_ThreeParams(FOnDebugDrawEntries, float /*InElapsed*/, const TArray<TSharedPtr<IDashboardDataViewEntry>>& /*InSelectedItems*/, ::Audio::FDeviceId /*InAudioDeviceId*/);
+		inline static FOnDebugDrawEntries OnDebugDrawEntries;
+#endif // WITH_EDITOR
+
+	protected:
+		virtual void ProcessEntries(FTraceTableDashboardViewFactory::EProcessReason Reason) override;
+		virtual const TMap<FName, FTraceTableDashboardViewFactory::FColumnData>& GetColumns() const override;
+
+		virtual void SortTable() override;
+
+#if WITH_EDITOR
+		virtual bool IsDebugDrawEnabled() const override;
+		virtual void DebugDraw(float InElapsed, const TArray<TSharedPtr<IDashboardDataViewEntry>>& InSelectedItems, ::Audio::FDeviceId InAudioDeviceId) const;
+#endif // WITH_EDITOR
+	};
+} // namespace UE::Audio::Insights
diff --git a/Engine/Plugins/AudioInsights/Source/AudioInsightsEditor/AudioInsightsEditor.Build.cs b/Engine/Plugins/AudioInsights/Source/AudioInsightsEditor/AudioInsightsEditor.Build.cs
new file mode 100644
--- /dev/null
+++ b/Engine/Plugins/AudioInsights/Source/AudioInsightsEditor/AudioInsightsEditor.Build.cs
@@ -0,0 +1,41 @@
+// Copyright Epic Games, Inc. All Rights Reserved.
+
+using UnrealBuildTool;
+
+public class AudioInsightsEditor : ModuleRules
+{
+	public AudioInsightsEditor(ReadOnlyTargetRules Target) : base(Target)
+	{
+		PCHUsage = ModuleRules.PCHUsageMode.UseExplicitOrSharedPCHs;
+
+		PublicDependencyModuleNames.AddRange
+		(
+			new string[]
+			{
+				"AudioMixerCore",
+				"Core",
+			}
+		);	
+		
+		PrivateDependencyModuleNames.AddRange
+		(
+			new string[]
+			{
+				"AudioInsights",
+				"AudioMixer",
+				"AudioMixerCore",
+				"AudioWidgets",
+				"CoreUObject",
+				"Engine",
+				"InputCore",
+				"OutputLog",
+				"SignalProcessing",
+				"Slate",
+				"SlateCore",
+				"ToolMenus",
+				"UnrealEd",
+				"WorkspaceMenuStructure",
+			}
+		);
+	}
+}
diff --git a/Engine/Plugins/AudioInsights/Source/AudioInsightsEditor/Private/Analyzers/AudioMeterAnalyzer.cpp b/Engine/Plugins/AudioInsights/Source/AudioInsightsEditor/Private/Analyzers/AudioMeterAnalyzer.cpp
new file mode 100644
--- /dev/null
+++ b/Engine/Plugins/AudioInsights/Source/AudioInsightsEditor/Private/Analyzers/AudioMeterAnalyzer.cpp
@@ -0,0 +1,61 @@
+// Copyright Epic Games, Inc. All Rights Reserved.
+#include "AudioMeterAnalyzer.h"
+
+#include "AudioDeviceManager.h"
+#include "AudioInsightsEditorModule.h"
+#include "AudioMeter.h"
+#include "AudioMixerDevice.h"
+
+namespace UE::Audio::Insights
+{
+	namespace FAudioMeterAnalyzerPrivate
+	{
+		TSharedRef<AudioWidgets::FAudioMeter> CreateAudioMeter(TWeakObjectPtr<UAudioBus> InExternalAudioBus)
+		{
+			using namespace ::Audio;
+
+			const FAudioInsightsEditorModule AudioInsightsEditorModule = FAudioInsightsEditorModule::GetChecked();
+			const FDeviceId AudioDeviceId = AudioInsightsEditorModule.GetDeviceId();
+
+			if (const FAudioDeviceManager* AudioDeviceManager = FAudioDeviceManager::Get())
+			{
+				if (const FMixerDevice* MixerDevice = static_cast<const FMixerDevice*>(AudioDeviceManager->GetAudioDeviceRaw(AudioDeviceId)))
+				{
+					return MakeShared<AudioWidgets::FAudioMeter>(InExternalAudioBus.IsValid() ? InExternalAudioBus->GetNumChannels() : MixerDevice->GetNumDeviceChannels(),
+						AudioDeviceId,
+						InExternalAudioBus.Get());
+				}
+			}
+			
+			return MakeShared<AudioWidgets::FAudioMeter>(1, AudioDeviceId);
+		}
+	}
+
+	FAudioMeterAnalyzer::FAudioMeterAnalyzer(TWeakObjectPtr<UAudioBus> InExternalAudioBus)
+		: AudioMeter(FAudioMeterAnalyzerPrivate::CreateAudioMeter(InExternalAudioBus))
+	{
+		
+	}
+
+	void FAudioMeterAnalyzer::RebuildAudioMeter(TWeakObjectPtr<UAudioBus> InExternalAudioBus)
+	{
+		using namespace ::Audio;
+
+		FAudioDeviceManager* AudioDeviceManager = FAudioDeviceManager::Get();
+		if (!AudioDeviceManager)
+		{
+			return;
+		}
+
+		const FAudioInsightsEditorModule AudioInsightsEditorModule = FAudioInsightsEditorModule::GetChecked();
+		const FDeviceId AudioDeviceId = AudioInsightsEditorModule.GetDeviceId();
+
+		const FMixerDevice* MixerDevice = static_cast<const FMixerDevice*>(AudioDeviceManager->GetAudioDeviceRaw(AudioDeviceId));
+		if (!MixerDevice)
+		{
+			return;
+		}
+
+		AudioMeter->Init(InExternalAudioBus.IsValid() ? InExternalAudioBus->GetNumChannels() : MixerDevice->GetNumDeviceChannels(), AudioDeviceId, InExternalAudioBus.Get());
+	}
+} // namespace UE::Audio::Insights
diff --git a/Engine/Plugins/AudioInsights/Source/AudioInsightsEditor/Private/Analyzers/AudioMeterAnalyzer.h b/Engine/Plugins/AudioInsights/Source/AudioInsightsEditor/Private/Analyzers/AudioMeterAnalyzer.h
new file mode 100644
--- /dev/null
+++ b/Engine/Plugins/AudioInsights/Source/AudioInsightsEditor/Private/Analyzers/AudioMeterAnalyzer.h
@@ -0,0 +1,26 @@
+// Copyright Epic Games, Inc. All Rights Reserved.
+#pragma once
+
+#include "Templates/SharedPointer.h"
+#include "UObject/WeakObjectPtr.h"
+
+namespace AudioWidgets { class FAudioMeter; }
+
+class UAudioBus;
+
+namespace UE::Audio::Insights
+{
+	class FAudioMeterAnalyzer : public TSharedFromThis<FAudioMeterAnalyzer>
+	{
+	public:
+		explicit FAudioMeterAnalyzer(TWeakObjectPtr<UAudioBus> InExternalAudioBus = nullptr);
+		virtual ~FAudioMeterAnalyzer() = default;
+
+		TSharedRef<AudioWidgets::FAudioMeter> GetAudioMeter() { return AudioMeter; };
+
+		void RebuildAudioMeter(TWeakObjectPtr<UAudioBus> InExternalAudioBus = nullptr);
+
+	private:
+		TSharedRef<AudioWidgets::FAudioMeter> AudioMeter;
+	};
+} // namespace UE::Audio::Insights
diff --git a/Engine/Plugins/AudioInsights/Source/AudioInsightsEditor/Private/Analyzers/AudioMeterSubmixAnalyzer.cpp b/Engine/Plugins/AudioInsights/Source/AudioInsightsEditor/Private/Analyzers/AudioMeterSubmixAnalyzer.cpp
new file mode 100644
--- /dev/null
+++ b/Engine/Plugins/AudioInsights/Source/AudioInsightsEditor/Private/Analyzers/AudioMeterSubmixAnalyzer.cpp
@@ -0,0 +1,140 @@
+// Copyright Epic Games, Inc. All Rights Reserved.
+#include "AudioMeterSubmixAnalyzer.h"
+
+#include "AudioBusSubsystem.h"
+#include "AudioDefines.h"
+#include "AudioDeviceManager.h"
+#include "AudioInsightsEditorModule.h"
+#include "AudioMeter.h"
+#include "AudioMixerDevice.h"
+#include "AudioMixerSubmix.h"
+
+namespace UE::Audio::Insights
+{
+	FAudioMeterSubmixAnalyzer::FAudioMeterSubmixAnalyzer(TWeakObjectPtr<USoundSubmix> InSoundSubmix)
+	{
+		SetSubmix(InSoundSubmix);
+	}
+
+	FAudioMeterSubmixAnalyzer::~FAudioMeterSubmixAnalyzer()
+	{
+		UnregisterAudioBusFromSubmix();
+	}
+
+	void FAudioMeterSubmixAnalyzer::SetSubmix(TWeakObjectPtr<USoundSubmix> InSoundSubmix)
+	{
+		using namespace ::Audio;
+
+		if (SoundSubmix.IsValid())
+		{
+			UnregisterAudioBusFromSubmix();
+			SoundSubmix.Reset();
+		}
+
+		AudioMeterAnalyzer.RebuildAudioMeter();
+
+		SoundSubmix = InSoundSubmix;
+
+		if (!SoundSubmix.IsValid())
+		{
+			return;
+		}
+		
+		FAudioDeviceManager* AudioDeviceManager = FAudioDeviceManager::Get();
+		if (!AudioDeviceManager)
+		{
+			return;
+		}
+
+		const FAudioInsightsEditorModule AudioInsightsEditorModule = FAudioInsightsEditorModule::GetChecked();
+		const FDeviceId AudioDeviceId = AudioInsightsEditorModule.GetDeviceId();
+
+		const FMixerDevice* MixerDevice = static_cast<const FMixerDevice*>(AudioDeviceManager->GetAudioDeviceRaw(AudioDeviceId));
+		if (!MixerDevice)
+		{
+			return;
+		}
+
+		FMixerSubmixWeakPtr MixerSubmixWeakPtr = MixerDevice->GetSubmixInstance(SoundSubmix.Get());
+		if (!MixerSubmixWeakPtr.IsValid())
+		{
+			return;
+		}
+
+		// Register audio bus in submix
+		const TObjectPtr<UAudioBus> AudioBus = AudioMeterAnalyzer.GetAudioMeter()->GetAudioBus();
+		if (!AudioBus)
+		{
+			return;
+		}
+
+		const FAudioBusKey AudioBusKey(AudioBus->GetUniqueID());
+		const int32 AudioBusNumChannels = AudioBus->GetNumChannels();
+
+		FAudioThread::RunCommandOnAudioThread([MixerDevice, MixerSubmixWeakPtr, AudioBusKey, AudioBusNumChannels]()
+		{
+			TObjectPtr<UAudioBusSubsystem> AudioBusSubsystem = MixerDevice->GetSubsystem<UAudioBusSubsystem>();
+			check(AudioBusSubsystem);
+
+			if (FMixerSubmixPtr MixerSubmix = MixerSubmixWeakPtr.Pin();
+				MixerSubmix.IsValid())
+			{
+				MixerSubmix->RegisterAudioBus(AudioBusKey, AudioBusSubsystem->AddPatchInputForAudioBus(AudioBusKey, MixerDevice->GetNumOutputFrames(), AudioBusNumChannels));
+			}
+		});
+	}
+
+	void FAudioMeterSubmixAnalyzer::UnregisterAudioBusFromSubmix()
+	{
+		using namespace ::Audio;
+
+		FAudioDeviceManager* AudioDeviceManager = FAudioDeviceManager::Get();
+		if (!AudioDeviceManager)
+		{
+			return;
+		}
+
+		const FAudioInsightsEditorModule AudioInsightsEditorModule = FAudioInsightsEditorModule::GetChecked();
+		const FDeviceId AudioDeviceId = AudioInsightsEditorModule.GetDeviceId();
+
+		const FMixerDevice* MixerDevice = static_cast<const FMixerDevice*>(AudioDeviceManager->GetAudioDeviceRaw(AudioDeviceId));
+		if (!MixerDevice)
+		{
+			return;
+		}
+
+		if (!SoundSubmix.IsValid())
+		{
+			return;
+		}
+
+		FMixerSubmixWeakPtr MixerSubmixWeakPtr = MixerDevice->GetSubmixInstance(SoundSubmix.Get());
+		if (!MixerSubmixWeakPtr.IsValid())
+		{
+			return;
+		}
+
+		// Unregister audio bus from submix
+		const TObjectPtr<UAudioBus> AudioBus = AudioMeterAnalyzer.GetAudioMeter()->GetAudioBus();
+		if (!AudioBus)
+		{
+			return;
+		}
+
+		const FAudioBusKey AudioBusKey(AudioBus->GetUniqueID());
+
+		FAudioThread::RunCommandOnAudioThread([MixerSubmixWeakPtr, AudioBusKey]()
+		{
+			if (FMixerSubmixPtr MixerSubmix = MixerSubmixWeakPtr.Pin();
+				MixerSubmix.IsValid())
+			{
+				MixerSubmix->UnregisterAudioBus(AudioBusKey);
+			}
+		});
+	}
+
+	TSharedRef<SAudioMeter> FAudioMeterSubmixAnalyzer::GetWidget()
+	{
+		return AudioMeterAnalyzer.GetAudioMeter()->GetWidget();
+	}
+} // namespace UE::Audio::Insights
diff --git a/Engine/Plugins/AudioInsights/Source/AudioInsightsEditor/Private/Analyzers/AudioMeterSubmixAnalyzer.h b/Engine/Plugins/AudioInsights/Source/AudioInsightsEditor/Private/Analyzers/AudioMeterSubmixAnalyzer.h
new file mode 100644
--- /dev/null
+++ b/Engine/Plugins/AudioInsights/Source/AudioInsightsEditor/Private/Analyzers/AudioMeterSubmixAnalyzer.h
@@ -0,0 +1,31 @@
+// Copyright Epic Games, Inc. All Rights Reserved.
+#pragma once
+
+#include "AudioMeterAnalyzer.h"
+#include "Templates/SharedPointer.h"
+#include "UObject/WeakObjectPtr.h"
+
+namespace AudioWidgets { class FAudioMeter; }
+
+class SAudioMeter;
+class USoundSubmix;
+
+namespace UE::Audio::Insights
+{
+	class FAudioMeterSubmixAnalyzer final : public TSharedFromThis<FAudioMeterSubmixAnalyzer>
+	{
+	public:
+		FAudioMeterSubmixAnalyzer(TWeakObjectPtr<USoundSubmix> InSoundSubmix);
+		~FAudioMeterSubmixAnalyzer();
+
+		// Rebuilds the audio meter and registers the newly selected submix
+		void SetSubmix(TWeakObjectPtr<USoundSubmix> InSoundSubmix);
+
+		TSharedRef<SAudioMeter> GetWidget();
+	private:
+		void UnregisterAudioBusFromSubmix();
+
+		TWeakObjectPtr<USoundSubmix> SoundSubmix;
+		FAudioMeterAnalyzer AudioMeterAnalyzer;
+	};
+} // namespace UE::Audio::Insights
diff --git a/Engine/Plugins/AudioInsights/Source/AudioInsightsEditor/Private/Analyzers/AudioOscilloscopeAnalyzer.cpp b/Engine/Plugins/AudioInsights/Source/AudioInsightsEditor/Private/Analyzers/AudioOscilloscopeAnalyzer.cpp
new file mode 100644
--- /dev/null
+++ b/Engine/Plugins/AudioInsights/Source/AudioInsightsEditor/Private/Analyzers/AudioOscilloscopeAnalyzer.cpp
@@ -0,0 +1,173 @@
+// Copyright Epic Games, Inc. All Rights Reserved.
+#include "AudioOscilloscopeAnalyzer.h"
+
+#include "AudioBusSubsystem.h"
+#include "AudioDefines.h"
+#include "AudioDeviceManager.h"
+#include "AudioInsightsEditorModule.h"
+#include "AudioOscilloscope.h"
+#include "AudioMixerDevice.h"
+#include "AudioMixerSubmix.h"
+
+namespace UE::Audio::Insights
+{
+	namespace FAudioOscilloscopeAnalyzerPrivate
+	{
+		TSharedRef<AudioWidgets::FAudioOscilloscope> CreateAudioOscilloscope(const float InTimeWindowMs,
+			const float InMaxTimeWindowMs,
+			const float InAnalysisPeriodMs,
+			const EAudioPanelLayoutType InPanelLayoutType)
+		{
+			using namespace ::Audio;
+
+			const FAudioInsightsEditorModule AudioInsightsEditorModule = FAudioInsightsEditorModule::GetChecked();
+			const FDeviceId AudioDeviceId = AudioInsightsEditorModule.GetDeviceId();
+
+			if (FAudioDeviceManager* AudioDeviceManager = FAudioDeviceManager::Get())
+			{
+				if (const FMixerDevice* MixerDevice = static_cast<const FMixerDevice*>(AudioDeviceManager->GetAudioDeviceRaw(AudioDeviceId)))
+				{
+					return MakeShared<AudioWidgets::FAudioOscilloscope>(AudioDeviceId,
+						MixerDevice->GetNumDeviceChannels(),
+						InTimeWindowMs,
+						InMaxTimeWindowMs,
+						InAnalysisPeriodMs,
+						InPanelLayoutType);
+				}
+			}
+			
+			return MakeShared<AudioWidgets::FAudioOscilloscope>(AudioDeviceId, 1, InTimeWindowMs, InMaxTimeWindowMs, InAnalysisPeriodMs, InPanelLayoutType);
+		}
+	}
+
+	FAudioOscilloscopeAnalyzer::FAudioOscilloscopeAnalyzer(TWeakObjectPtr<USoundSubmix> InSoundSubmix)
+		: AudioOscilloscope(FAudioOscilloscopeAnalyzerPrivate::CreateAudioOscilloscope(TimeWindowMs, MaxTimeWindowMs, AnalysisPeriodMs, PanelLayoutType))
+	{
+		RebuildAudioOscilloscope(InSoundSubmix);
+	}
+
+	FAudioOscilloscopeAnalyzer::~FAudioOscilloscopeAnalyzer()
+	{
+		CleanupAudioOscilloscope();
+	}
+
+	void FAudioOscilloscopeAnalyzer::RebuildAudioOscilloscope(TWeakObjectPtr<USoundSubmix> InSoundSubmix)
+	{
+		using namespace ::Audio;
+
+		if (SoundSubmix.IsValid())
+		{
+			CleanupAudioOscilloscope();
+		}
+
+		SoundSubmix = InSoundSubmix;
+
+		FAudioDeviceManager* AudioDeviceManager = FAudioDeviceManager::Get();
+		if (!AudioDeviceManager)
+		{
+			return;
+		}
+
+		const FAudioInsightsEditorModule AudioInsightsEditorModule = FAudioInsightsEditorModule::GetChecked();
+		const FDeviceId AudioDeviceId = AudioInsightsEditorModule.GetDeviceId();
+
+		const FMixerDevice* MixerDevice = static_cast<FMixerDevice*>(AudioDeviceManager->GetAudioDeviceRaw(AudioDeviceId));
+		if (!MixerDevice)
+		{
+			return;
+		}
+
+		if (!SoundSubmix.IsValid())
+		{
+			return;
+		}
+
+		FMixerSubmixWeakPtr MixerSubmixWeakPtr = MixerDevice->GetSubmixInstance(SoundSubmix.Get());
+		if (!MixerSubmixWeakPtr.IsValid())
+		{
+			return;
+		}
+
+		AudioOscilloscope->CreateDataProvider(AudioDeviceId, TimeWindowMs, MaxTimeWindowMs, AnalysisPeriodMs, PanelLayoutType);
+		AudioOscilloscope->CreateOscilloscopeWidget(MixerDevice->GetNumDeviceChannels(), PanelLayoutType);
+
+		// Start processing
+		AudioOscilloscope->StartProcessing();
+		
+		// Register audio bus in submix
+		const TObjectPtr<UAudioBus> AudioBus = AudioOscilloscope->GetAudioBus();
+		if (!AudioBus)
+		{
+			return;
+		}
+
+		const FAudioBusKey AudioBusKey(AudioBus->GetUniqueID());
+		const int32 AudioBusNumChannels = AudioBus->GetNumChannels();
+
+		FAudioThread::RunCommandOnAudioThread([MixerDevice, MixerSubmixWeakPtr, AudioBusKey, AudioBusNumChannels]()
+		{
+			TObjectPtr<UAudioBusSubsystem> AudioBusSubsystem = MixerDevice->GetSubsystem<UAudioBusSubsystem>();
+			check(AudioBusSubsystem);
+
+			if (FMixerSubmixPtr MixerSubmix = MixerSubmixWeakPtr.Pin();
+				MixerSubmix.IsValid())
+			{
+				MixerSubmix->RegisterAudioBus(AudioBusKey, AudioBusSubsystem->AddPatchInputForAudioBus(AudioBusKey, MixerDevice->GetNumOutputFrames(), AudioBusNumChannels));
+			}
+		});
+	}
+
+	void FAudioOscilloscopeAnalyzer::CleanupAudioOscilloscope()
+	{
+		using namespace ::Audio;
+
+		FAudioDeviceManager* AudioDeviceManager = FAudioDeviceManager::Get();
+		if (!AudioDeviceManager)
+		{
+			return;
+		}
+
+		const FAudioInsightsEditorModule AudioInsightsEditorModule = FAudioInsightsEditorModule::GetChecked();
+		const FDeviceId AudioDeviceId = AudioInsightsEditorModule.GetDeviceId();
+
+		const FMixerDevice* MixerDevice = static_cast<FMixerDevice*>(AudioDeviceManager->GetAudioDeviceRaw(AudioDeviceId));
+		if (!MixerDevice)
+		{
+			return;
+		}
+
+		if (!SoundSubmix.IsValid())
+		{
+			return;
+		}
+
+		FMixerSubmixWeakPtr MixerSubmixWeakPtr = MixerDevice->GetSubmixInstance(SoundSubmix.Get());
+		if (!MixerSubmixWeakPtr.IsValid())
+		{
+			return;
+		}
+
+		// Unregister audio bus from submix
+		const TObjectPtr<UAudioBus> AudioBus = AudioOscilloscope->GetAudioBus();
+		if (!AudioBus)
+		{
+			return;
+		}
+
+		const FAudioBusKey AudioBusKey(AudioBus->GetUniqueID());
+
+		FAudioThread::RunCommandOnAudioThread([MixerSubmixWeakPtr, AudioBusKey]()
+		{
+			if (FMixerSubmixPtr MixerSubmix = MixerSubmixWeakPtr.Pin();
+				MixerSubmix.IsValid())
+			{
+				MixerSubmix->UnregisterAudioBus(AudioBusKey);
+			}
+		});
+
+		// Stop processing
+		AudioOscilloscope->StopProcessing();
+
+		SoundSubmix.Reset();
+	}
+} // namespace UE::Audio::Insights
diff --git a/Engine/Plugins/AudioInsights/Source/AudioInsightsEditor/Private/Analyzers/AudioOscilloscopeAnalyzer.h b/Engine/Plugins/AudioInsights/Source/AudioInsightsEditor/Private/Analyzers/AudioOscilloscopeAnalyzer.h
new file mode 100644
--- /dev/null
+++ b/Engine/Plugins/AudioInsights/Source/AudioInsightsEditor/Private/Analyzers/AudioOscilloscopeAnalyzer.h
@@ -0,0 +1,35 @@
+// Copyright Epic Games, Inc. All Rights Reserved.
+#pragma once
+
+#include "AudioWidgetsEnums.h"
+#include "Templates/SharedPointer.h"
+#include "UObject/WeakObjectPtr.h"
+
+namespace AudioWidgets { class FAudioOscilloscope; }
+
+class USoundSubmix;
+
+namespace UE::Audio::Insights
+{
+	class FAudioOscilloscopeAnalyzer : public TSharedFromThis<FAudioOscilloscopeAnalyzer>
+	{
+	public:
+		FAudioOscilloscopeAnalyzer(TWeakObjectPtr<USoundSubmix> InSoundSubmix);
+		virtual ~FAudioOscilloscopeAnalyzer();
+
+		TSharedRef<AudioWidgets::FAudioOscilloscope> GetAudioOscilloscope() { return AudioOscilloscope; };
+
+		void RebuildAudioOscilloscope(TWeakObjectPtr<USoundSubmix> InSoundSubmix);
+
+	private:
+		void CleanupAudioOscilloscope();
+
+		TSharedRef<AudioWidgets::FAudioOscilloscope> AudioOscilloscope;
+		TWeakObjectPtr<USoundSubmix> SoundSubmix;
+		
+		static constexpr float TimeWindowMs     = 10.0f;
+		static constexpr float MaxTimeWindowMs  = 10.0f;
+		static constexpr float AnalysisPeriodMs = 10.0f;
+		static constexpr EAudioPanelLayoutType PanelLayoutType = EAudioPanelLayoutType::Basic;
+	};
+} // namespace UE::Audio::Insights
diff --git a/Engine/Plugins/AudioInsights/Source/AudioInsightsEditor/Private/AudioInsightsEditorDashboardFactory.cpp b/Engine/Plugins/AudioInsights/Source/AudioInsightsEditor/Private/AudioInsightsEditorDashboardFactory.cpp
new file mode 100644
--- /dev/null
+++ b/Engine/Plugins/AudioInsights/Source/AudioInsightsEditor/Private/AudioInsightsEditorDashboardFactory.cpp
@@ -0,0 +1,682 @@
+// Copyright Epic Games, Inc. All Rights Reserved.
+#include "AudioInsightsEditorDashboardFactory.h"
+
+#include "AudioDevice.h"
+#include "AudioDeviceManager.h"
+#include "AudioInsightsModule.h"
+#include "AudioInsightsStyle.h"
+#include "Editor.h"
+#include "Engine/World.h"
+#include "Framework/Docking/TabManager.h"
+#include "Internationalization/Text.h"
+#include "IPropertyTypeCustomization.h"
+#include "Kismet2/DebuggerCommands.h"
+#include "Templates/SharedPointer.h"
+#include "ToolMenu.h"
+#include "ToolMenus.h"
+#include "Widgets/Input/SCheckBox.h"
+
+#define LOCTEXT_NAMESPACE "AudioInsights"
+
+namespace UE::Audio::Insights
+{
+	namespace EditorDashboardFactoryPrivate
+	{
+		static const FText ToolName = LOCTEXT("AudioDashboard_ToolName", "Audio Insights");
+
+		static const FName MainToolbarName = "MainToolbar";
+		static const FText MainToolbarDisplayName = LOCTEXT("AudioDashboard_MainToolbarDisplayName", "Dashboard Transport");
+
+		static const FText PreviewDeviceDisplayName = LOCTEXT("AudioDashboard_PreviewDevice", "[Preview Audio]");
+		static const FText DashboardWorldSelectDescription = LOCTEXT("AudioDashboard_SelectWorldDescription", "Select world(s) to monitor (worlds may share audio output).");
+
+		FText GetDebugNameFromDeviceId(::Audio::FDeviceId InDeviceId)
+		{
+			FString WorldName;
+			if (FAudioDeviceManager* DeviceManager = FAudioDeviceManager::Get())
+			{
+				TArray<UWorld*> DeviceWorlds = DeviceManager->GetWorldsUsingAudioDevice(InDeviceId);
+				for (const UWorld* World : DeviceWorlds)
+				{
+					if (!WorldName.IsEmpty())
+					{
+						WorldName += TEXT(", ");
+					}
+					WorldName += World->GetDebugDisplayName();
+				}
+			}
+
+			if (WorldName.IsEmpty())
+			{
+				return PreviewDeviceDisplayName;
+			}
+
+			return FText::FromString(WorldName);
+		}
+	} // namespace EditorDashboardFactoryPrivate
+
+	void FEditorDashboardFactory::OnWorldRegisteredToAudioDevice(const UWorld* InWorld, ::Audio::FDeviceId InDeviceId)
+	{
+		if (InDeviceId != INDEX_NONE)
+		{
+			if (bStartWithPIE)
+			{
+				const FTraceModule& TraceModule = FAudioInsightsModule::GetChecked().GetTraceModule();
+				TraceModule.StartTraceAnalysis();
+
+				ActiveDeviceId = InDeviceId;
+			}
+		}
+
+		RefreshDeviceSelector();
+	}
+
+	void FEditorDashboardFactory::OnPIEStarted(bool bSimulating)
+	{
+		if (bStartWithPIE)
+		{
+			const FTraceModule& TraceModule = FAudioInsightsModule::GetChecked().GetTraceModule();
+			TraceModule.StartTraceAnalysis();
+		}
+	}
+
+	void FEditorDashboardFactory::OnPostPIEStarted(bool bSimulating)
+	{
+		OnActiveAudioDeviceChanged.Broadcast();
+	}
+
+	void FEditorDashboardFactory::OnPIEStopped(bool bSimulating)
+	{
+		if (bStopWithPIE)
+		{
+			const FTraceModule& TraceModule = FAudioInsightsModule::GetChecked().GetTraceModule();
+			TraceModule.StopTraceAnalysis();
+		}
+
+		RefreshDeviceSelector();
+	}
+
+	void FEditorDashboardFactory::OnWorldUnregisteredFromAudioDevice(const UWorld* InWorld, ::Audio::FDeviceId InDeviceId)
+	{
+		RefreshDeviceSelector();
+	}
+
+	void FEditorDashboardFactory::OnDeviceDestroyed(::Audio::FDeviceId InDeviceId)
+	{
+		if (ActiveDeviceId == InDeviceId)
+		{
+			if (FAudioDeviceManager* DeviceManager = FAudioDeviceManager::Get())
+			{
+				ActiveDeviceId = DeviceManager->GetMainAudioDeviceID();
+			}
+		}
+
+		AudioDeviceIds.RemoveAll([InDeviceId](const TSharedPtr<::Audio::FDeviceId>& DeviceIdPtr)
+		{
+			return *DeviceIdPtr.Get() == InDeviceId;
+		});
+
+		if (AudioDeviceComboBox.IsValid())
+		{
+			AudioDeviceComboBox->RefreshOptions();
+		}
+
+		OnActiveAudioDeviceChanged.Broadcast();
+	}
+
+	void FEditorDashboardFactory::RefreshDeviceSelector()
+	{
+		if (FAudioDeviceManager* DeviceManager = FAudioDeviceManager::Get())
+		{
+			if (!DeviceManager->IsValidAudioDevice(ActiveDeviceId))
+			{
+				ActiveDeviceId = DeviceManager->GetMainAudioDeviceID();
+			}
+		}
+
+		AudioDeviceIds.Empty();
+		if (const FAudioDeviceManager* DeviceManager = FAudioDeviceManager::Get())
+		{
+			DeviceManager->IterateOverAllDevices([this, &DeviceManager](::Audio::FDeviceId DeviceId, const FAudioDevice* AudioDevice)
+			{
+				AudioDeviceIds.Add(MakeShared<::Audio::FDeviceId>(DeviceId));
+			});
+		}
+
+		if (AudioDeviceComboBox.IsValid())
+		{
+			AudioDeviceComboBox->RefreshOptions();
+		}
+	}
+
+	void FEditorDashboardFactory::ResetDelegates()
+	{
+		if (OnWorldRegisteredToAudioDeviceHandle.IsValid())
+		{
+			FAudioDeviceWorldDelegates::OnWorldRegisteredToAudioDevice.Remove(OnWorldRegisteredToAudioDeviceHandle);
+			OnWorldRegisteredToAudioDeviceHandle.Reset();
+		}
+
+		if (OnWorldUnregisteredFromAudioDeviceHandle.IsValid())
+		{
+			FAudioDeviceWorldDelegates::OnWorldUnregisteredWithAudioDevice.Remove(OnWorldUnregisteredFromAudioDeviceHandle);
+			OnWorldUnregisteredFromAudioDeviceHandle.Reset();
+		}
+
+		if (OnDeviceDestroyedHandle.IsValid())
+		{
+			FAudioDeviceManagerDelegates::OnAudioDeviceDestroyed.Remove(OnDeviceDestroyedHandle);
+			OnDeviceDestroyedHandle.Reset();
+		}
+
+		if (OnPIEStartedHandle.IsValid())
+		{
+			FEditorDelegates::PreBeginPIE.Remove(OnPIEStartedHandle);
+			OnPIEStartedHandle.Reset();
+		}
+
+		if (OnPostPIEStartedHandle.IsValid())
+		{
+			FEditorDelegates::PostPIEStarted.Remove(OnPostPIEStartedHandle);
+			OnPostPIEStartedHandle.Reset();
+		}
+
+		if (OnPIEStoppedHandle.IsValid())
+		{
+			FEditorDelegates::EndPIE.Remove(OnPIEStoppedHandle);
+			OnPIEStoppedHandle.Reset();
+		}
+	}
+
+	::Audio::FDeviceId FEditorDashboardFactory::GetDeviceId() const
+	{
+		return ActiveDeviceId;
+	}
+
+	TSharedRef<SDockTab> FEditorDashboardFactory::MakeDockTabWidget(const FSpawnTabArgs& Args)
+	{
+		const TSharedRef<SDockTab> DockTab = SNew(SDockTab)
+			.Label(EditorDashboardFactoryPrivate::ToolName)
+			.Clipping(EWidgetClipping::ClipToBounds)
+			.TabRole(ETabRole::NomadTab);
+
+		DashboardTabManager = FGlobalTabmanager::Get()->NewTabManager(DockTab);
+
+		InitDelegates();
+		TabLayout = GetDefaultTabLayout();
+
+		RegisterTabSpawners();
+		RefreshDeviceSelector();
+
+		const TSharedRef<SWidget> TabContent = SNew(SVerticalBox)
+			+ SVerticalBox::Slot()
+			.AutoHeight()
+			[
+				MakeMenuBarWidget()
+			]
+			+ SVerticalBox::Slot()
+			.AutoHeight()
+			[
+				MakeMainToolbarWidget()
+			]
+			+ SVerticalBox::Slot()
+			.AutoHeight()
+			[
+				SNew(SBox)
+				.HeightOverride(4.0f)
+			]
+			+ SVerticalBox::Slot()
+			[
+				DashboardTabManager->RestoreFrom(TabLayout->AsShared(), Args.GetOwnerWindow()).ToSharedRef()
+			];
+
+		DockTab->SetContent(TabContent);
+
+		DockTab->SetOnTabClosed(SDockTab::FOnTabClosedCallback::CreateLambda([this](TSharedRef<SDockTab> TabClosed)
+		{
+			ResetDelegates();
+			UnregisterTabSpawners();
+		}));
+
+		return DockTab;
+	}
+
+	TSharedRef<SWidget> FEditorDashboardFactory::MakeMenuBarWidget()
+	{
+		FMenuBarBuilder MenuBarBuilder = FMenuBarBuilder(TSharedPtr<FUICommandList>());
+
+		MenuBarBuilder.AddPullDownMenu(
+			LOCTEXT("File_MenuLabel", "File"),
+			FText::GetEmpty(),
+			FNewMenuDelegate::CreateLambda([this](FMenuBuilder& MenuBuilder)
+			{
+				MenuBuilder.AddMenuEntry(LOCTEXT("Close_MenuLabel", "Close"),
+					LOCTEXT("Close_MenuLabel_Tooltip", "Closes the Audio Insights dashboard."),
+					FSlateIcon(),
+					FUIAction(FExecuteAction::CreateLambda([this]()
+					{
+						if (DashboardTabManager.IsValid())
+						{
+							if (TSharedPtr<SDockTab> OwnerTab = DashboardTabManager->GetOwnerTab())
+							{
+								OwnerTab->RequestCloseTab();
+							}
+						}
+					}))
+				);
+			}),
+			"File"
+		);
+
+		MenuBarBuilder.AddPullDownMenu(
+			LOCTEXT("ViewMenuLabel", "View"),
+			FText::GetEmpty(),
+			FNewMenuDelegate::CreateLambda([this](FMenuBuilder& MenuBuilder)
+			{
+				for (const auto& KVP : DashboardViewFactories)
+				{
+					const FName& FactoryName = KVP.Key;
+					const TSharedPtr<IDashboardViewFactory>& Factory = KVP.Value;
+
+					MenuBuilder.AddMenuEntry(Factory->GetDisplayName(),
+						FText::GetEmpty(),
+						FSlateStyle::Get().CreateIcon(Factory->GetIcon().GetStyleName()),
+						FUIAction(FExecuteAction::CreateLambda([this, FactoryName]()
+						{
+							if (DashboardTabManager.IsValid())
+							{
+								if (TSharedPtr<SDockTab> ViewportTab = DashboardTabManager->FindExistingLiveTab(FactoryName);
+									!ViewportTab.IsValid())
+								{
+									DashboardTabManager->TryInvokeTab(FactoryName);
+
+									if (TSharedPtr<SDockTab> InvokedOutputMeterTab = DashboardTabManager->TryInvokeTab(FactoryName);
+										InvokedOutputMeterTab.IsValid() && DashboardViewFactories[FactoryName].IsValid())
+									{
+										if (const EDefaultDashboardTabStack DefaultTabStack = DashboardViewFactories[FactoryName]->GetDefaultTabStack();
+											DefaultTabStack == EDefaultDashboardTabStack::AudioMeter ||
+											DefaultTabStack == EDefaultDashboardTabStack::Oscilloscope)
+										{
+											InvokedOutputMeterTab->SetParentDockTabStackTabWellHidden(true);
+										}
+									}
+								}
+								else
+								{
+									ViewportTab->RequestCloseTab();
+								}
+							}
+						}),
+						FCanExecuteAction(),
+						FIsActionChecked::CreateLambda([&DashboardTabManager = DashboardTabManager, FactoryName]()
+						{
+							return DashboardTabManager.IsValid() ? DashboardTabManager->FindExistingLiveTab(FactoryName).IsValid() : false;
+						})),
+						NAME_None,
+						EUserInterfaceActionType::Check
+					);
+
+					if (const EDefaultDashboardTabStack DefaultTabStack = DashboardViewFactories[FactoryName]->GetDefaultTabStack();
+						DefaultTabStack == EDefaultDashboardTabStack::Log || DefaultTabStack == EDefaultDashboardTabStack::AudioMeters)
+					{
+						MenuBuilder.AddMenuSeparator();
+					}
+				}
+			}),
+			"View"
+		);
+
+		return MenuBarBuilder.MakeWidget();
+	}
+
+	TSharedRef<SWidget> FEditorDashboardFactory::MakeMainToolbarWidget()
+	{
+		using namespace EditorDashboardFactoryPrivate;
+
+		static const FName PlayWorldToolBarName = "Kismet.DebuggingViewToolBar";
+		if (!UToolMenus::Get()->IsMenuRegistered(PlayWorldToolBarName))
+		{
+			UToolMenu* ToolBar = UToolMenus::Get()->RegisterMenu(PlayWorldToolBarName, NAME_None, EMultiBoxType::SlimHorizontalToolBar);
+			FToolMenuSection& Section = ToolBar->AddSection("Debug");
+			FPlayWorldCommands::BuildToolbar(Section);
+		}
+
+		static FSlateBrush TransportBackgroundColorBrush;
+		TransportBackgroundColorBrush.TintColor = FSlateColor(FLinearColor(0.018f, 0.018f, 0.018f, 1.0f));
+		TransportBackgroundColorBrush.DrawAs    = ESlateBrushDrawType::Box;
+
+		return SNew(SBorder)
+			.BorderImage(&TransportBackgroundColorBrush)
+			[
+				SNew(SHorizontalBox)
+				+ SHorizontalBox::Slot()
+				.HAlign(HAlign_Left)
+				.VAlign(VAlign_Center)
+				.AutoWidth()
+				[
+					SNew(SBorder)
+					.BorderImage(FAppStyle::Get().GetBrush("NoBorder"))
+					[
+						UToolMenus::Get()->GenerateWidget(PlayWorldToolBarName, { FPlayWorldCommands::GlobalPlayWorldActions })
+					]
+				]
+				+ SHorizontalBox::Slot()
+				.HAlign(HAlign_Right)
+				.VAlign(VAlign_Center)
+				.AutoWidth()
+				.Padding(2.0f, 0.0f)
+				[
+					SNew(STextBlock)
+					.Text(LOCTEXT("StartOnPIE_DisplayName", "Start with PIE:"))
+					.Font(IPropertyTypeCustomizationUtils::GetRegularFont())
+				]
+				+ SHorizontalBox::Slot()
+				.HAlign(HAlign_Left)
+				.VAlign(VAlign_Center)
+				.AutoWidth()
+				.Padding(2.0f, 0.0f)
+				[
+					SNew(SCheckBox)
+					.IsChecked_Lambda([this]() { return bStartWithPIE ? ECheckBoxState::Checked : ECheckBoxState::Unchecked; })
+					.OnCheckStateChanged_Lambda([this](ECheckBoxState NewState) { bStartWithPIE = NewState == ECheckBoxState::Checked; })
+				]
+				+ SHorizontalBox::Slot()
+				.HAlign(HAlign_Right)
+				.VAlign(VAlign_Center)
+				.AutoWidth()
+				.Padding(2.0f, 0.0f)
+				[
+					SNew(STextBlock)
+					.Text(LOCTEXT("StopOnPIE_DisplayName", "Stop with PIE:"))
+					.Font(IPropertyTypeCustomizationUtils::GetRegularFont())
+				]
+				+ SHorizontalBox::Slot()
+				.HAlign(HAlign_Left)
+				.VAlign(VAlign_Center)
+				.AutoWidth()
+				.Padding(2.0f, 0.0f)
+				[
+					SNew(SCheckBox)
+					.IsChecked_Lambda([this]() { return bStopWithPIE ? ECheckBoxState::Checked : ECheckBoxState::Unchecked; })
+					.OnCheckStateChanged_Lambda([this](ECheckBoxState NewState) { bStopWithPIE = NewState == ECheckBoxState::Checked; })
+				]
+				+ SHorizontalBox::Slot()
+				.HAlign(HAlign_Right)
+				.VAlign(VAlign_Center)
+				.Padding(2.0f, 0.0f)
+				[
+					SNew(STextBlock)
+					.Text(LOCTEXT("SelectDashboardWorld_DisplayName", "World Filter:"))
+					.ToolTipText(DashboardWorldSelectDescription)
+					.Font(IPropertyTypeCustomizationUtils::GetRegularFont())
+				]
+				+ SHorizontalBox::Slot()
+				.HAlign(HAlign_Left)
+				.VAlign(VAlign_Center)
+				.AutoWidth()
+				.Padding(2.0f, 0.0f)
+				[
+					SAssignNew(AudioDeviceComboBox, SComboBox<TSharedPtr<::Audio::FDeviceId>>)
+					.ToolTipText(DashboardWorldSelectDescription)
+					.OptionsSource(&AudioDeviceIds)
+					.OnGenerateWidget_Lambda([](const TSharedPtr<::Audio::FDeviceId>& WidgetDeviceId)
+					{
+						FText NameText = GetDebugNameFromDeviceId(*WidgetDeviceId);
+						return SNew(STextBlock)
+							.Text(NameText)
+							.Font(IPropertyTypeCustomizationUtils::GetRegularFont());
+					})
+					.OnSelectionChanged_Lambda([this](TSharedPtr<::Audio::FDeviceId> NewDeviceId, ESelectInfo::Type)
+					{
+						if (NewDeviceId.IsValid())
+						{
+							ActiveDeviceId = *NewDeviceId;
+							RefreshDeviceSelector();
+
+							OnActiveAudioDeviceChanged.Broadcast();
+						}
+					})
+					[
+						SNew(STextBlock)
+						.Font(IPropertyTypeCustomizationUtils::GetRegularFont())
+						.Text_Lambda([this]()
+						{
+							return GetDebugNameFromDeviceId(ActiveDeviceId);
+						})
+					]
+				]
+			];
+	}
+
+	void FEditorDashboardFactory::InitDelegates()
+	{
+		if (!OnWorldRegisteredToAudioDeviceHandle.IsValid())
+		{
+			OnWorldRegisteredToAudioDeviceHandle = FAudioDeviceWorldDelegates::OnWorldRegisteredToAudioDevice.AddSP(this, &FEditorDashboardFactory::OnWorldRegisteredToAudioDevice);
+		}
+
+		if (!OnWorldUnregisteredFromAudioDeviceHandle.IsValid())
+		{
+			OnWorldUnregisteredFromAudioDeviceHandle = FAudioDeviceWorldDelegates::OnWorldUnregisteredWithAudioDevice.AddSP(this, &FEditorDashboardFactory::OnWorldUnregisteredFromAudioDevice);
+		}
+
+		if (!OnDeviceDestroyedHandle.IsValid())
+		{
+			OnDeviceDestroyedHandle = FAudioDeviceManagerDelegates::OnAudioDeviceDestroyed.AddSP(this, &FEditorDashboardFactory::OnDeviceDestroyed);
+		}
+
+		if (!OnPIEStartedHandle.IsValid())
+		{
+			OnPIEStartedHandle = FEditorDelegates::PreBeginPIE.AddSP(this, &FEditorDashboardFactory::OnPIEStarted);
+		}
+
+		if (!OnPostPIEStartedHandle.IsValid())
+		{
+			OnPostPIEStartedHandle = FEditorDelegates::PostPIEStarted.AddSP(this, &FEditorDashboardFactory::OnPostPIEStarted);
+		}
+
+		if (!OnPIEStoppedHandle.IsValid())
+		{
+			OnPIEStoppedHandle = FEditorDelegates::EndPIE.AddSP(this, &FEditorDashboardFactory::OnPIEStopped);
+		}
+	}
+
+	TSharedPtr<FTabManager::FLayout> FEditorDashboardFactory::GetDefaultTabLayout()
+	{
+		using namespace EditorDashboardFactoryPrivate;
+
+		TSharedRef<FTabManager::FStack> ViewportTabStack = FTabManager::NewStack();
+		TSharedRef<FTabManager::FStack> LogTabStack = FTabManager::NewStack();
+		TSharedRef<FTabManager::FStack> AnalysisTabStack = FTabManager::NewStack();
+		TSharedRef<FTabManager::FStack> AudioMetersTabStack = FTabManager::NewStack();
+		TSharedRef<FTabManager::FStack> AudioMeterTabStack = FTabManager::NewStack();
+		TSharedRef<FTabManager::FStack> OscilloscopeTabStack = FTabManager::NewStack();
+
+		for (const auto& KVP : DashboardViewFactories)
+		{
+			const FName& FactoryName = KVP.Key;
+			const TSharedPtr<IDashboardViewFactory>& Factory = KVP.Value;
+
+			const EDefaultDashboardTabStack DefaultTabStack = Factory->GetDefaultTabStack();
+			switch (DefaultTabStack)
+			{
+				case EDefaultDashboardTabStack::Viewport:
+				{
+					ViewportTabStack->AddTab(FactoryName, ETabState::OpenedTab);
+				}
+				break;
+
+				case EDefaultDashboardTabStack::Log:
+				{
+					LogTabStack->AddTab(FactoryName, ETabState::OpenedTab);
+				}
+				break;
+
+				case EDefaultDashboardTabStack::Analysis:
+				{
+					AnalysisTabStack->AddTab(FactoryName, ETabState::OpenedTab);
+				}
+				break;
+
+				case EDefaultDashboardTabStack::AudioMeters:
+				{
+					AudioMetersTabStack->AddTab(FactoryName, ETabState::OpenedTab);
+				}
+				break;
+
+				case EDefaultDashboardTabStack::AudioMeter:
+				{
+					AudioMeterTabStack->AddTab(FactoryName, ETabState::OpenedTab);
+				}
+				break;
+
+				case EDefaultDashboardTabStack::Oscilloscope:
+				{
+					OscilloscopeTabStack->AddTab(FactoryName, ETabState::OpenedTab);
+				}
+				break;
+
+				default:
+					break;
+			}
+		}
+
+		AnalysisTabStack->SetForegroundTab(FName("MixerSources"));
+
+		return FTabManager::NewLayout("AudioDashboard_Editor_Layout_v1")
+		->AddArea
+		(
+			FTabManager::NewPrimaryArea()
+			->SetOrientation(Orient_Vertical)
+			->Split
+			(
+				// Left column
+				FTabManager::NewSplitter()
+				->SetOrientation(Orient_Horizontal)
+				->Split
+				(
+					// Top
+					FTabManager::NewSplitter()
+					->SetOrientation(Orient_Vertical)
+					->SetSizeCoefficient(0.25f) // Column width
+					->Split
+					(
+						ViewportTabStack
+						->SetSizeCoefficient(0.5f)
+					)
+					// Bottom
+					->Split
+					(
+						LogTabStack
+						->SetSizeCoefficient(0.5f)
+					)
+				)
+				
+				// Middle column
+				->Split
+				(
+					// Top
+					FTabManager::NewSplitter()
+					->SetOrientation(Orient_Vertical)
+					->SetSizeCoefficient(0.6f) // Column width
+					->Split
+					(
+						FTabManager::NewSplitter()
+						->SetOrientation(Orient_Horizontal)
+						->Split
+						(
+							AnalysisTabStack
+							->SetSizeCoefficient(0.58f)
+						)
+					)
+					// Bottom
+					->Split
+					(
+						AudioMetersTabStack
+						->SetSizeCoefficient(0.42f)
+					)
+				)
+				// Right column
+				->Split
+				(
+					// Top
+					FTabManager::NewSplitter()
+					->SetOrientation(Orient_Vertical)
+					->SetSizeCoefficient(0.15f) // Column width
+					->Split
+					(
+						AudioMeterTabStack
+						->SetSizeCoefficient(0.7f)
+						->SetHideTabWell(true)
+					)
+					// Bottom
+					->Split
+					(
+						OscilloscopeTabStack
+						->SetSizeCoefficient(0.3f)
+						->SetHideTabWell(true)
+					)
+				)
+			)
+		);
+	}
+
+	void FEditorDashboardFactory::RegisterTabSpawners()
+	{
+		using namespace EditorDashboardFactoryPrivate;
+
+		DashboardWorkspace = DashboardTabManager->AddLocalWorkspaceMenuCategory(ToolName);
+
+		for (const auto& KVP : DashboardViewFactories)
+		{
+			const FName& FactoryName = KVP.Key;
+			const TSharedPtr<IDashboardViewFactory>& Factory = KVP.Value;
+
+			DashboardTabManager->RegisterTabSpawner(FactoryName, FOnSpawnTab::CreateLambda([this, Factory](const FSpawnTabArgs& Args)
+			{
+				TSharedPtr<SWidget> DashboardView = Factory->MakeWidget();
+				return SNew(SDockTab)
+					.Clipping(EWidgetClipping::ClipToBounds)
+					.Label(Factory->GetDisplayName())
+					[
+						DashboardView->AsShared()
+					];
+			}))
+			.SetDisplayName(Factory->GetDisplayName())
+			.SetGroup(DashboardWorkspace->AsShared())
+			.SetIcon(Factory->GetIcon());
+		}
+	}
+
+	void FEditorDashboardFactory::RegisterViewFactory(TSharedRef<IDashboardViewFactory> InFactory)
+	{
+		if (const FName Name = InFactory->GetName(); 
+			ensureAlwaysMsgf(!DashboardViewFactories.Contains(Name), TEXT("Failed to register Audio Insights Dashboard '%s': Dashboard with name already registered"), *Name.ToString()))
+		{
+			DashboardViewFactories.Add(Name, InFactory);
+		}
+	}
+
+	void FEditorDashboardFactory::UnregisterTabSpawners()
+	{
+		if (DashboardTabManager.IsValid())
+		{
+			for (const auto& KVP : DashboardViewFactories)
+			{
+				const FName& FactoryName = KVP.Key;
+				DashboardTabManager->UnregisterTabSpawner(FactoryName);
+			}
+
+			DashboardTabManager.Reset();
+		}
+
+		DashboardWorkspace.Reset();
+	}
+
+	void FEditorDashboardFactory::UnregisterViewFactory(FName InName)
+	{
+		DashboardViewFactories.Remove(InName);
+	}
+} // namespace UE::Audio::Insights
+
+#undef LOCTEXT_NAMESPACE
diff --git a/Engine/Plugins/AudioInsights/Source/AudioInsightsEditor/Private/AudioInsightsEditorDashboardFactory.h b/Engine/Plugins/AudioInsights/Source/AudioInsightsEditor/Private/AudioInsightsEditorDashboardFactory.h
new file mode 100644
--- /dev/null
+++ b/Engine/Plugins/AudioInsights/Source/AudioInsightsEditor/Private/AudioInsightsEditorDashboardFactory.h
@@ -0,0 +1,74 @@
+// Copyright Epic Games, Inc. All Rights Reserved.
+#pragma once
+
+#include "AudioDefines.h"
+#include "Containers/Map.h"
+#include "Engine/World.h"
+#include "Framework/Docking/TabManager.h"
+#include "IAudioInsightsDashboardFactory.h"
+#include "Templates/SharedPointer.h"
+#include "Views/DashboardViewFactory.h"
+#include "Widgets/Docking/SDockTab.h"
+#include "Widgets/Input/SComboBox.h"
+
+namespace UE::Audio::Insights
+{
+	class FEditorDashboardFactory : public IDashboardFactory, public TSharedFromThis<FEditorDashboardFactory>
+	{
+	public:
+		FEditorDashboardFactory() = default;
+		virtual ~FEditorDashboardFactory() = default;
+
+		TSharedRef<SDockTab> MakeDockTabWidget(const FSpawnTabArgs& Args);
+
+		virtual void RegisterViewFactory(TSharedRef<IDashboardViewFactory> InFactory) override;
+		virtual void UnregisterViewFactory(FName InDashboardName) override;
+		virtual ::Audio::FDeviceId GetDeviceId() const override;
+
+		DECLARE_MULTICAST_DELEGATE(FOnActiveAudioDeviceChanged);
+		inline static FOnActiveAudioDeviceChanged OnActiveAudioDeviceChanged;
+
+	private:
+		void OnWorldRegisteredToAudioDevice(const UWorld* InWorld, ::Audio::FDeviceId InDeviceId);
+		void OnWorldUnregisteredFromAudioDevice(const UWorld* InWorld, ::Audio::FDeviceId InDeviceId);
+		void OnDeviceDestroyed(::Audio::FDeviceId InDeviceId);
+
+		void OnPIEStarted(bool bSimulating);
+		void OnPostPIEStarted(bool bSimulating);
+		void OnPIEStopped(bool bSimulating);
+
+		void RefreshDeviceSelector();
+		void ResetDelegates();
+
+		TSharedRef<SWidget> MakeMenuBarWidget();
+		TSharedRef<SWidget> MakeMainToolbarWidget();
+
+		void InitDelegates();
+		TSharedPtr<FTabManager::FLayout> GetDefaultTabLayout();
+
+		void RegisterTabSpawners();
+		void UnregisterTabSpawners();
+
+		FDelegateHandle OnDeviceDestroyedHandle;
+		FDelegateHandle OnWorldRegisteredToAudioDeviceHandle;
+		FDelegateHandle OnWorldUnregisteredFromAudioDeviceHandle;
+
+		FDelegateHandle OnPIEStartedHandle;
+		FDelegateHandle OnPostPIEStartedHandle;
+		FDelegateHandle OnPIEStoppedHandle;
+
+		bool bStartWithPIE = true;
+		bool bStopWithPIE = true;
+
+		TSharedPtr<FTabManager> DashboardTabManager;
+		TSharedPtr<FWorkspaceItem> DashboardWorkspace;
+		TSharedPtr<FTabManager::FLayout> TabLayout;
+
+		TArray<TSharedPtr<::Audio::FDeviceId>> AudioDeviceIds;
+		TSharedPtr<SComboBox<TSharedPtr<::Audio::FDeviceId>>> AudioDeviceComboBox;
+
+		::Audio::FDeviceId ActiveDeviceId = INDEX_NONE;
+
+		TMap<FName, TSharedPtr<IDashboardViewFactory>> DashboardViewFactories;
+	};
+} // namespace UE::Audio::Insights
diff --git a/Engine/Plugins/AudioInsights/Source/AudioInsightsEditor/Private/AudioInsightsEditorModule.cpp b/Engine/Plugins/AudioInsights/Source/AudioInsightsEditor/Private/AudioInsightsEditorModule.cpp
new file mode 100644
--- /dev/null
+++ b/Engine/Plugins/AudioInsights/Source/AudioInsightsEditor/Private/AudioInsightsEditorModule.cpp
@@ -0,0 +1,104 @@
+// Copyright Epic Games, Inc. All Rights Reserved.
+#include "AudioInsightsEditorModule.h"
+
+#include "AudioInsightsEditorLog.h"
+#include "AudioInsightsStyle.h"
+#include "Framework/Docking/TabManager.h"
+#include "Modules/ModuleManager.h"
+#include "Templates/SharedPointer.h"
+#include "UObject/NameTypes.h"
+#include "Views/AudioBusesDashboardViewFactory.h"
+#include "Views/AudioMetersDashboardViewFactory.h"
+#include "Views/LogDashboardViewFactory.h"
+#include "Views/MixerSourceDashboardViewFactory.h"
+#include "Views/OutputMeterDashboardViewFactory.h"
+#include "Views/OutputOscilloscopeDashboardViewFactory.h"
+#include "Views/SubmixesDashboardViewFactory.h"
+#include "Views/ViewportDashboardViewFactory.h"
+#include "Views/VirtualLoopDashboardViewFactory.h"
+#include "WorkspaceMenuStructure.h"
+#include "WorkspaceMenuStructureModule.h"
+
+#define LOCTEXT_NAMESPACE "AudioInsights"
+
+DEFINE_LOG_CATEGORY(LogAudioInsightsEditor);
+
+namespace UE::Audio::Insights
+{
+	void FAudioInsightsEditorModule::StartupModule()
+	{
+		// Don't run providers in commandlet to avoid additional, unnecessary overhead as audio insights is dormant.
+		if (!IsRunningCommandlet())
+		{
+			RegisterMenus();
+
+			DashboardFactory = MakeShared<FEditorDashboardFactory>();
+			
+			DashboardFactory->RegisterViewFactory(MakeShared<FViewportDashboardViewFactory>());
+			DashboardFactory->RegisterViewFactory(MakeShared<FLogDashboardViewFactory>());
+			DashboardFactory->RegisterViewFactory(MakeShared<FMixerSourceDashboardViewFactory>());
+			DashboardFactory->RegisterViewFactory(MakeShared<FVirtualLoopDashboardViewFactory>());
+			DashboardFactory->RegisterViewFactory(MakeShared<FSubmixesDashboardViewFactory>());
+			DashboardFactory->RegisterViewFactory(MakeShared<FAudioBusesDashboardViewFactory>());
+			DashboardFactory->RegisterViewFactory(MakeShared<FAudioMetersDashboardViewFactory>());
+			DashboardFactory->RegisterViewFactory(MakeShared<FOutputMeterDashboardViewFactory>());
+			DashboardFactory->RegisterViewFactory(MakeShared<FOutputOscilloscopeDashboardViewFactory>());
+		}
+	}
+
+	void FAudioInsightsEditorModule::ShutdownModule()
+	{
+		if (!IsRunningCommandlet())
+		{
+			DashboardFactory.Reset();
+		}
+	}
+
+	void FAudioInsightsEditorModule::RegisterDashboardViewFactory(TSharedRef<IDashboardViewFactory> InDashboardFactory)
+	{
+		DashboardFactory->RegisterViewFactory(InDashboardFactory);
+	}
+
+	void FAudioInsightsEditorModule::UnregisterDashboardViewFactory(FName InName)
+	{
+		DashboardFactory->UnregisterViewFactory(InName);
+	}
+
+	::Audio::FDeviceId FAudioInsightsEditorModule::GetDeviceId() const
+	{
+		return DashboardFactory->GetDeviceId();
+	}
+
+	FAudioInsightsEditorModule& FAudioInsightsEditorModule::GetChecked()
+	{
+		return static_cast<FAudioInsightsEditorModule&>(FModuleManager::GetModuleChecked<IAudioInsightsModule>(GetEditorName()));
+	}
+
+	TSharedRef<FEditorDashboardFactory> FAudioInsightsEditorModule::GetDashboardFactory()
+	{
+		return DashboardFactory->AsShared();
+	}
+
+	const TSharedRef<FEditorDashboardFactory> FAudioInsightsEditorModule::GetDashboardFactory() const
+	{
+		return DashboardFactory->AsShared();
+	}
+
+	TSharedRef<SDockTab> FAudioInsightsEditorModule::CreateDashboardTabWidget(const FSpawnTabArgs& Args)
+	{
+		return DashboardFactory->MakeDockTabWidget(Args);
+	}
+
+	void FAudioInsightsEditorModule::RegisterMenus()
+	{
+		const IWorkspaceMenuStructure& MenuStructure = WorkspaceMenu::GetMenuStructure();
+		FGlobalTabmanager::Get()->RegisterNomadTabSpawner("AudioInsights", FOnSpawnTab::CreateRaw(this, &FAudioInsightsEditorModule::CreateDashboardTabWidget))
+			.SetDisplayName(LOCTEXT("OpenDashboard_TabDisplayName", "Audio Insights"))
+			.SetTooltipText(LOCTEXT("OpenDashboard_TabTooltip", "Opens Audio Insights, an extensible suite of tools and visualizers which enable monitoring and debugging audio in the Unreal Engine."))
+			.SetGroup(MenuStructure.GetToolsCategory())
+			.SetIcon(FSlateStyle::Get().CreateIcon("AudioInsights.Icon.Dashboard"));
+	};
+} // namespace UE::Audio::Insights
+#undef LOCTEXT_NAMESPACE // AudioInsights
+
+IMPLEMENT_MODULE(UE::Audio::Insights::FAudioInsightsEditorModule, AudioInsightsEditor)
diff --git a/Engine/Plugins/AudioInsights/Source/AudioInsightsEditor/Private/AudioInsightsEditorModule.h b/Engine/Plugins/AudioInsights/Source/AudioInsightsEditor/Private/AudioInsightsEditorModule.h
new file mode 100644
--- /dev/null
+++ b/Engine/Plugins/AudioInsights/Source/AudioInsightsEditor/Private/AudioInsightsEditorModule.h
@@ -0,0 +1,43 @@
+// Copyright Epic Games, Inc. All Rights Reserved.
+#pragma once 
+
+#include "AudioInsightsEditorDashboardFactory.h"
+#include "AudioInsightsTraceModule.h"
+#include "Framework/Docking/TabManager.h"
+#include "IAudioInsightsModule.h"
+#include "Templates/SharedPointer.h"
+#include "Views/DashboardViewFactory.h"
+#include "Views/MixerSourceMuteSoloFilter.h"
+#include "Views/VirtualLoopsDebugDraw.h"
+#include "Widgets/Docking/SDockTab.h"
+
+namespace UE::Audio::Insights
+{
+	class FAudioInsightsEditorModule final : public IAudioInsightsModule
+	{
+	public:
+		FAudioInsightsEditorModule() = default;
+
+		virtual ~FAudioInsightsEditorModule() = default;
+
+		virtual void StartupModule() override;
+		virtual void ShutdownModule() override;
+		virtual void RegisterDashboardViewFactory(TSharedRef<IDashboardViewFactory> InDashboardFactory) override;
+		virtual void UnregisterDashboardViewFactory(FName InName) override;
+		virtual ::Audio::FDeviceId GetDeviceId() const override;
+
+		TSharedRef<FEditorDashboardFactory> GetDashboardFactory();
+		const TSharedRef<FEditorDashboardFactory> GetDashboardFactory() const;
+
+		static FAudioInsightsEditorModule& GetChecked();
+
+	private:
+		void RegisterMenus();
+		TSharedRef<SDockTab> CreateDashboardTabWidget(const FSpawnTabArgs& Args);
+		
+		TSharedPtr<FEditorDashboardFactory> DashboardFactory;
+		
+		FVirtualLoopsDebugDraw VirtualLoopsDebugDraw;
+		FMuteSoloFilter MuteSoloFilter;
+	};
+} // namespace UE::Audio::Insights
diff --git a/Engine/Plugins/AudioInsights/Source/AudioInsightsEditor/Private/Providers/AudioBusProvider.cpp b/Engine/Plugins/AudioInsights/Source/AudioInsightsEditor/Private/Providers/AudioBusProvider.cpp
new file mode 100644
--- /dev/null
+++ b/Engine/Plugins/AudioInsights/Source/AudioInsightsEditor/Private/Providers/AudioBusProvider.cpp
@@ -0,0 +1,167 @@
+// Copyright Epic Games, Inc. All Rights Reserved.
+
+#include "Providers/AudioBusProvider.h"
+
+#include "AssetRegistry/AssetData.h"
+#include "AssetRegistry/AssetRegistryModule.h"
+#include "AudioInsightsEditorDashboardFactory.h"
+#include "AudioInsightsEditorModule.h"
+#include "Sound/AudioBus.h"
+
+namespace UE::Audio::Insights
+{
+	FAudioBusProvider::FAudioBusProvider()
+		: TDeviceDataMapTraceProvider<uint32, TSharedPtr<FAudioBusAssetDashboardEntry>>(GetName_Static())
+	{
+		FAssetRegistryModule& AssetRegistryModule = FModuleManager::LoadModuleChecked<FAssetRegistryModule>("AssetRegistry");
+
+		AssetRegistryModule.Get().OnAssetAdded().AddRaw(this, &FAudioBusProvider::OnAssetAdded);
+		AssetRegistryModule.Get().OnAssetRemoved().AddRaw(this, &FAudioBusProvider::OnAssetRemoved);
+		AssetRegistryModule.Get().OnFilesLoaded().AddRaw(this, &FAudioBusProvider::OnFilesLoaded);
+
+		FEditorDashboardFactory::OnActiveAudioDeviceChanged.AddRaw(this, &FAudioBusProvider::OnActiveAudioDeviceChanged);
+	}
+	
+	FAudioBusProvider::~FAudioBusProvider()
+	{
+		if (FAssetRegistryModule* AssetRegistryModule = FModuleManager::GetModulePtr<FAssetRegistryModule>("AssetRegistry"))
+		{
+			AssetRegistryModule->Get().OnAssetAdded().RemoveAll(this);
+			AssetRegistryModule->Get().OnAssetRemoved().RemoveAll(this);
+			AssetRegistryModule->Get().OnFilesLoaded().RemoveAll(this);
+		}
+
+		FEditorDashboardFactory::OnActiveAudioDeviceChanged.RemoveAll(this);
+	}
+
+	FName FAudioBusProvider::GetName_Static()
+	{
+		return "AudioBusProvider";
+	}
+
+	void FAudioBusProvider::OnAssetAdded(const FAssetData& InAssetData)
+	{
+		if (bAreFilesLoaded && InAssetData.AssetClassPath == FTopLevelAssetPath(UAudioBus::StaticClass()))
+		{
+			AddAudioBusAsset(InAssetData);
+		}
+	}
+
+	void FAudioBusProvider::OnAssetRemoved(const FAssetData& InAssetData)
+	{
+		if (InAssetData.AssetClassPath == FTopLevelAssetPath(UAudioBus::StaticClass()))
+		{
+			RemoveAudioBusAsset(InAssetData);
+		}
+	}
+
+	void FAudioBusProvider::OnFilesLoaded()
+	{
+		bAreFilesLoaded = true;
+		UpdateAudioBusAssetNames();
+	}
+
+	void FAudioBusProvider::OnActiveAudioDeviceChanged()
+	{
+		UpdateAudioBusAssetNames();
+	}
+
+	void FAudioBusProvider::AddAudioBusAsset(const FAssetData& InAssetData)
+	{
+		const bool bIsAudioBusAssetAlreadAdded = AudioBusDataViewEntries.ContainsByPredicate(
+			[AssetName = InAssetData.GetObjectPathString()](const TSharedPtr<FAudioBusAssetDashboardEntry>& AudioBusAssetDashboardEntry)
+			{
+				return AudioBusAssetDashboardEntry->Name == AssetName;
+			});
+
+		if (!bIsAudioBusAssetAlreadAdded)
+		{
+			const FAudioInsightsEditorModule AudioInsightsEditorModule = FAudioInsightsEditorModule::GetChecked();
+			const ::Audio::FDeviceId AudioDeviceId = AudioInsightsEditorModule.GetDeviceId();
+
+			TSharedPtr<FAudioBusAssetDashboardEntry> AudioBusAssetDashboardEntry = MakeShared<FAudioBusAssetDashboardEntry>();
+			AudioBusAssetDashboardEntry->DeviceId = AudioDeviceId;
+			AudioBusAssetDashboardEntry->Name     = InAssetData.GetObjectPathString();
+			AudioBusAssetDashboardEntry->AudioBus = Cast<UAudioBus>(InAssetData.GetAsset());
+
+			AudioBusDataViewEntries.Add(MoveTemp(AudioBusAssetDashboardEntry));
+
+			OnAudioBusAssetAdded.Broadcast(InAssetData.GetAsset());
+			++LastUpdateId;
+		}
+	}
+
+	void FAudioBusProvider::RemoveAudioBusAsset(const FAssetData& InAssetData)
+	{
+		const int32 FoundAudioBusAssetNameIndex = AudioBusDataViewEntries.IndexOfByPredicate(
+			[AssetName = InAssetData.GetObjectPathString()](const TSharedPtr<FAudioBusAssetDashboardEntry>& AudioBusAssetDashboardEntry)
+			{
+				return AudioBusAssetDashboardEntry->Name == AssetName;
+			});
+
+		if (FoundAudioBusAssetNameIndex != INDEX_NONE)
+		{
+			const FAudioInsightsEditorModule AudioInsightsEditorModule = FAudioInsightsEditorModule::GetChecked();
+			const ::Audio::FDeviceId AudioDeviceId = AudioInsightsEditorModule.GetDeviceId();
+
+			RemoveDeviceEntry(AudioDeviceId, AudioBusDataViewEntries[FoundAudioBusAssetNameIndex]->AudioBus->GetUniqueID());
+
+			AudioBusDataViewEntries.RemoveAt(FoundAudioBusAssetNameIndex);
+
+			OnAudioBusAssetRemoved.Broadcast(InAssetData.GetAsset());
+			++LastUpdateId;
+		}
+	}
+
+	void FAudioBusProvider::UpdateAudioBusAssetNames()
+	{
+		// Get all UAudioBus assets
+		TArray<FAssetData> AssetDataArray;
+
+		FAssetRegistryModule& AssetRegistryModule = FModuleManager::LoadModuleChecked<FAssetRegistryModule>("AssetRegistry");
+		AssetRegistryModule.Get().GetAssetsByClass(FTopLevelAssetPath(UAudioBus::StaticClass()), AssetDataArray);
+
+		// Build AudioBusDataViewEntries
+		Reset();
+		AudioBusDataViewEntries.Empty();
+
+		for (const FAssetData& AssetData : AssetDataArray)
+		{
+			AddAudioBusAsset(AssetData);
+		}
+
+		AudioBusDataViewEntries.Sort([](const TSharedPtr<FAudioBusAssetDashboardEntry>& A, const TSharedPtr<FAudioBusAssetDashboardEntry>& B)
+		{
+			return A->GetDisplayName().CompareToCaseIgnored(B->GetDisplayName()) < 0;
+		});
+
+		OnAudioBusAssetListUpdated.Broadcast();
+	}
+
+	bool FAudioBusProvider::ProcessMessages()
+	{
+		const FAudioInsightsEditorModule AudioInsightsEditorModule = FAudioInsightsEditorModule::GetChecked();
+		const ::Audio::FDeviceId AudioDeviceId = AudioInsightsEditorModule.GetDeviceId();
+
+		for (const TSharedPtr<FAudioBusAssetDashboardEntry>& AudioBusDataViewEntry : AudioBusDataViewEntries)
+		{
+			if (AudioBusDataViewEntry.IsValid() && AudioBusDataViewEntry->AudioBus.IsValid())
+			{
+				UpdateDeviceEntry(AudioDeviceId, AudioBusDataViewEntry->AudioBus->GetUniqueID(), [&AudioBusDataViewEntry](TSharedPtr<FAudioBusAssetDashboardEntry>& Entry)
+				{
+					if (!Entry.IsValid())
+					{
+						Entry = AudioBusDataViewEntry;
+					}
+				});
+			}
+		}
+
+		return true;
+	}
+
+	UE::Trace::IAnalyzer* FAudioBusProvider::ConstructAnalyzer()
+	{
+		return nullptr;
+	}
+} // namespace UE::Audio::Insights
diff --git a/Engine/Plugins/AudioInsights/Source/AudioInsightsEditor/Private/Providers/AudioBusProvider.h b/Engine/Plugins/AudioInsights/Source/AudioInsightsEditor/Private/Providers/AudioBusProvider.h
new file mode 100644
--- /dev/null
+++ b/Engine/Plugins/AudioInsights/Source/AudioInsightsEditor/Private/Providers/AudioBusProvider.h
@@ -0,0 +1,46 @@
+// Copyright Epic Games, Inc. All Rights Reserved.
+#pragma once
+
+#include "Delegates/Delegate.h"
+#include "UObject/NameTypes.h"
+#include "Views/AudioBusAssetDashboardEntry.h"
+
+namespace UE::Audio::Insights
+{
+	class FAudioBusProvider : public TDeviceDataMapTraceProvider<uint32, TSharedPtr<FAudioBusAssetDashboardEntry>>, public TSharedFromThis<FAudioBusProvider>
+	{
+	public:
+		FAudioBusProvider();
+		virtual ~FAudioBusProvider();
+
+		static FName GetName_Static();
+
+		virtual UE::Trace::IAnalyzer* ConstructAnalyzer() override;
+		
+		DECLARE_MULTICAST_DELEGATE_OneParam(FOnAudioBusAssetAdded, const TWeakObjectPtr<UObject> /*Asset*/);
+		inline static FOnAudioBusAssetAdded OnAudioBusAssetAdded;
+
+		DECLARE_MULTICAST_DELEGATE_OneParam(FOnAudioBusAssetRemoved, const TWeakObjectPtr<UObject> /*Asset*/);
+		inline static FOnAudioBusAssetRemoved OnAudioBusAssetRemoved;
+
+		DECLARE_MULTICAST_DELEGATE(FOnAudioBusAssetListUpdated);
+		inline static FOnAudioBusAssetListUpdated OnAudioBusAssetListUpdated;
+
+	private:
+		void OnAssetAdded(const FAssetData& InAssetData);
+		void OnAssetRemoved(const FAssetData& InAssetData);
+		void OnFilesLoaded();
+		void OnActiveAudioDeviceChanged();
+
+		void AddAudioBusAsset(const FAssetData& InAssetData);
+		void RemoveAudioBusAsset(const FAssetData& InAssetData);
+
+		void UpdateAudioBusAssetNames();
+
+		virtual bool ProcessMessages() override;
+
+		bool bAreFilesLoaded = false;
+
+		TArray<TSharedPtr<FAudioBusAssetDashboardEntry>> AudioBusDataViewEntries;
+	};
+} // namespace UE::Audio::Insights
diff --git a/Engine/Plugins/AudioInsights/Source/AudioInsightsEditor/Private/Providers/SoundSubmixProvider.cpp b/Engine/Plugins/AudioInsights/Source/AudioInsightsEditor/Private/Providers/SoundSubmixProvider.cpp
new file mode 100644
--- /dev/null
+++ b/Engine/Plugins/AudioInsights/Source/AudioInsightsEditor/Private/Providers/SoundSubmixProvider.cpp
@@ -0,0 +1,167 @@
+// Copyright Epic Games, Inc. All Rights Reserved.
+
+#include "Providers/SoundSubmixProvider.h"
+
+#include "AssetRegistry/AssetData.h"
+#include "AssetRegistry/AssetRegistryModule.h"
+#include "AudioInsightsEditorDashboardFactory.h"
+#include "AudioInsightsEditorModule.h"
+#include "Sound/SoundSubmix.h"
+
+namespace UE::Audio::Insights
+{
+	FSoundSubmixProvider::FSoundSubmixProvider()
+		: TDeviceDataMapTraceProvider<uint32, TSharedPtr<FSoundSubmixAssetDashboardEntry>>(GetName_Static())
+	{
+		FAssetRegistryModule& AssetRegistryModule = FModuleManager::LoadModuleChecked<FAssetRegistryModule>("AssetRegistry");
+
+		AssetRegistryModule.Get().OnAssetAdded().AddRaw(this, &FSoundSubmixProvider::OnAssetAdded);
+		AssetRegistryModule.Get().OnAssetRemoved().AddRaw(this, &FSoundSubmixProvider::OnAssetRemoved);
+		AssetRegistryModule.Get().OnFilesLoaded().AddRaw(this, &FSoundSubmixProvider::OnFilesLoaded);
+
+		FEditorDashboardFactory::OnActiveAudioDeviceChanged.AddRaw(this, &FSoundSubmixProvider::OnActiveAudioDeviceChanged);
+	}
+	
+	FSoundSubmixProvider::~FSoundSubmixProvider()
+	{
+		if (FAssetRegistryModule* AssetRegistryModule = FModuleManager::GetModulePtr<FAssetRegistryModule>("AssetRegistry"))
+		{
+			AssetRegistryModule->Get().OnAssetAdded().RemoveAll(this);
+			AssetRegistryModule->Get().OnAssetRemoved().RemoveAll(this);
+			AssetRegistryModule->Get().OnFilesLoaded().RemoveAll(this);
+		}
+
+		FEditorDashboardFactory::OnActiveAudioDeviceChanged.RemoveAll(this);
+	}
+
+	FName FSoundSubmixProvider::GetName_Static()
+	{
+		return "SubmixesProvider";
+	}
+
+	void FSoundSubmixProvider::OnAssetAdded(const FAssetData& InAssetData)
+	{
+		if (bAreFilesLoaded && InAssetData.AssetClassPath == FTopLevelAssetPath(USoundSubmix::StaticClass()))
+		{
+			AddSubmixAsset(InAssetData);
+		}
+	}
+
+	void FSoundSubmixProvider::OnAssetRemoved(const FAssetData& InAssetData)
+	{
+		if (InAssetData.AssetClassPath == FTopLevelAssetPath(USoundSubmix::StaticClass()))
+		{
+			RemoveSubmixAsset(InAssetData);
+		}
+	}
+
+	void FSoundSubmixProvider::OnFilesLoaded()
+	{
+		bAreFilesLoaded = true;
+		UpdateSubmixAssetNames();
+	}
+
+	void FSoundSubmixProvider::OnActiveAudioDeviceChanged()
+	{
+		UpdateSubmixAssetNames();
+	}
+
+	void FSoundSubmixProvider::AddSubmixAsset(const FAssetData& InAssetData)
+	{
+		const bool bIsSubmixAssetAlreadAdded = SubmixDataViewEntries.ContainsByPredicate(
+			[AssetName = InAssetData.GetObjectPathString()](const TSharedPtr<FSoundSubmixAssetDashboardEntry>& SoundSubmixAssetDashboardEntry)
+			{
+				return SoundSubmixAssetDashboardEntry->Name == AssetName;
+			});
+
+		if (!bIsSubmixAssetAlreadAdded)
+		{
+			const FAudioInsightsEditorModule AudioInsightsEditorModule = FAudioInsightsEditorModule::GetChecked();
+			const ::Audio::FDeviceId AudioDeviceId = AudioInsightsEditorModule.GetDeviceId();
+
+			TSharedPtr<FSoundSubmixAssetDashboardEntry> SoundSubmixAssetDashboardEntry = MakeShared<FSoundSubmixAssetDashboardEntry>();
+			SoundSubmixAssetDashboardEntry->DeviceId    = AudioDeviceId;
+			SoundSubmixAssetDashboardEntry->Name        = InAssetData.GetObjectPathString();
+			SoundSubmixAssetDashboardEntry->SoundSubmix = Cast<USoundSubmix>(InAssetData.GetAsset());
+
+			SubmixDataViewEntries.Add(MoveTemp(SoundSubmixAssetDashboardEntry));
+
+			OnSubmixAssetAdded.Broadcast(InAssetData.GetAsset());
+			++LastUpdateId;
+		}
+	}
+
+	void FSoundSubmixProvider::RemoveSubmixAsset(const FAssetData& InAssetData)
+	{
+		const int32 FoundSubmixAssetNameIndex = SubmixDataViewEntries.IndexOfByPredicate(
+			[AssetName = InAssetData.GetObjectPathString()](const TSharedPtr<FSoundSubmixAssetDashboardEntry>& SoundSubmixAssetDashboardEntry)
+			{
+				return SoundSubmixAssetDashboardEntry->Name == AssetName;
+			});
+
+		if (FoundSubmixAssetNameIndex != INDEX_NONE)
+		{
+			const FAudioInsightsEditorModule AudioInsightsEditorModule = FAudioInsightsEditorModule::GetChecked();
+			const ::Audio::FDeviceId AudioDeviceId = AudioInsightsEditorModule.GetDeviceId();
+
+			RemoveDeviceEntry(AudioDeviceId, SubmixDataViewEntries[FoundSubmixAssetNameIndex]->SoundSubmix->GetUniqueID());
+
+			SubmixDataViewEntries.RemoveAt(FoundSubmixAssetNameIndex);
+
+			OnSubmixAssetRemoved.Broadcast(InAssetData.GetAsset());
+			++LastUpdateId;
+		}
+	}
+
+	void FSoundSubmixProvider::UpdateSubmixAssetNames()
+	{
+		// Get all USoundSubmix assets
+		TArray<FAssetData> AssetDataArray;
+
+		FAssetRegistryModule& AssetRegistryModule = FModuleManager::LoadModuleChecked<FAssetRegistryModule>("AssetRegistry");
+		AssetRegistryModule.Get().GetAssetsByClass(FTopLevelAssetPath(USoundSubmix::StaticClass()), AssetDataArray);
+
+		// Build SubmixDataViewEntries
+		Reset();
+		SubmixDataViewEntries.Empty();
+
+		for (const FAssetData& AssetData : AssetDataArray)
+		{
+			AddSubmixAsset(AssetData);
+		}
+
+		SubmixDataViewEntries.Sort([](const TSharedPtr<FSoundSubmixAssetDashboardEntry>& A, const TSharedPtr<FSoundSubmixAssetDashboardEntry>& B)
+		{
+			return A->GetDisplayName().CompareToCaseIgnored(B->GetDisplayName()) < 0;
+		});
+
+		OnSubmixAssetListUpdated.Broadcast();
+	}
+
+	bool FSoundSubmixProvider::ProcessMessages()
+	{
+		const FAudioInsightsEditorModule AudioInsightsEditorModule = FAudioInsightsEditorModule::GetChecked();
+		const ::Audio::FDeviceId AudioDeviceId = AudioInsightsEditorModule.GetDeviceId();
+
+		for (const TSharedPtr<FSoundSubmixAssetDashboardEntry>& SubmixDataViewEntry : SubmixDataViewEntries)
+		{
+			if (SubmixDataViewEntry.IsValid() && SubmixDataViewEntry->SoundSubmix.IsValid())
+			{
+				UpdateDeviceEntry(AudioDeviceId, SubmixDataViewEntry->SoundSubmix->GetUniqueID(), [&SubmixDataViewEntry](TSharedPtr<FSoundSubmixAssetDashboardEntry>& Entry)
+				{
+					if (!Entry.IsValid())
+					{
+						Entry = SubmixDataViewEntry;
+					}
+				});
+			}
+		}
+
+		return true;
+	}
+
+	UE::Trace::IAnalyzer* FSoundSubmixProvider::ConstructAnalyzer()
+	{
+		return nullptr;
+	}
+} // namespace UE::Audio::Insights
diff --git a/Engine/Plugins/AudioInsights/Source/AudioInsightsEditor/Private/Providers/SoundSubmixProvider.h b/Engine/Plugins/AudioInsights/Source/AudioInsightsEditor/Private/Providers/SoundSubmixProvider.h
new file mode 100644
--- /dev/null
+++ b/Engine/Plugins/AudioInsights/Source/AudioInsightsEditor/Private/Providers/SoundSubmixProvider.h
@@ -0,0 +1,46 @@
+// Copyright Epic Games, Inc. All Rights Reserved.
+#pragma once
+
+#include "Delegates/Delegate.h"
+#include "UObject/NameTypes.h"
+#include "Views/SoundSubmixAssetDashboardEntry.h"
+
+namespace UE::Audio::Insights
+{
+	class FSoundSubmixProvider : public TDeviceDataMapTraceProvider<uint32, TSharedPtr<FSoundSubmixAssetDashboardEntry>>, public TSharedFromThis<FSoundSubmixProvider>
+	{
+	public:
+		FSoundSubmixProvider();
+		virtual ~FSoundSubmixProvider();
+
+		static FName GetName_Static();
+
+		virtual UE::Trace::IAnalyzer* ConstructAnalyzer() override;
+		
+		DECLARE_MULTICAST_DELEGATE_OneParam(FOnSubmixAssetAdded, const TWeakObjectPtr<UObject> /*Asset*/);
+		inline static FOnSubmixAssetAdded OnSubmixAssetAdded;
+
+		DECLARE_MULTICAST_DELEGATE_OneParam(FOnSubmixAssetRemoved, const TWeakObjectPtr<UObject> /*Asset*/);
+		inline static FOnSubmixAssetRemoved OnSubmixAssetRemoved;
+
+		DECLARE_MULTICAST_DELEGATE(FOnSubmixAssetListUpdated);
+		inline static FOnSubmixAssetListUpdated OnSubmixAssetListUpdated;
+
+	private:
+		void OnAssetAdded(const FAssetData& InAssetData);
+		void OnAssetRemoved(const FAssetData& InAssetData);
+		void OnFilesLoaded();
+		void OnActiveAudioDeviceChanged();
+
+		void AddSubmixAsset(const FAssetData& InAssetData);
+		void RemoveSubmixAsset(const FAssetData& InAssetData);
+
+		void UpdateSubmixAssetNames();
+
+		virtual bool ProcessMessages() override;
+
+		bool bAreFilesLoaded = false;
+
+		TArray<TSharedPtr<FSoundSubmixAssetDashboardEntry>> SubmixDataViewEntries;
+	};
+} // namespace UE::Audio::Insights
diff --git a/Engine/Plugins/AudioInsights/Source/AudioInsightsEditor/Private/Views/AudioBusAssetDashboardEntry.cpp b/Engine/Plugins/AudioInsights/Source/AudioInsightsEditor/Private/Views/AudioBusAssetDashboardEntry.cpp
new file mode 100644
--- /dev/null
+++ b/Engine/Plugins/AudioInsights/Source/AudioInsightsEditor/Private/Views/AudioBusAssetDashboardEntry.cpp
@@ -0,0 +1,26 @@
+// Copyright Epic Games, Inc. All Rights Reserved.
+
+#include "Views/AudioBusAssetDashboardEntry.h"
+
+namespace UE::Audio::Insights
+{
+	FText FAudioBusAssetDashboardEntry::GetDisplayName() const
+	{
+		return FText::FromString(FSoftObjectPath(Name).GetAssetName());
+	}
+
+	const UObject* FAudioBusAssetDashboardEntry::GetObject() const
+	{
+		return FSoftObjectPath(Name).ResolveObject();
+	}
+
+	UObject* FAudioBusAssetDashboardEntry::GetObject()
+	{
+		return FSoftObjectPath(Name).ResolveObject();
+	}
+
+	bool FAudioBusAssetDashboardEntry::IsValid() const
+	{
+		return true;
+	}
+} // namespace UE::Audio::Insights
diff --git a/Engine/Plugins/AudioInsights/Source/AudioInsightsEditor/Private/Views/AudioBusAssetDashboardEntry.h b/Engine/Plugins/AudioInsights/Source/AudioInsightsEditor/Private/Views/AudioBusAssetDashboardEntry.h
new file mode 100644
--- /dev/null
+++ b/Engine/Plugins/AudioInsights/Source/AudioInsightsEditor/Private/Views/AudioBusAssetDashboardEntry.h
@@ -0,0 +1,24 @@
+// Copyright Epic Games, Inc. All Rights Reserved.
+#pragma once
+
+#include "Views/TableDashboardViewFactory.h"
+
+class UAudioBus;
+
+namespace UE::Audio::Insights
+{
+	struct FAudioBusAssetDashboardEntry : public IObjectDashboardEntry
+	{
+		virtual ~FAudioBusAssetDashboardEntry() = default;
+
+		virtual FText GetDisplayName() const override;
+		virtual const UObject* GetObject() const override;
+		virtual UObject* GetObject() override;
+		virtual bool IsValid() const override;
+
+		::Audio::FDeviceId DeviceId = INDEX_NONE;
+		uint64 ComponentId = TNumericLimits<uint64>::Max();
+		FString Name;
+		TWeakObjectPtr<UAudioBus> AudioBus;
+	};
+} // namespace UE::Audio::Insights
diff --git a/Engine/Plugins/AudioInsights/Source/AudioInsightsEditor/Private/Views/AudioBusesDashboardViewFactory.cpp b/Engine/Plugins/AudioInsights/Source/AudioInsightsEditor/Private/Views/AudioBusesDashboardViewFactory.cpp
new file mode 100644
--- /dev/null
+++ b/Engine/Plugins/AudioInsights/Source/AudioInsightsEditor/Private/Views/AudioBusesDashboardViewFactory.cpp
@@ -0,0 +1,217 @@
+// Copyright Epic Games, Inc. All Rights Reserved.
+#include "Views/AudioBusesDashboardViewFactory.h"
+
+#include "AudioInsightsStyle.h"
+#include "Editor.h"
+#include "Internationalization/Text.h"
+#include "Providers/AudioBusProvider.h"
+#include "Sound/AudioBus.h"
+#include "Subsystems/AssetEditorSubsystem.h"
+#include "Templates/SharedPointer.h"
+
+#define LOCTEXT_NAMESPACE "AudioInsights"
+
+namespace UE::Audio::Insights
+{
+	namespace AudioBusesPrivate
+	{
+		const FAudioBusAssetDashboardEntry& CastEntry(const IDashboardDataViewEntry& InData)
+		{
+			return static_cast<const FAudioBusAssetDashboardEntry&>(InData);
+		};
+	}
+
+	FAudioBusesDashboardViewFactory::FAudioBusesDashboardViewFactory()
+	{
+		FAudioBusProvider::OnAudioBusAssetAdded.AddRaw(this, &FAudioBusesDashboardViewFactory::HandleOnAudioBusAssetListUpdated);
+		FAudioBusProvider::OnAudioBusAssetRemoved.AddRaw(this, &FAudioBusesDashboardViewFactory::HandleOnAudioBusAssetListUpdated);
+		FAudioBusProvider::OnAudioBusAssetListUpdated.AddRaw(this, &FAudioBusesDashboardViewFactory::RequestListRefresh);
+
+		Providers = TArray<TSharedPtr<FTraceProviderBase>>
+		{
+			MakeShared<FAudioBusProvider>()
+		};
+
+		SortByColumn = "Name";
+		SortMode     = EColumnSortMode::Ascending;
+	}
+
+	FAudioBusesDashboardViewFactory::~FAudioBusesDashboardViewFactory()
+	{
+		FAudioBusProvider::OnAudioBusAssetAdded.RemoveAll(this);
+		FAudioBusProvider::OnAudioBusAssetRemoved.RemoveAll(this);
+		FAudioBusProvider::OnAudioBusAssetListUpdated.RemoveAll(this);
+	}
+
+	FName FAudioBusesDashboardViewFactory::GetName() const
+	{
+		return "AudioBuses";
+	}
+
+	FText FAudioBusesDashboardViewFactory::GetDisplayName() const
+	{
+		return LOCTEXT("AudioDashboard_AudioBuses_DisplayName", "Audio Buses");
+	}
+
+	TSharedRef<SWidget> FAudioBusesDashboardViewFactory::GenerateWidgetForColumn(TSharedRef<IDashboardDataViewEntry> InRowData, const FName& InColumnName)
+	{
+		if (InColumnName == "Name")
+		{
+			const TWeakObjectPtr<UAudioBus> AudioBus = AudioBusesPrivate::CastEntry(InRowData.Get()).AudioBus;
+
+			AudioBusCheckboxCheckedStates.FindOrAdd(AudioBus, false);
+
+			const FColumnData& ColumnData = GetColumns()[InColumnName];
+			const FText Label = ColumnData.GetDisplayValue(InRowData.Get());
+
+			return (Label.IsEmpty()) ? 
+				SNullWidget::NullWidget :
+				SNew(SHorizontalBox)
+				.Clipping(EWidgetClipping::ClipToBounds)
+				+ SHorizontalBox::Slot()
+				.FillWidth(0.1f)
+				[
+					SNew(SCheckBox)
+					.IsChecked(AudioBusCheckboxCheckedStates.Contains(AudioBus) && AudioBusCheckboxCheckedStates[AudioBus] == true)
+					.OnCheckStateChanged_Lambda([AudioBus, this](ECheckBoxState NewState)
+					{
+						if (AudioBusCheckboxCheckedStates.Contains(AudioBus))
+						{
+							AudioBusCheckboxCheckedStates[AudioBus] = NewState == ECheckBoxState::Checked;
+						}
+
+						OnAudioBusAssetChecked.Broadcast(NewState == ECheckBoxState::Checked, AudioBus);
+					})
+				]
+				+ SHorizontalBox::Slot()
+				.FillWidth(0.8f)
+				[
+					SNew(STextBlock)
+					.Text(Label)
+					.MinDesiredWidth(300)
+					.OnDoubleClicked_Lambda([this, InRowData](const FGeometry& MyGeometry, const FPointerEvent& PointerEvent)
+					{
+						if (GEditor)
+						{
+							if (TSharedPtr<IObjectDashboardEntry> ObjectData = StaticCastSharedPtr<IObjectDashboardEntry>(InRowData.ToSharedPtr()); 
+								ObjectData.IsValid())
+							{
+								if (const UObject* Object = ObjectData->GetObject();
+									Object && Object->IsAsset())
+								{
+									GEditor->GetEditorSubsystem<UAssetEditorSubsystem>()->OpenEditorForAsset(Object);
+									return FReply::Handled();
+								}
+							}
+						}
+
+						return FReply::Unhandled();
+					})
+				];
+		}
+
+		return SNullWidget::NullWidget;
+	}
+
+	void FAudioBusesDashboardViewFactory::ProcessEntries(FTraceTableDashboardViewFactory::EProcessReason Reason)
+	{
+		const FString FilterString = GetSearchFilterText().ToString();
+		
+		FTraceTableDashboardViewFactory::FilterEntries<FAudioBusProvider>([&FilterString](const IDashboardDataViewEntry& Entry)
+		{
+			const FAudioBusAssetDashboardEntry& AudioBusEntry = static_cast<const FAudioBusAssetDashboardEntry&>(Entry);
+			
+			return !AudioBusEntry.GetDisplayName().ToString().Contains(FilterString);
+		});
+	}
+
+	FSlateIcon FAudioBusesDashboardViewFactory::GetIcon() const
+	{
+		return FSlateStyle::Get().CreateIcon("AudioInsights.Icon");
+	}
+
+	EDefaultDashboardTabStack FAudioBusesDashboardViewFactory::GetDefaultTabStack() const
+	{
+		return EDefaultDashboardTabStack::Analysis;
+	}
+
+	TSharedRef<SWidget> FAudioBusesDashboardViewFactory::MakeWidget()
+	{
+		if (!DashboardWidget.IsValid())
+		{
+			DashboardWidget = FTraceTableDashboardViewFactory::MakeWidget();
+
+			if (FilteredEntriesListView.IsValid())
+			{
+				FilteredEntriesListView->SetSelectionMode(ESelectionMode::Single);
+			}
+		}
+
+		return DashboardWidget->AsShared();
+	}
+
+	const TMap<FName, FTraceTableDashboardViewFactory::FColumnData>& FAudioBusesDashboardViewFactory::GetColumns() const
+	{
+		auto CreateColumnData = []()
+		{
+			return TMap<FName, FTraceTableDashboardViewFactory::FColumnData>
+			{
+				{
+					"Name",
+					{
+						LOCTEXT("AudioBuses_NameColumnDisplayName", "Name"),
+						[](const IDashboardDataViewEntry& InData) { return AudioBusesPrivate::CastEntry(InData).GetDisplayName(); },
+						false,		/* bDefaultHidden */
+						0.88f,		/* FillWidth */
+						EHorizontalAlignment::HAlign_Left
+					}
+				}
+			};
+		};
+		
+		static const TMap<FName, FTraceTableDashboardViewFactory::FColumnData> ColumnData = CreateColumnData();
+		
+		return ColumnData;
+	}
+
+	void FAudioBusesDashboardViewFactory::SortTable()
+	{
+		if (SortByColumn == "Name")
+		{
+			if (SortMode == EColumnSortMode::Ascending)
+			{
+				DataViewEntries.Sort([](const TSharedPtr<IDashboardDataViewEntry>& A, const TSharedPtr<IDashboardDataViewEntry>& B)
+				{
+					const FAudioBusAssetDashboardEntry& AData = AudioBusesPrivate::CastEntry(*A.Get());
+					const FAudioBusAssetDashboardEntry& BData = AudioBusesPrivate::CastEntry(*B.Get());
+
+					return AData.GetDisplayName().CompareToCaseIgnored(BData.GetDisplayName()) < 0;
+				});
+			}
+			else if (SortMode == EColumnSortMode::Descending)
+			{
+				DataViewEntries.Sort([](const TSharedPtr<IDashboardDataViewEntry>& A, const TSharedPtr<IDashboardDataViewEntry>& B)
+				{
+					const FAudioBusAssetDashboardEntry& AData = AudioBusesPrivate::CastEntry(*A.Get());
+					const FAudioBusAssetDashboardEntry& BData = AudioBusesPrivate::CastEntry(*B.Get());
+
+					return BData.GetDisplayName().CompareToCaseIgnored(AData.GetDisplayName()) < 0;
+				});
+			}
+		}
+	}
+
+	void FAudioBusesDashboardViewFactory::RequestListRefresh()
+	{
+		if (FilteredEntriesListView.IsValid())
+		{
+			FilteredEntriesListView->RequestListRefresh();
+		}
+	}
+
+	void FAudioBusesDashboardViewFactory::HandleOnAudioBusAssetListUpdated(const TWeakObjectPtr<UObject> InAsset)
+	{
+		RequestListRefresh();
+	}
+} // namespace UE::Audio::Insights
+#undef LOCTEXT_NAMESPACE
diff --git a/Engine/Plugins/AudioInsights/Source/AudioInsightsEditor/Private/Views/AudioBusesDashboardViewFactory.h b/Engine/Plugins/AudioInsights/Source/AudioInsightsEditor/Private/Views/AudioBusesDashboardViewFactory.h
new file mode 100644
--- /dev/null
+++ b/Engine/Plugins/AudioInsights/Source/AudioInsightsEditor/Private/Views/AudioBusesDashboardViewFactory.h
@@ -0,0 +1,43 @@
+// Copyright Epic Games, Inc. All Rights Reserved.
+#pragma once
+
+#include "Delegates/Delegate.h"
+#include "Templates/SharedPointer.h"
+#include "Views/AudioBusAssetDashboardEntry.h"
+#include "Views/AudioBusAssetDashboardEntry.h"
+#include "Views/TableDashboardViewFactory.h"
+#include "Widgets/Input/SCheckBox.h"
+
+class UAudioBus;
+
+namespace UE::Audio::Insights
+{
+	class FAudioBusesDashboardViewFactory : public FTraceObjectTableDashboardViewFactory
+	{
+	public:
+		FAudioBusesDashboardViewFactory();
+		virtual ~FAudioBusesDashboardViewFactory();
+
+		virtual FName GetName() const override;
+		virtual FText GetDisplayName() const override;
+		virtual FSlateIcon GetIcon() const override;
+		virtual EDefaultDashboardTabStack GetDefaultTabStack() const override;
+		virtual TSharedRef<SWidget> MakeWidget() override;
+
+		DECLARE_MULTICAST_DELEGATE_TwoParams(FOnAudioBusAssetChecked, const bool /*bIsChecked*/, const TWeakObjectPtr<UAudioBus> /*AudioBus*/);
+		inline static FOnAudioBusAssetChecked OnAudioBusAssetChecked;
+
+	protected:
+		virtual TSharedRef<SWidget> GenerateWidgetForColumn(TSharedRef<IDashboardDataViewEntry> InRowData, const FName& InColumnName) override;
+		virtual void ProcessEntries(FTraceTableDashboardViewFactory::EProcessReason Reason) override;
+		virtual const TMap<FName, FTraceTableDashboardViewFactory::FColumnData>& GetColumns() const override;
+
+		virtual void SortTable() override;
+
+		void RequestListRefresh();
+
+		void HandleOnAudioBusAssetListUpdated(const TWeakObjectPtr<UObject> InAsset);
+
+		TMap<const TWeakObjectPtr<UAudioBus>, bool> AudioBusCheckboxCheckedStates;
+	};
+} // namespace UE::Audio::Insights
diff --git a/Engine/Plugins/AudioInsights/Source/AudioInsightsEditor/Private/Views/AudioMeterView.cpp b/Engine/Plugins/AudioInsights/Source/AudioInsightsEditor/Private/Views/AudioMeterView.cpp
new file mode 100644
--- /dev/null
+++ b/Engine/Plugins/AudioInsights/Source/AudioInsightsEditor/Private/Views/AudioMeterView.cpp
@@ -0,0 +1,177 @@
+// Copyright Epic Games, Inc. All Rights Reserved.
+#include "AudioMeterView.h"
+
+#include "Analyzers/AudioMeterSubmixAnalyzer.h"
+#include "AudioInsightsEditorDashboardFactory.h"
+#include "AudioMeter.h"
+#include "Sound/AudioBus.h"
+#include "Sound/SoundSubmix.h"
+#include "Widgets/SNullWidget.h"
+
+namespace UE::Audio::Insights
+{
+	FAudioMeterView::FAudioMeterView(FAudioAssetVariant InAudioAssetVariant)
+		: AudioAssetVariant(InAudioAssetVariant)
+		, AudioMeterAnalyzerVariant(MakeAudioMeterAnalyzerVariant(InAudioAssetVariant))
+		, AudioAssetNameTextBlock(MakeAudioAssetNameTextBlock(InAudioAssetVariant))
+		, AudioMeterViewWidget(MakeWidget())
+		, OnActiveAudioDeviceChangedHandle(FEditorDashboardFactory::OnActiveAudioDeviceChanged.AddRaw(this, &FAudioMeterView::HandleOnActiveAudioDeviceChanged))
+	{
+		
+	}
+
+	FAudioMeterView::~FAudioMeterView()
+	{
+		FEditorDashboardFactory::OnActiveAudioDeviceChanged.Remove(OnActiveAudioDeviceChangedHandle);
+	}
+
+	FAudioMeterView::FAudioMeterVariant FAudioMeterView::MakeAudioMeterAnalyzerVariant(const FAudioAssetVariant InAudioAssetVariant)
+	{
+		switch (AudioAssetVariant.GetIndex())
+		{
+			case FAudioAssetVariant::IndexOfType<TWeakObjectPtr<UAudioBus>>():
+			{
+				return FAudioMeterVariant(TInPlaceType<TSharedPtr<FAudioMeterAnalyzer>>(), MakeShared<FAudioMeterAnalyzer>(InAudioAssetVariant.Get<TWeakObjectPtr<UAudioBus>>()));
+			}
+
+			case FAudioAssetVariant::IndexOfType<TWeakObjectPtr<USoundSubmix>>():
+			{
+				return FAudioMeterVariant(TInPlaceType<TSharedPtr<FAudioMeterSubmixAnalyzer>>(), MakeShared<FAudioMeterSubmixAnalyzer>(InAudioAssetVariant.Get<TWeakObjectPtr<USoundSubmix>>()));
+			}
+
+			default:
+				return FAudioMeterView::FAudioMeterVariant();
+		}
+	}
+
+	TSharedRef<STextBlock> FAudioMeterView::MakeAudioAssetNameTextBlock(const FAudioAssetVariant InAudioAssetVariant)
+	{
+		switch (InAudioAssetVariant.GetIndex())
+		{
+			case FAudioAssetVariant::IndexOfType<TWeakObjectPtr<UAudioBus>>():
+			{
+				if (TWeakObjectPtr<UAudioBus> AudioBusAsset = InAudioAssetVariant.Get<TWeakObjectPtr<UAudioBus>>();
+					AudioBusAsset.IsValid())
+				{
+					return SNew(STextBlock)
+						.Text_Lambda([AudioBusAsset] { return AudioBusAsset.IsValid() ? FText::FromString(AudioBusAsset->GetName()) : FText::GetEmpty(); })
+						.Justification(ETextJustify::Center)
+						.ColorAndOpacity(FSlateColor(FColor(80, 200, 255)));
+				}
+			}
+			break;
+
+			case FAudioAssetVariant::IndexOfType<TWeakObjectPtr<USoundSubmix>>():
+			{
+				if (TWeakObjectPtr<USoundSubmix> SoundSubmixAsset = AudioAssetVariant.Get<TWeakObjectPtr<USoundSubmix>>();
+					SoundSubmixAsset.IsValid())
+				{
+					return SNew(STextBlock)
+						.Text_Lambda([SoundSubmixAsset] { return SoundSubmixAsset.IsValid() ? FText::FromString(SoundSubmixAsset->GetName()) : FText::GetEmpty(); })
+						.Justification(ETextJustify::Center)
+						.ColorAndOpacity(FSlateColor(FColor(30, 240, 90)));
+				}
+			}
+			break;
+		}
+
+		return SNew(STextBlock);
+	}
+
+	TSharedRef<SWidget> FAudioMeterView::GetAudioMeterAnalyzerWidget() const
+	{
+		switch (AudioMeterAnalyzerVariant.GetIndex())
+		{
+			case FAudioMeterVariant::IndexOfType<TSharedPtr<FAudioMeterAnalyzer>>():
+			{
+				TSharedPtr<FAudioMeterAnalyzer> AudioMeterAnalyzer = AudioMeterAnalyzerVariant.Get<TSharedPtr<FAudioMeterAnalyzer>>();
+				return AudioMeterAnalyzer.IsValid() ? AudioMeterAnalyzer->GetAudioMeter()->GetWidget() : SNullWidget::NullWidget;
+			}
+
+			case FAudioMeterVariant::IndexOfType<TSharedPtr<FAudioMeterSubmixAnalyzer>>():
+			{
+				TSharedPtr<FAudioMeterSubmixAnalyzer> AudioMeterSubmixAnalyzer = AudioMeterAnalyzerVariant.Get<TSharedPtr<FAudioMeterSubmixAnalyzer>>();
+				return AudioMeterSubmixAnalyzer.IsValid() ? AudioMeterSubmixAnalyzer->GetWidget() : SNullWidget::NullWidget;
+			}
+
+			default:
+				return SNullWidget::NullWidget;
+		}
+	}
+
+	TSharedRef<SWidget> FAudioMeterView::MakeWidget()
+	{
+		return SNew(SVerticalBox)
+			+ SVerticalBox::Slot()
+			[
+				SNew(SVerticalBox)
+				.Clipping(EWidgetClipping::ClipToBounds)
+				// Height padding
+				+ SVerticalBox::Slot()
+				.FillHeight(0.05)
+				[
+					SNew(SBox)
+				]
+				// Audio Meter container
+				+ SVerticalBox::Slot()
+				.FillHeight(0.8)
+				[
+					// Audio Meter
+					SNew(SHorizontalBox)
+					+ SHorizontalBox::Slot()
+					.HAlign(HAlign_Center)
+					.VAlign(VAlign_Fill)
+					[
+						GetAudioMeterAnalyzerWidget()
+					]
+				]
+				// Height padding
+				+ SVerticalBox::Slot()
+				.FillHeight(0.025)
+				[
+					SNew(SBox)
+				]
+				// Audio aaset name label
+				+ SVerticalBox::Slot()
+				.FillHeight(0.1)
+				[
+					AudioAssetNameTextBlock
+				]
+				// Height padding
+				+ SVerticalBox::Slot()
+				.FillHeight(0.025)
+				[
+					SNew(SBox)
+				]
+			];
+	}
+
+	void FAudioMeterView::HandleOnActiveAudioDeviceChanged()
+	{
+		switch (AudioAssetVariant.GetIndex())
+		{
+			case FAudioAssetVariant::IndexOfType<TWeakObjectPtr<UAudioBus>>():
+			{
+				if (TSharedPtr<FAudioMeterAnalyzer> AudioMeterAnalyzer = AudioMeterAnalyzerVariant.Get<TSharedPtr<FAudioMeterAnalyzer>>();
+					AudioMeterAnalyzer.IsValid())
+				{
+					AudioMeterAnalyzer->RebuildAudioMeter(AudioAssetVariant.Get<TWeakObjectPtr<UAudioBus>>());
+				}
+			}
+			break;
+
+			case FAudioAssetVariant::IndexOfType<TWeakObjectPtr<USoundSubmix>>():
+			{
+				if (TSharedPtr<FAudioMeterSubmixAnalyzer> AudioMeterSubmixAnalyzer = AudioMeterAnalyzerVariant.Get<TSharedPtr<FAudioMeterSubmixAnalyzer>>();
+					AudioMeterSubmixAnalyzer.IsValid())
+				{
+					AudioMeterSubmixAnalyzer->SetSubmix(AudioAssetVariant.Get<TWeakObjectPtr<USoundSubmix>>());
+				}
+			}
+			break;
+
+			default:
+				break;
+		}
+	}
+} // namespace UE::Audio::Insights
diff --git a/Engine/Plugins/AudioInsights/Source/AudioInsightsEditor/Private/Views/AudioMeterView.h b/Engine/Plugins/AudioInsights/Source/AudioInsightsEditor/Private/Views/AudioMeterView.h
new file mode 100644
--- /dev/null
+++ b/Engine/Plugins/AudioInsights/Source/AudioInsightsEditor/Private/Views/AudioMeterView.h
@@ -0,0 +1,46 @@
+// Copyright Epic Games, Inc. All Rights Reserved.
+#pragma once
+
+#include "Templates/SharedPointer.h"
+#include "UObject/WeakObjectPtr.h"
+#include "Widgets/SWidget.h"
+#include "Widgets/Text/STextBlock.h"
+
+class UAudioBus;
+class USoundSubmix;
+
+namespace UE::Audio::Insights
+{
+	class FAudioMeterAnalyzer;
+	class FAudioMeterSubmixAnalyzer;
+
+	class FAudioMeterView : public TSharedFromThis<FAudioMeterView>
+	{
+	public:
+		using FAudioAssetVariant = TVariant<TWeakObjectPtr<UAudioBus>, TWeakObjectPtr<USoundSubmix>>;
+		using FAudioMeterVariant = TVariant<TSharedPtr<FAudioMeterAnalyzer>, TSharedPtr<FAudioMeterSubmixAnalyzer>>;
+
+		FAudioMeterView(FAudioAssetVariant InAudioAssetVariant);
+		virtual ~FAudioMeterView();
+
+		TSharedRef<SWidget> GetWidget() const { return AudioMeterViewWidget; };
+
+	private:
+		FAudioMeterVariant MakeAudioMeterAnalyzerVariant(const FAudioAssetVariant InAudioAssetVariant);
+		TSharedRef<STextBlock> MakeAudioAssetNameTextBlock(const FAudioAssetVariant InAudioAssetVariant);
+		TSharedRef<SWidget> MakeWidget();
+
+		TSharedRef<SWidget> GetAudioMeterAnalyzerWidget() const;
+
+		void HandleOnActiveAudioDeviceChanged();
+
+		FAudioAssetVariant AudioAssetVariant;
+		FAudioMeterVariant AudioMeterAnalyzerVariant;
+
+		TSharedRef<STextBlock> AudioAssetNameTextBlock;
+
+		TSharedRef<SWidget> AudioMeterViewWidget;
+
+		FDelegateHandle OnActiveAudioDeviceChangedHandle;
+	};
+} // namespace UE::Audio::Insights
diff --git a/Engine/Plugins/AudioInsights/Source/AudioInsightsEditor/Private/Views/AudioMetersDashboardViewFactory.cpp b/Engine/Plugins/AudioInsights/Source/AudioInsightsEditor/Private/Views/AudioMetersDashboardViewFactory.cpp
new file mode 100644
--- /dev/null
+++ b/Engine/Plugins/AudioInsights/Source/AudioInsightsEditor/Private/Views/AudioMetersDashboardViewFactory.cpp
@@ -0,0 +1,162 @@
+// Copyright Epic Games, Inc. All Rights Reserved.
+#include "AudioMetersDashboardViewFactory.h"
+
+#include "AudioInsightsStyle.h"
+#include "Providers/AudioBusProvider.h"
+#include "Providers/SoundSubmixProvider.h"
+#include "Sound/AudioBus.h"
+#include "Sound/AudioSettings.h"
+#include "Sound/SoundSubmix.h"
+#include "Views/AudioBusesDashboardViewFactory.h"
+#include "Views/SubmixesDashboardViewFactory.h"
+
+#define LOCTEXT_NAMESPACE "AudioInsights"
+
+namespace UE::Audio::Insights
+{
+	FName FAudioMetersDashboardViewFactory::GetName() const
+	{
+		return "AudioMeters";
+	}
+
+	FText FAudioMetersDashboardViewFactory::GetDisplayName() const
+	{
+		return LOCTEXT("AudioDashboard_AudioMetersTab_DisplayName", "Audio Meters");
+	}
+
+	EDefaultDashboardTabStack FAudioMetersDashboardViewFactory::GetDefaultTabStack() const
+	{
+		return EDefaultDashboardTabStack::AudioMeters;
+	}
+
+	FSlateIcon FAudioMetersDashboardViewFactory::GetIcon() const
+	{
+		return FSlateStyle::Get().CreateIcon("AudioInsights.Icon.Submix");
+	}
+
+	TSharedRef<SWidget> FAudioMetersDashboardViewFactory::MakeWidget()
+	{
+		if (!OnSubmixAssetCheckedHandle.IsValid())
+		{
+			OnSubmixAssetCheckedHandle = FSubmixesDashboardViewFactory::OnSubmixAssetChecked.AddSP(this, &FAudioMetersDashboardViewFactory::HandleOnSubmixAssetChecked);
+		}
+
+		if (!OnAudioBusAssetCheckedHandle.IsValid())
+		{
+			OnAudioBusAssetCheckedHandle = FAudioBusesDashboardViewFactory::OnAudioBusAssetChecked.AddSP(this, &FAudioMetersDashboardViewFactory::HandleOnAudioBusAssetChecked);
+		}
+
+		if (!OnSubmixAssetRemovedHandle.IsValid())
+		{
+			OnSubmixAssetRemovedHandle = FSoundSubmixProvider::OnSubmixAssetRemoved.AddSP(this, &FAudioMetersDashboardViewFactory::HandleOnAssetRemoved);
+		}
+
+		if (!OnAudioBusAssetRemovedHandle.IsValid())
+		{
+			OnAudioBusAssetRemovedHandle = FAudioBusProvider::OnAudioBusAssetRemoved.AddSP(this, &FAudioMetersDashboardViewFactory::HandleOnAssetRemoved);
+		}
+		
+		if (!MeterViewsScrollBox.IsValid())
+		{
+			SAssignNew(MeterViewsScrollBox, SScrollBox)
+			.Orientation(Orient_Horizontal)
+			+ SScrollBox::Slot()
+			[
+				SAssignNew(AudioMeterViewsContainer, SHorizontalBox)
+			];
+		}
+
+		return MeterViewsScrollBox.ToSharedRef();
+	}
+
+	void FAudioMetersDashboardViewFactory::HandleOnSubmixAssetChecked(const bool bInIsChecked, const TWeakObjectPtr<USoundSubmix> InSoundSubmix)
+	{
+		if (!InSoundSubmix.IsValid())
+		{
+			return;
+		}
+
+		const uint32 SubmixUniqueId = InSoundSubmix->GetUniqueID();
+
+		if (bInIsChecked)
+		{
+			AudioMeterViews.Add(SubmixUniqueId, MakeShared<FAudioMeterView>(FAudioMeterView::FAudioAssetVariant(TInPlaceType<TWeakObjectPtr<USoundSubmix>>(), InSoundSubmix)));
+
+			AudioMeterViewsContainer->AddSlot()
+			.AutoWidth()
+			.Padding(10.0f, 0.0f, 10.0f, 0.0f)
+			[
+				AudioMeterViews[SubmixUniqueId]->GetWidget()
+			];
+		}
+		else
+		{
+			if (AudioMeterViews.Contains(SubmixUniqueId))
+			{
+				AudioMeterViewsContainer->RemoveSlot(AudioMeterViews[SubmixUniqueId]->GetWidget());
+				AudioMeterViews.Remove(SubmixUniqueId);
+			}
+		}
+
+		if (MeterViewsScrollBox.IsValid())
+		{
+			MeterViewsScrollBox->Invalidate(EInvalidateWidget::Layout);
+		}
+	}
+
+	void FAudioMetersDashboardViewFactory::HandleOnAudioBusAssetChecked(const bool bInIsChecked, const TWeakObjectPtr<UAudioBus> InAudioBus)
+	{
+		if (!InAudioBus.IsValid())
+		{
+			return;
+		}
+
+		const uint32 AudioBusUniqueId = InAudioBus->GetUniqueID();
+
+		if (bInIsChecked)
+		{
+			AudioMeterViews.Add(AudioBusUniqueId, MakeShared<FAudioMeterView>(FAudioMeterView::FAudioAssetVariant(TInPlaceType<TWeakObjectPtr<UAudioBus>>(), InAudioBus)));
+
+			AudioMeterViewsContainer->AddSlot()
+			.AutoWidth()
+			.Padding(10.0f, 0.0f, 10.0f, 0.0f)
+			[
+				AudioMeterViews[AudioBusUniqueId]->GetWidget()
+			];
+		}
+		else
+		{
+			if (AudioMeterViews.Contains(AudioBusUniqueId))
+			{
+				AudioMeterViewsContainer->RemoveSlot(AudioMeterViews[AudioBusUniqueId]->GetWidget());
+				AudioMeterViews.Remove(AudioBusUniqueId);
+			}
+		}
+
+		if (MeterViewsScrollBox.IsValid())
+		{
+			MeterViewsScrollBox->Invalidate(EInvalidateWidget::Layout);
+		}
+	}
+
+	void FAudioMetersDashboardViewFactory::HandleOnAssetRemoved(const TWeakObjectPtr<UObject> InAsset)
+	{
+		if (!InAsset.IsValid())
+		{
+			return;
+		}
+
+		if (const uint32 AssetUniqueId = InAsset->GetUniqueID();
+			AudioMeterViews.Contains(AssetUniqueId))
+		{
+			AudioMeterViewsContainer->RemoveSlot(AudioMeterViews[AssetUniqueId]->GetWidget());
+			AudioMeterViews.Remove(AssetUniqueId);
+		}
+
+		if (MeterViewsScrollBox.IsValid())
+		{
+			MeterViewsScrollBox->Invalidate(EInvalidateWidget::Layout);
+		}
+	}
+} // namespace UE::Audio::Insights
+#undef LOCTEXT_NAMESPACE
diff --git a/Engine/Plugins/AudioInsights/Source/AudioInsightsEditor/Private/Views/AudioMetersDashboardViewFactory.h b/Engine/Plugins/AudioInsights/Source/AudioInsightsEditor/Private/Views/AudioMetersDashboardViewFactory.h
new file mode 100644
--- /dev/null
+++ b/Engine/Plugins/AudioInsights/Source/AudioInsightsEditor/Private/Views/AudioMetersDashboardViewFactory.h
@@ -0,0 +1,43 @@
+// Copyright Epic Games, Inc. All Rights Reserved.
+#pragma once
+
+#include "AudioMeterView.h"
+#include "Templates/SharedPointer.h"
+#include "Views/DashboardViewFactory.h"
+#include "Widgets/Layout/SScrollBox.h"
+
+class SAudioVolumeSlider;
+class UAudioBus;
+class USoundSubmix;
+
+namespace UE::Audio::Insights
+{
+	class FAudioMeterAnalyzer;
+
+	class FAudioMetersDashboardViewFactory : public IDashboardViewFactory, public TSharedFromThis<FAudioMetersDashboardViewFactory>
+	{
+	public:
+		virtual ~FAudioMetersDashboardViewFactory() = default;
+
+		virtual FName GetName() const override;
+		virtual FText GetDisplayName() const override;
+		virtual EDefaultDashboardTabStack GetDefaultTabStack() const override;
+		virtual FSlateIcon GetIcon() const override;
+		virtual TSharedRef<SWidget> MakeWidget() override;
+
+	private:
+		void HandleOnSubmixAssetChecked(const bool bInIsChecked, const TWeakObjectPtr<USoundSubmix> InSoundSubmix);
+		void HandleOnAudioBusAssetChecked(const bool bInIsChecked, const TWeakObjectPtr<UAudioBus> InAudioBus);
+
+		void HandleOnAssetRemoved(const TWeakObjectPtr<UObject> InAsset);
+
+		TSharedPtr<SScrollBox> MeterViewsScrollBox;
+		TSharedPtr<SHorizontalBox> AudioMeterViewsContainer;
+		TMap<uint32, TSharedRef<FAudioMeterView>> AudioMeterViews;
+
+		FDelegateHandle OnSubmixAssetCheckedHandle;
+		FDelegateHandle OnAudioBusAssetCheckedHandle;
+		FDelegateHandle OnSubmixAssetRemovedHandle;
+		FDelegateHandle OnAudioBusAssetRemovedHandle;
+	};
+} // namespace UE::Audio::Insights
diff --git a/Engine/Plugins/AudioInsights/Source/AudioInsightsEditor/Private/Views/MixerSourceMuteSoloFilter.cpp b/Engine/Plugins/AudioInsights/Source/AudioInsightsEditor/Private/Views/MixerSourceMuteSoloFilter.cpp
new file mode 100644
--- /dev/null
+++ b/Engine/Plugins/AudioInsights/Source/AudioInsightsEditor/Private/Views/MixerSourceMuteSoloFilter.cpp
@@ -0,0 +1,49 @@
+// Copyright Epic Games, Inc. All Rights Reserved.
+#include "MixerSourceMuteSoloFilter.h"
+
+#include "Audio/AudioDebug.h"
+#include "AudioDeviceManager.h"
+#include "Views/MixerSourceDashboardViewFactory.h"
+
+namespace UE::Audio::Insights
+{
+	FMuteSoloFilter::FMuteSoloFilter()
+	{
+		FMixerSourceDashboardViewFactory::OnUpdateMuteSoloState.AddRaw(this, &FMuteSoloFilter::FilterMuteSolo);
+	}
+
+	FMuteSoloFilter::~FMuteSoloFilter()
+	{
+		FMixerSourceDashboardViewFactory::OnUpdateMuteSoloState.RemoveAll(this);
+	}
+
+	void FMuteSoloFilter::FilterMuteSolo(ECheckBoxState InMuteState, ECheckBoxState InSoloState, const FString& InCurrentFilterString) const
+	{
+#if ENABLE_AUDIO_DEBUG
+		if (FAudioDeviceManager* AudioDeviceManager = FAudioDeviceManager::Get())
+		{
+			::Audio::FAudioDebugger& AudioDebugger = AudioDeviceManager->GetDebugger();
+
+			FName CurrentFilterStringName{ InCurrentFilterString };
+			
+			if (InMuteState == ECheckBoxState::Checked && !InCurrentFilterString.IsEmpty())
+			{
+				AudioDebugger.ToggleMuteSoundWave(CurrentFilterStringName, true);
+			}
+			else
+			{
+				AudioDebugger.ToggleMuteSoundWave(NAME_None, true);
+			}
+
+			if (InSoloState == ECheckBoxState::Checked && !InCurrentFilterString.IsEmpty())
+			{
+				AudioDebugger.ToggleSoloSoundWave(CurrentFilterStringName, true);
+			}
+			else
+			{
+				AudioDebugger.ToggleSoloSoundWave(NAME_None, true);
+			}
+		}
+#endif // ENABLE_AUDIO_DEBUG
+	}
+} // namespace UE::Audio::Insights
diff --git a/Engine/Plugins/AudioInsights/Source/AudioInsightsEditor/Private/Views/MixerSourceMuteSoloFilter.h b/Engine/Plugins/AudioInsights/Source/AudioInsightsEditor/Private/Views/MixerSourceMuteSoloFilter.h
new file mode 100644
--- /dev/null
+++ b/Engine/Plugins/AudioInsights/Source/AudioInsightsEditor/Private/Views/MixerSourceMuteSoloFilter.h
@@ -0,0 +1,18 @@
+// Copyright Epic Games, Inc. All Rights Reserved.
+#pragma once 
+
+#include "HAL/Platform.h"
+#include "Styling/SlateTypes.h"
+
+namespace UE::Audio::Insights
+{
+	class FMuteSoloFilter
+	{
+	public:
+		FMuteSoloFilter();
+		~FMuteSoloFilter();
+		
+	private:
+		void FilterMuteSolo(ECheckBoxState InMuteState, ECheckBoxState InSoloState, const FString& InCurrentFilterString) const;
+	};
+} // namespace UE::Audio::Insights
diff --git a/Engine/Plugins/AudioInsights/Source/AudioInsightsEditor/Private/Views/OutputMeterDashboardViewFactory.cpp b/Engine/Plugins/AudioInsights/Source/AudioInsightsEditor/Private/Views/OutputMeterDashboardViewFactory.cpp
new file mode 100644
--- /dev/null
+++ b/Engine/Plugins/AudioInsights/Source/AudioInsightsEditor/Private/Views/OutputMeterDashboardViewFactory.cpp
@@ -0,0 +1,93 @@
+// Copyright Epic Games, Inc. All Rights Reserved.
+
+#include "Views/OutputMeterDashboardViewFactory.h"
+
+#include "Analyzers/AudioMeterSubmixAnalyzer.h"
+#include "AudioInsightsEditorDashboardFactory.h"
+#include "AudioInsightsStyle.h"
+#include "SAudioMeter.h"
+#include "Sound/AudioSettings.h"
+#include "Sound/SoundSubmix.h"
+#include "SubmixesDashboardViewFactory.h"
+#include "UObject/WeakObjectPtrTemplates.h"
+#include "Widgets/Colors/SColorBlock.h"
+#include "Widgets/SBoxPanel.h"
+#include "Widgets/SOverlay.h"
+
+#define LOCTEXT_NAMESPACE "AudioInsights"
+
+namespace UE::Audio::Insights
+{
+	FName FOutputMeterDashboardViewFactory::GetName() const
+	{
+		return "OutputMeter";
+	}
+
+	FText FOutputMeterDashboardViewFactory::GetDisplayName() const
+	{
+		return LOCTEXT("AudioDashboard_DashboardsAudioMeterTab_DisplayName", "Output Meter");
+	}
+
+	EDefaultDashboardTabStack FOutputMeterDashboardViewFactory::GetDefaultTabStack() const
+	{
+		return EDefaultDashboardTabStack::AudioMeter;
+	}
+
+	FSlateIcon FOutputMeterDashboardViewFactory::GetIcon() const
+	{
+		return FSlateStyle::Get().CreateIcon("AudioInsights.Icon");
+	}
+
+	TSharedRef<SWidget> FOutputMeterDashboardViewFactory::MakeWidget()
+	{
+		// Create OutputMeterSubmixAnalyzer (use MainSubmix as default)
+		if (TWeakObjectPtr<UAudioSettings> AudioSettings = GetMutableDefault<UAudioSettings>();
+			AudioSettings.IsValid())
+		{
+			MainSubmix = Cast<USoundSubmix>(AudioSettings->MasterSubmix.ResolveObject());
+
+			if (MainSubmix)
+			{
+				OutputMeterSubmixAnalyzer = MakeShared<FAudioMeterSubmixAnalyzer>(MainSubmix);
+
+				FEditorDashboardFactory::OnActiveAudioDeviceChanged.AddSP(this, &FOutputMeterDashboardViewFactory::HandleOnActiveAudioDeviceChanged);
+				FSubmixesDashboardViewFactory::OnSubmixSelectionChanged.AddSP(this, &FOutputMeterDashboardViewFactory::HandleOnSubmixSelectionChanged);
+			}
+		}
+
+		return SNew(SOverlay)
+		+ SOverlay::Slot()
+		[
+			SNew(SColorBlock)
+			.Color(FSlateStyle::Get().GetColor("AudioInsights.Analyzers.BackgroundColor"))
+		]
+		+ SOverlay::Slot()
+		[
+			SNew(SVerticalBox)
+			+ SVerticalBox::Slot()
+			.FillHeight(1.0f)
+			.HAlign(HAlign_Center)
+			.VAlign(VAlign_Fill)
+			[
+				OutputMeterSubmixAnalyzer.IsValid() ? OutputMeterSubmixAnalyzer->GetWidget() : SNullWidget::NullWidget
+			]
+		];
+	}
+
+	void FOutputMeterDashboardViewFactory::HandleOnActiveAudioDeviceChanged()
+	{
+		if (OutputMeterSubmixAnalyzer.IsValid())
+		{
+			OutputMeterSubmixAnalyzer->SetSubmix(MainSubmix);
+		}
+	}
+
+	void FOutputMeterDashboardViewFactory::HandleOnSubmixSelectionChanged(const TWeakObjectPtr<USoundSubmix> InSoundSubmix)
+	{
+		if (OutputMeterSubmixAnalyzer.IsValid())
+		{
+			OutputMeterSubmixAnalyzer->SetSubmix(InSoundSubmix);
+		}
+	}
+} // namespace UE::Audio::Insights
+#undef LOCTEXT_NAMESPACE
diff --git a/Engine/Plugins/AudioInsights/Source/AudioInsightsEditor/Private/Views/OutputMeterDashboardViewFactory.h b/Engine/Plugins/AudioInsights/Source/AudioInsightsEditor/Private/Views/OutputMeterDashboardViewFactory.h
new file mode 100644
--- /dev/null
+++ b/Engine/Plugins/AudioInsights/Source/AudioInsightsEditor/Private/Views/OutputMeterDashboardViewFactory.h
@@ -0,0 +1,29 @@
+// Copyright Epic Games, Inc. All Rights Reserved.
+#pragma once
+
+#include "Templates/SharedPointer.h"
+#include "Views/DashboardViewFactory.h"
+
+class USoundSubmix;
+
+namespace UE::Audio::Insights
+{
+	class FAudioMeterSubmixAnalyzer;
+
+	class FOutputMeterDashboardViewFactory : public IDashboardViewFactory, public TSharedFromThis<FOutputMeterDashboardViewFactory>
+	{
+	public:
+		virtual FName GetName() const override;
+		virtual FText GetDisplayName() const override;
+		virtual EDefaultDashboardTabStack GetDefaultTabStack() const override;
+		virtual FSlateIcon GetIcon() const override;
+		virtual TSharedRef<SWidget> MakeWidget() override;
+
+	private:
+		void HandleOnActiveAudioDeviceChanged();
+		void HandleOnSubmixSelectionChanged(const TWeakObjectPtr<USoundSubmix> InSoundSubmix);
+
+		TSharedPtr<FAudioMeterSubmixAnalyzer> OutputMeterSubmixAnalyzer;
+		TObjectPtr<USoundSubmix> MainSubmix;
+	};
+} // namespace UE::Audio::Insights
diff --git a/Engine/Plugins/AudioInsights/Source/AudioInsightsEditor/Private/Views/OutputOscilloscopeDashboardViewFactory.cpp b/Engine/Plugins/AudioInsights/Source/AudioInsightsEditor/Private/Views/OutputOscilloscopeDashboardViewFactory.cpp
new file mode 100644
--- /dev/null
+++ b/Engine/Plugins/AudioInsights/Source/AudioInsightsEditor/Private/Views/OutputOscilloscopeDashboardViewFactory.cpp
@@ -0,0 +1,91 @@
+// Copyright Epic Games, Inc. All Rights Reserved.
+
+#include "Views/OutputOscilloscopeDashboardViewFactory.h"
+
+#include "Analyzers/AudioOscilloscopeAnalyzer.h"
+#include "AudioInsightsEditorDashboardFactory.h"
+#include "AudioInsightsStyle.h"
+#include "AudioOscilloscope.h"
+#include "Sound/AudioSettings.h"
+#include "Sound/SoundSubmix.h"
+#include "SubmixesDashboardViewFactory.h"
+#include "UObject/WeakObjectPtr.h"
+#include "Widgets/Colors/SColorBlock.h"
+
+#define LOCTEXT_NAMESPACE "AudioInsights"
+
+namespace UE::Audio::Insights
+{
+	FName FOutputOscilloscopeDashboardViewFactory::GetName() const
+	{
+		return "OutputOscilloscope";
+	}
+
+	FText FOutputOscilloscopeDashboardViewFactory::GetDisplayName() const
+	{
+		return LOCTEXT("AudioDashboard_DashboardsOscilloscopeTab_DisplayName", "Output Oscilloscope");
+	}
+
+	EDefaultDashboardTabStack FOutputOscilloscopeDashboardViewFactory::GetDefaultTabStack() const
+	{
+		return EDefaultDashboardTabStack::Oscilloscope;
+	}
+
+	FSlateIcon FOutputOscilloscopeDashboardViewFactory::GetIcon() const
+	{
+		return FSlateStyle::Get().CreateIcon("AudioInsights.Icon");
+	}
+
+	TSharedRef<SWidget> FOutputOscilloscopeDashboardViewFactory::MakeWidget()
+	{
+		// Create OutputOscilloscopeAnalyzer (use MainSubmix as default)
+		if (TWeakObjectPtr<UAudioSettings> AudioSettings = GetMutableDefault<UAudioSettings>();
+			AudioSettings.IsValid())
+		{
+			MainSubmix = Cast<USoundSubmix>(AudioSettings->MasterSubmix.ResolveObject());
+
+			if (MainSubmix)
+			{
+				OutputOscilloscopeAnalyzer = MakeShared<FAudioOscilloscopeAnalyzer>(MainSubmix);
+
+				FEditorDashboardFactory::OnActiveAudioDeviceChanged.AddSP(this, &FOutputOscilloscopeDashboardViewFactory::HandleOnActiveAudioDeviceChanged);
+				FSubmixesDashboardViewFactory::OnSubmixSelectionChanged.AddSP(this, &FOutputOscilloscopeDashboardViewFactory::HandleOnSubmixSelectionChanged);
+			}
+		}
+
+		return SNew(SOverlay)
+		+ SOverlay::Slot()
+		[
+			SNew(SColorBlock)
+			.Color(FSlateStyle::Get().GetColor("AudioInsights.Analyzers.BackgroundColor"))
+		]
+		+ SOverlay::Slot()
+		[
+			SNew(SVerticalBox)
+			+ SVerticalBox::Slot()
+			.FillHeight(1.0f)
+			.HAlign(HAlign_Center)
+			.VAlign(VAlign_Fill)
+			[
+				OutputOscilloscopeAnalyzer.IsValid() ? OutputOscilloscopeAnalyzer->GetAudioOscilloscope()->GetPanelWidget() : SNullWidget::NullWidget
+			]
+		];
+	}
+
+	void FOutputOscilloscopeDashboardViewFactory::HandleOnActiveAudioDeviceChanged()
+	{
+		if (OutputOscilloscopeAnalyzer.IsValid())
+		{
+			OutputOscilloscopeAnalyzer->RebuildAudioOscilloscope(MainSubmix);
+		}
+	}
+
+	void FOutputOscilloscopeDashboardViewFactory::HandleOnSubmixSelectionChanged(const TWeakObjectPtr<USoundSubmix> InSoundSubmix)
+	{
+		if (OutputOscilloscopeAnalyzer.IsValid())
+		{
+			OutputOscilloscopeAnalyzer->RebuildAudioOscilloscope(InSoundSubmix);
+		}
+	}
+} // namespace UE::Audio::Insights
+#undef LOCTEXT_NAMESPACE
diff --git a/Engine/Plugins/AudioInsights/Source/AudioInsightsEditor/Private/Views/OutputOscilloscopeDashboardViewFactory.h b/Engine/Plugins/AudioInsights/Source/AudioInsightsEditor/Private/Views/OutputOscilloscopeDashboardViewFactory.h
new file mode 100644
--- /dev/null
+++ b/Engine/Plugins/AudioInsights/Source/AudioInsightsEditor/Private/Views/OutputOscilloscopeDashboardViewFactory.h
@@ -0,0 +1,29 @@
+// Copyright Epic Games, Inc. All Rights Reserved.
+#pragma once
+
+#include "Templates/SharedPointer.h"
+#include "Views/DashboardViewFactory.h"
+
+class USoundSubmix;
+
+namespace UE::Audio::Insights
+{
+	class FAudioOscilloscopeAnalyzer;
+
+	class FOutputOscilloscopeDashboardViewFactory : public IDashboardViewFactory, public TSharedFromThis<FOutputOscilloscopeDashboardViewFactory>
+	{
+	public:
+		virtual FName GetName() const override;
+		virtual FText GetDisplayName() const override;
+		virtual EDefaultDashboardTabStack GetDefaultTabStack() const override;
+		virtual FSlateIcon GetIcon() const override;
+		virtual TSharedRef<SWidget> MakeWidget() override;
+
+	private:
+		void HandleOnActiveAudioDeviceChanged();
+		void HandleOnSubmixSelectionChanged(const TWeakObjectPtr<USoundSubmix> InSoundSubmix);
+
+		TSharedPtr<FAudioOscilloscopeAnalyzer> OutputOscilloscopeAnalyzer;
+		TObjectPtr<USoundSubmix> MainSubmix;
+	};
+} // namespace UE::Audio::Insights
diff --git a/Engine/Plugins/AudioInsights/Source/AudioInsightsEditor/Private/Views/SoundAttenuationVisualizer.cpp b/Engine/Plugins/AudioInsights/Source/AudioInsightsEditor/Private/Views/SoundAttenuationVisualizer.cpp
new file mode 100644
--- /dev/null
+++ b/Engine/Plugins/AudioInsights/Source/AudioInsightsEditor/Private/Views/SoundAttenuationVisualizer.cpp
@@ -0,0 +1,141 @@
+// Copyright Epic Games, Inc. All Rights Reserved.
+#include "SoundAttenuationVisualizer.h"
+
+#include "DrawDebugHelpers.h"
+#include "Engine/World.h"
+#include "Sound/SoundCue.h"
+#include "Sound/SoundNodeAttenuation.h"
+
+namespace UE::Audio::Insights
+{
+	FSoundAttenuationVisualizer::FSoundAttenuationVisualizer(const FColor& InColor)
+		: Color(InColor)
+	{
+		
+	}
+	
+	void FSoundAttenuationVisualizer::Draw(float InDeltaTime, const FTransform& InTransform, const UObject& InObject, const UWorld& InWorld) const
+	{
+		if (LastObjectId != InObject.GetUniqueID())
+		{
+			ShapeDetailsMap.Reset();
+			
+			if (const USoundCue* SoundCue = Cast<const USoundCue>(&InObject))
+			{
+				TArray<const USoundNodeAttenuation*> AttenuationNodes;
+				SoundCue->RecursiveFindAttenuation(SoundCue->FirstNode, AttenuationNodes);
+				
+				for (const USoundNodeAttenuation* Node : AttenuationNodes)
+				{
+					if (Node)
+					{
+						if (const FSoundAttenuationSettings* AttenuationSettingsToApply = Node->GetAttenuationSettingsToApply())
+						{
+							AttenuationSettingsToApply->CollectAttenuationShapesForVisualization(ShapeDetailsMap);
+						}
+					}
+				}
+			}
+			else if (const USoundBase* SoundBase = Cast<USoundBase>(&InObject))
+			{
+				if (const FSoundAttenuationSettings* Settings = SoundBase->GetAttenuationSettingsToApply())
+				{
+					Settings->CollectAttenuationShapesForVisualization(ShapeDetailsMap);
+				}
+			}
+			else
+			{
+				return;
+			}
+		}
+		
+		LastObjectId = InObject.GetUniqueID();
+
+		for (const auto& [AttenuationShape, ShapeDetails] : ShapeDetailsMap)
+		{
+			switch (AttenuationShape)
+			{
+				case EAttenuationShape::Sphere:
+				{
+					if (ShapeDetails.Falloff > 0.f)
+					{
+						DrawDebugSphere(&InWorld, InTransform.GetTranslation(), ShapeDetails.Extents.X + ShapeDetails.Falloff, 10, Color);
+						DrawDebugSphere(&InWorld, InTransform.GetTranslation(), ShapeDetails.Extents.X, 10, Color);
+					}
+					else
+					{
+						DrawDebugSphere(&InWorld, InTransform.GetTranslation(), ShapeDetails.Extents.X, 10, Color);
+					}
+					break;
+				}
+
+				case EAttenuationShape::Box:
+				{
+					if (ShapeDetails.Falloff > 0.f)
+					{
+						DrawDebugBox(&InWorld, InTransform.GetTranslation(), ShapeDetails.Extents + FVector(ShapeDetails.Falloff), InTransform.GetRotation(), Color);
+						DrawDebugBox(&InWorld, InTransform.GetTranslation(), ShapeDetails.Extents, InTransform.GetRotation(), Color);
+					}
+					else
+					{
+						DrawDebugBox(&InWorld, InTransform.GetTranslation(), ShapeDetails.Extents, InTransform.GetRotation(), Color);
+					}
+					break;
+				}
+
+				case EAttenuationShape::Capsule:
+				{
+					if (ShapeDetails.Falloff > 0.f)
+					{
+						DrawDebugCapsule(&InWorld, InTransform.GetTranslation(), ShapeDetails.Extents.X + ShapeDetails.Falloff, ShapeDetails.Extents.Y + ShapeDetails.Falloff, InTransform.GetRotation(), Color);
+						DrawDebugCapsule(&InWorld, InTransform.GetTranslation(), ShapeDetails.Extents.X, ShapeDetails.Extents.Y, InTransform.GetRotation(), Color);
+					}
+					else
+					{
+						DrawDebugCapsule(&InWorld, InTransform.GetTranslation(), ShapeDetails.Extents.X, ShapeDetails.Extents.Y, InTransform.GetRotation(), Color);
+					}
+					break;
+				}
+
+				case EAttenuationShape::Cone:
+				{
+					const FVector Origin = InTransform.GetTranslation() - (InTransform.GetUnitAxis(EAxis::X) * ShapeDetails.ConeOffset);
+
+					if (ShapeDetails.Falloff > 0.f || ShapeDetails.Extents.Z > 0.f)
+					{
+						const float OuterAngle = FMath::DegreesToRadians(ShapeDetails.Extents.Y + ShapeDetails.Extents.Z);
+						const float InnerAngle = FMath::DegreesToRadians(ShapeDetails.Extents.Y);
+						DrawDebugCone(&InWorld, Origin, InTransform.GetUnitAxis(EAxis::X), ShapeDetails.Extents.X + ShapeDetails.Falloff + ShapeDetails.ConeOffset, OuterAngle, OuterAngle, 10, Color);
+						DrawDebugCone(&InWorld, Origin, InTransform.GetUnitAxis(EAxis::X), ShapeDetails.Extents.X + ShapeDetails.ConeOffset, InnerAngle, InnerAngle, 10, Color);
+					}
+					else
+					{
+						const float Angle = FMath::DegreesToRadians(ShapeDetails.Extents.Y);
+						DrawDebugCone(&InWorld, Origin, InTransform.GetUnitAxis(EAxis::X), ShapeDetails.Extents.X + ShapeDetails.ConeOffset, Angle, Angle, 10, Color);
+					}
+
+					if (!FMath::IsNearlyZero(ShapeDetails.ConeSphereRadius, UE_KINDA_SMALL_NUMBER))
+					{
+						if (ShapeDetails.ConeSphereFalloff > 0.f)
+						{
+
+							DrawDebugSphere(&InWorld, Origin, ShapeDetails.ConeSphereRadius + ShapeDetails.ConeSphereFalloff, 10, Color);
+							DrawDebugSphere(&InWorld, Origin, ShapeDetails.ConeSphereRadius, 10, Color);
+						}
+						else
+						{
+							DrawDebugSphere(&InWorld, Origin, ShapeDetails.ConeSphereRadius, 10, Color);
+						}
+					}
+
+					break;
+				}
+
+				default:
+				{
+					break;
+				}
+			}
+		}
+	}
+} // namespace UE::Audio::Insights
diff --git a/Engine/Plugins/AudioInsights/Source/AudioInsightsEditor/Private/Views/SoundAttenuationVisualizer.h b/Engine/Plugins/AudioInsights/Source/AudioInsightsEditor/Private/Views/SoundAttenuationVisualizer.h
new file mode 100644
--- /dev/null
+++ b/Engine/Plugins/AudioInsights/Source/AudioInsightsEditor/Private/Views/SoundAttenuationVisualizer.h
@@ -0,0 +1,25 @@
+// Copyright Epic Games, Inc. All Rights Reserved.
+#pragma once 
+
+#include "Engine/Attenuation.h"
+#include "Math/Color.h"
+
+namespace UE::Audio::Insights
+{
+	class IDashboardDataViewEntry;
+
+	class FSoundAttenuationVisualizer
+	{
+	public:
+		explicit FSoundAttenuationVisualizer(const FColor& InColor);
+		void Draw(float InDeltaTime, const FTransform& InTransform, const UObject& InObject, const UWorld& InWorld) const;
+		
+		const FColor& GetColor() const { return Color; }
+
+	private:
+		const FColor Color { 155, 155, 255 };
+		
+		mutable TMultiMap<EAttenuationShape::Type, FBaseAttenuationSettings::AttenuationShapeDetails> ShapeDetailsMap;
+		mutable uint32 LastObjectId = INDEX_NONE;
+	};
+} // namespace UE::Audio::Insights
diff --git a/Engine/Plugins/AudioInsights/Source/AudioInsightsEditor/Private/Views/SoundSubmixAssetDashboardEntry.cpp b/Engine/Plugins/AudioInsights/Source/AudioInsightsEditor/Private/Views/SoundSubmixAssetDashboardEntry.cpp
new file mode 100644
--- /dev/null
+++ b/Engine/Plugins/AudioInsights/Source/AudioInsightsEditor/Private/Views/SoundSubmixAssetDashboardEntry.cpp
@@ -0,0 +1,26 @@
+// Copyright Epic Games, Inc. All Rights Reserved.
+
+#include "Views/SoundSubmixAssetDashboardEntry.h"
+
+namespace UE::Audio::Insights
+{
+	FText FSoundSubmixAssetDashboardEntry::GetDisplayName() const
+	{
+		return FText::FromString(FSoftObjectPath(Name).GetAssetName());
+	}
+
+	const UObject* FSoundSubmixAssetDashboardEntry::GetObject() const
+	{
+		return FSoftObjectPath(Name).ResolveObject();
+	}
+
+	UObject* FSoundSubmixAssetDashboardEntry::GetObject()
+	{
+		return FSoftObjectPath(Name).ResolveObject();
+	}
+
+	bool FSoundSubmixAssetDashboardEntry::IsValid() const
+	{
+		return true;
+	}
+} // namespace UE::Audio::Insights
diff --git a/Engine/Plugins/AudioInsights/Source/AudioInsightsEditor/Private/Views/SoundSubmixAssetDashboardEntry.h b/Engine/Plugins/AudioInsights/Source/AudioInsightsEditor/Private/Views/SoundSubmixAssetDashboardEntry.h
new file mode 100644
--- /dev/null
+++ b/Engine/Plugins/AudioInsights/Source/AudioInsightsEditor/Private/Views/SoundSubmixAssetDashboardEntry.h
@@ -0,0 +1,24 @@
+// Copyright Epic Games, Inc. All Rights Reserved.
+#pragma once
+
+#include "Views/TableDashboardViewFactory.h"
+
+class USoundSubmix;
+
+namespace UE::Audio::Insights
+{
+	struct FSoundSubmixAssetDashboardEntry : public IObjectDashboardEntry
+	{
+		virtual ~FSoundSubmixAssetDashboardEntry() = default;
+
+		virtual FText GetDisplayName() const override;
+		virtual const UObject* GetObject() const override;
+		virtual UObject* GetObject() override;
+		virtual bool IsValid() const override;
+
+		::Audio::FDeviceId DeviceId = INDEX_NONE;
+		uint64 ComponentId = TNumericLimits<uint64>::Max();
+		FString Name;
+		TWeakObjectPtr<USoundSubmix> SoundSubmix;
+	};
+} // namespace UE::Audio::Insights
diff --git a/Engine/Plugins/AudioInsights/Source/AudioInsightsEditor/Private/Views/SubmixesDashboardViewFactory.cpp b/Engine/Plugins/AudioInsights/Source/AudioInsightsEditor/Private/Views/SubmixesDashboardViewFactory.cpp
new file mode 100644
--- /dev/null
+++ b/Engine/Plugins/AudioInsights/Source/AudioInsightsEditor/Private/Views/SubmixesDashboardViewFactory.cpp
@@ -0,0 +1,227 @@
+// Copyright Epic Games, Inc. All Rights Reserved.
+#include "Views/SubmixesDashboardViewFactory.h"
+
+#include "AudioInsightsStyle.h"
+#include "Editor.h"
+#include "Internationalization/Text.h"
+#include "Providers/SoundSubmixProvider.h"
+#include "Sound/SoundSubmix.h"
+#include "Subsystems/AssetEditorSubsystem.h"
+#include "Templates/SharedPointer.h"
+
+#define LOCTEXT_NAMESPACE "AudioInsights"
+
+namespace UE::Audio::Insights
+{
+	namespace SubmixesPrivate
+	{
+		const FSoundSubmixAssetDashboardEntry& CastEntry(const IDashboardDataViewEntry& InData)
+		{
+			return static_cast<const FSoundSubmixAssetDashboardEntry&>(InData);
+		};
+	}
+
+	FSubmixesDashboardViewFactory::FSubmixesDashboardViewFactory()
+	{
+		FSoundSubmixProvider::OnSubmixAssetAdded.AddRaw(this, &FSubmixesDashboardViewFactory::HandleOnSubmixAssetListUpdated);
+		FSoundSubmixProvider::OnSubmixAssetRemoved.AddRaw(this, &FSubmixesDashboardViewFactory::HandleOnSubmixAssetListUpdated);
+		FSoundSubmixProvider::OnSubmixAssetListUpdated.AddRaw(this, &FSubmixesDashboardViewFactory::RequestListRefresh);
+
+		Providers = TArray<TSharedPtr<FTraceProviderBase>>
+		{
+			MakeShared<FSoundSubmixProvider>()
+		};
+
+		SortByColumn = "Name";
+		SortMode     = EColumnSortMode::Ascending;
+	}
+
+	FSubmixesDashboardViewFactory::~FSubmixesDashboardViewFactory()
+	{
+		FSoundSubmixProvider::OnSubmixAssetAdded.RemoveAll(this);
+		FSoundSubmixProvider::OnSubmixAssetRemoved.RemoveAll(this);
+		FSoundSubmixProvider::OnSubmixAssetListUpdated.RemoveAll(this);
+	}
+
+	FName FSubmixesDashboardViewFactory::GetName() const
+	{
+		return "Submixes";
+	}
+
+	FText FSubmixesDashboardViewFactory::GetDisplayName() const
+	{
+		return LOCTEXT("AudioDashboard_Submixes_DisplayName", "Submixes");
+	}
+
+	TSharedRef<SWidget> FSubmixesDashboardViewFactory::GenerateWidgetForColumn(TSharedRef<IDashboardDataViewEntry> InRowData, const FName& InColumnName)
+	{
+		if (InColumnName == "Name")
+		{
+			const TWeakObjectPtr<USoundSubmix> SoundSubmix = SubmixesPrivate::CastEntry(InRowData.Get()).SoundSubmix;
+
+			SubmixCheckboxCheckedStates.FindOrAdd(SoundSubmix, false);
+
+			const FColumnData& ColumnData = GetColumns()[InColumnName];
+			const FText Label = ColumnData.GetDisplayValue(InRowData.Get());
+
+			return (Label.IsEmpty()) ? 
+				SNullWidget::NullWidget :
+				SNew(SHorizontalBox)
+				.Clipping(EWidgetClipping::ClipToBounds)
+				+ SHorizontalBox::Slot()
+				.FillWidth(0.1f)
+				[
+					SNew(SCheckBox)
+					.IsChecked(SubmixCheckboxCheckedStates.Contains(SoundSubmix) && SubmixCheckboxCheckedStates[SoundSubmix] == true)
+					.OnCheckStateChanged_Lambda([SoundSubmix, this](ECheckBoxState NewState)
+					{
+						if (SubmixCheckboxCheckedStates.Contains(SoundSubmix))
+						{
+							SubmixCheckboxCheckedStates[SoundSubmix] = NewState == ECheckBoxState::Checked;
+						}
+
+						OnSubmixAssetChecked.Broadcast(NewState == ECheckBoxState::Checked, SoundSubmix);
+					})
+				]
+				+ SHorizontalBox::Slot()
+				.FillWidth(0.8f)
+				[
+					SNew(STextBlock)
+					.Text(Label)
+					.MinDesiredWidth(300)
+					.OnDoubleClicked_Lambda([this, InRowData](const FGeometry& MyGeometry, const FPointerEvent& PointerEvent)
+					{
+						if (GEditor)
+						{
+							if (TSharedPtr<IObjectDashboardEntry> ObjectData = StaticCastSharedPtr<IObjectDashboardEntry>(InRowData.ToSharedPtr()); 
+								ObjectData.IsValid())
+							{
+								if (const UObject* Object = ObjectData->GetObject();
+									Object && Object->IsAsset())
+								{
+									GEditor->GetEditorSubsystem<UAssetEditorSubsystem>()->OpenEditorForAsset(Object);
+									return FReply::Handled();
+								}
+							}
+						}
+
+						return FReply::Unhandled();
+					})
+				];
+		}
+
+		return SNullWidget::NullWidget;
+	}
+
+	void FSubmixesDashboardViewFactory::ProcessEntries(FTraceTableDashboardViewFactory::EProcessReason Reason)
+	{
+		const FString FilterString = GetSearchFilterText().ToString();
+		
+		FTraceTableDashboardViewFactory::FilterEntries<FSoundSubmixProvider>([&FilterString](const IDashboardDataViewEntry& Entry)
+		{
+			const FSoundSubmixAssetDashboardEntry& SubmixEntry = static_cast<const FSoundSubmixAssetDashboardEntry&>(Entry);
+			
+			return !SubmixEntry.GetDisplayName().ToString().Contains(FilterString);
+		});
+	}
+
+	FSlateIcon FSubmixesDashboardViewFactory::GetIcon() const
+	{
+		return FSlateStyle::Get().CreateIcon("AudioInsights.Icon.Submix");
+	}
+
+	EDefaultDashboardTabStack FSubmixesDashboardViewFactory::GetDefaultTabStack() const
+	{
+		return EDefaultDashboardTabStack::Analysis;
+	}
+
+	TSharedRef<SWidget> FSubmixesDashboardViewFactory::MakeWidget()
+	{
+		if (!DashboardWidget.IsValid())
+		{
+			DashboardWidget = FTraceTableDashboardViewFactory::MakeWidget();
+
+			if (FilteredEntriesListView.IsValid())
+			{
+				FilteredEntriesListView->SetSelectionMode(ESelectionMode::Single);
+			}
+		}
+
+		return DashboardWidget->AsShared();
+	}
+
+	const TMap<FName, FTraceTableDashboardViewFactory::FColumnData>& FSubmixesDashboardViewFactory::GetColumns() const
+	{
+		auto CreateColumnData = []()
+		{
+			return TMap<FName, FTraceTableDashboardViewFactory::FColumnData>
+			{
+				{
+					"Name",
+					{
+						LOCTEXT("Submixes_NameColumnDisplayName", "Name"),
+						[](const IDashboardDataViewEntry& InData) { return SubmixesPrivate::CastEntry(InData).GetDisplayName(); },
+						false,		/* bDefaultHidden */
+						0.88f,		/* FillWidth */
+						EHorizontalAlignment::HAlign_Left
+					}
+				}
+			};
+		};
+		
+		static const TMap<FName, FTraceTableDashboardViewFactory::FColumnData> ColumnData = CreateColumnData();
+		
+		return ColumnData;
+	}
+
+	void FSubmixesDashboardViewFactory::SortTable()
+	{
+		if (SortByColumn == "Name")
+		{
+			if (SortMode == EColumnSortMode::Ascending)
+			{
+				DataViewEntries.Sort([](const TSharedPtr<IDashboardDataViewEntry>& A, const TSharedPtr<IDashboardDataViewEntry>& B)
+				{
+					const FSoundSubmixAssetDashboardEntry& AData = SubmixesPrivate::CastEntry(*A.Get());
+					const FSoundSubmixAssetDashboardEntry& BData = SubmixesPrivate::CastEntry(*B.Get());
+
+					return AData.GetDisplayName().CompareToCaseIgnored(BData.GetDisplayName()) < 0;
+				});
+			}
+			else if (SortMode == EColumnSortMode::Descending)
+			{
+				DataViewEntries.Sort([](const TSharedPtr<IDashboardDataViewEntry>& A, const TSharedPtr<IDashboardDataViewEntry>& B)
+				{
+					const FSoundSubmixAssetDashboardEntry& AData = SubmixesPrivate::CastEntry(*A.Get());
+					const FSoundSubmixAssetDashboardEntry& BData = SubmixesPrivate::CastEntry(*B.Get());
+
+					return BData.GetDisplayName().CompareToCaseIgnored(AData.GetDisplayName()) < 0;
+				});
+			}
+		}
+	}
+
+	void FSubmixesDashboardViewFactory::OnSelectionChanged(TSharedPtr<IDashboardDataViewEntry> SelectedItem, ESelectInfo::Type SelectInfo)
+	{
+		if (SelectedItem.IsValid())
+		{
+			const FSoundSubmixAssetDashboardEntry& SoundSubmixAssetDashboardEntry = SubmixesPrivate::CastEntry(*SelectedItem.Get());
+
+			OnSubmixSelectionChanged.Broadcast(SoundSubmixAssetDashboardEntry.SoundSubmix);
+		}
+	}
+
+	void FSubmixesDashboardViewFactory::RequestListRefresh()
+	{
+		if (FilteredEntriesListView.IsValid())
+		{
+			FilteredEntriesListView->RequestListRefresh();
+		}
+	}
+
+	void FSubmixesDashboardViewFactory::HandleOnSubmixAssetListUpdated(const TWeakObjectPtr<UObject> InAsset)
+	{
+		RequestListRefresh();
+	}
+} // namespace UE::Audio::Insights
+#undef LOCTEXT_NAMESPACE
diff --git a/Engine/Plugins/AudioInsights/Source/AudioInsightsEditor/Private/Views/SubmixesDashboardViewFactory.h b/Engine/Plugins/AudioInsights/Source/AudioInsightsEditor/Private/Views/SubmixesDashboardViewFactory.h
new file mode 100644
--- /dev/null
+++ b/Engine/Plugins/AudioInsights/Source/AudioInsightsEditor/Private/Views/SubmixesDashboardViewFactory.h
@@ -0,0 +1,47 @@
+// Copyright Epic Games, Inc. All Rights Reserved.
+#pragma once
+
+#include "Delegates/Delegate.h"
+#include "Templates/SharedPointer.h"
+#include "Views/SoundSubmixAssetDashboardEntry.h"
+#include "Views/TableDashboardViewFactory.h"
+#include "Widgets/Input/SCheckBox.h"
+
+class USoundSubmix;
+
+namespace UE::Audio::Insights
+{
+	class FSubmixesDashboardViewFactory : public FTraceObjectTableDashboardViewFactory
+	{
+	public:
+		FSubmixesDashboardViewFactory();
+		virtual ~FSubmixesDashboardViewFactory();
+
+		virtual FName GetName() const override;
+		virtual FText GetDisplayName() const override;
+		virtual FSlateIcon GetIcon() const override;
+		virtual EDefaultDashboardTabStack GetDefaultTabStack() const override;
+		virtual TSharedRef<SWidget> MakeWidget() override;
+
+		DECLARE_MULTICAST_DELEGATE_TwoParams(FOnSubmixAssetChecked, const bool /*bIsChecked*/, TWeakObjectPtr<USoundSubmix> /*SoundSubmix*/);
+		inline static FOnSubmixAssetChecked OnSubmixAssetChecked;
+
+		DECLARE_MULTICAST_DELEGATE_OneParam(FOnSubmixSelectionChanged, const TWeakObjectPtr<USoundSubmix> /*SoundSubmix*/);
+		inline static FOnSubmixSelectionChanged OnSubmixSelectionChanged;
+
+	protected:
+		virtual TSharedRef<SWidget> GenerateWidgetForColumn(TSharedRef<IDashboardDataViewEntry> InRowData, const FName& InColumnName) override;
+		virtual void ProcessEntries(FTraceTableDashboardViewFactory::EProcessReason Reason) override;
+		virtual const TMap<FName, FTraceTableDashboardViewFactory::FColumnData>& GetColumns() const override;
+
+		virtual void SortTable() override;
+
+		virtual void OnSelectionChanged(TSharedPtr<IDashboardDataViewEntry> SelectedItem, ESelectInfo::Type SelectInfo) override;
+
+		void RequestListRefresh();
+
+		void HandleOnSubmixAssetListUpdated(const TWeakObjectPtr<UObject> InAsset);
+
+		TMap<const TWeakObjectPtr<USoundSubmix>, bool> SubmixCheckboxCheckedStates;
+	};
+} // namespace UE::Audio::Insights
diff --git a/Engine/Plugins/AudioInsights/Source/AudioInsightsEditor/Private/Views/ViewportDashboardViewFactory.h b/Engine/Plugins/AudioInsights/Source/AudioInsightsEditor/Private/Views/ViewportDashboardViewFactory.h
new file mode 100644
--- /dev/null
+++ b/Engine/Plugins/AudioInsights/Source/AudioInsightsEditor/Private/Views/ViewportDashboardViewFactory.h
@@ -0,0 +1,50 @@
+// Copyright Epic Games, Inc. All Rights Reserved.
+#pragma once
+
+#include "AudioInsightsStyle.h"
+#include "Views/DashboardViewFactory.h"
+#include "Widgets/SNullWidget.h"
+// #include "LevelEditor.h"
+// #include "SLevelViewport.h"
+
+#define LOCTEXT_NAMESPACE "AudioInsights"
+
+
+namespace UE::Audio::Insights
+{
+	class FViewportDashboardViewFactory : public IDashboardViewFactory
+	{
+	public:
+		virtual FName GetName() const override
+		{
+			return "Viewport";
+		}
+
+		virtual FText GetDisplayName() const override
+		{
+			return LOCTEXT("AudioDashboard_ViewportTab_DisplayName", "Viewport");
+		}
+
+		virtual EDefaultDashboardTabStack GetDefaultTabStack() const override
+		{
+			return EDefaultDashboardTabStack::Viewport;
+		}
+
+		virtual FSlateIcon GetIcon() const override
+		{
+			return FSlateStyle::Get().CreateIcon("AudioInsights.Icon.Viewport");
+		}
+
+		virtual TSharedRef<SWidget> MakeWidget() override
+		{
+			// TODO: This test currently hijacks the main level viewport which is not ideal.
+			// Find way to make our own instance and bind to tab here.
+			// FLevelEditorModule& LevelEditor = FModuleManager::GetModuleChecked<FLevelEditorModule>("LevelEditor");
+			// TSharedPtr<SLevelViewport> LevelViewport = LevelEditor.GetFirstActiveLevelViewport();
+			// ViewportContent = LevelViewport->AsShared();
+
+			return SNullWidget::NullWidget;
+		}
+	};
+} // namespace UE::Audio::Insights
+#undef LOCTEXT_NAMESPACE
diff --git a/Engine/Plugins/AudioInsights/Source/AudioInsightsEditor/Private/Views/VirtualLoopsDebugDraw.cpp b/Engine/Plugins/AudioInsights/Source/AudioInsightsEditor/Private/Views/VirtualLoopsDebugDraw.cpp
new file mode 100644
--- /dev/null
+++ b/Engine/Plugins/AudioInsights/Source/AudioInsightsEditor/Private/Views/VirtualLoopsDebugDraw.cpp
@@ -0,0 +1,78 @@
+// Copyright Epic Games, Inc. All Rights Reserved.
+#include "VirtualLoopsDebugDraw.h"
+
+#include "AudioInsightsDataSource.h"
+#include "AudioDeviceManager.h"
+#include "DrawDebugHelpers.h"
+#include "Engine/World.h"
+#include "Messages/VirtualLoopTraceMessages.h"
+#include "Views/VirtualLoopDashboardViewFactory.h"
+
+namespace UE::Audio::Insights
+{
+	namespace VirtualLoopsDebugDrawPrivate
+	{
+		const FVirtualLoopDashboardEntry& CastEntry(const IDashboardDataViewEntry& InData)
+		{
+			return static_cast<const FVirtualLoopDashboardEntry&>(InData);
+		};
+	}
+
+	FVirtualLoopsDebugDraw::FVirtualLoopsDebugDraw()
+		: AttenuationVisualizer(FColor::Blue)
+	{
+		FVirtualLoopDashboardViewFactory::OnDebugDrawEntries.AddRaw(this, &FVirtualLoopsDebugDraw::DebugDraw);
+	}
+
+	FVirtualLoopsDebugDraw::~FVirtualLoopsDebugDraw()
+	{
+		FVirtualLoopDashboardViewFactory::OnDebugDrawEntries.RemoveAll(this);
+	}
+
+	void FVirtualLoopsDebugDraw::DebugDrawEntries(float InElapsed, const TArray<TSharedPtr<IDashboardDataViewEntry>>& InSelectedItems, ::Audio::FDeviceId InAudioDeviceId) const
+	{
+		for (const TSharedPtr<IDashboardDataViewEntry>& SelectedEntry : InSelectedItems)
+		{
+			if (!SelectedEntry.IsValid())
+			{
+				continue;
+			}
+			
+			const FVirtualLoopDashboardEntry& LoopData = VirtualLoopsDebugDrawPrivate::CastEntry(*SelectedEntry);
+			
+			const FRotator& Rotator   = LoopData.Rotator;
+			const FVector& Location   = LoopData.Location;
+			const FString Description = FString::Printf(TEXT("%s [Virt: %.2fs]"), *LoopData.Name, LoopData.TimeVirtualized);
+
+			const TArray<UWorld*> Worlds = FAudioDeviceManager::Get()->GetWorldsUsingAudioDevice(InAudioDeviceId);
+
+			for (UWorld* World : Worlds)
+			{
+				DrawDebugSphere(World, Location, 30.0f, 8, AttenuationVisualizer.GetColor(), false, InElapsed, SDPG_Foreground);
+				DrawDebugString(World, Location + FVector(0, 0, 32), *Description, nullptr, AttenuationVisualizer.GetColor(), InElapsed, false, 1.0f);
+
+				if (const TObjectPtr<const UObject> Object = LoopData.GetObject())
+				{
+					FTransform Transform;
+					Transform.SetLocation(Location);
+					Transform.SetRotation(FQuat(Rotator));
+
+					AttenuationVisualizer.Draw(InElapsed, Transform, *Object, *World);
+				}
+			}
+		}
+	}
+
+	void FVirtualLoopsDebugDraw::DebugDraw(float InElapsed, const TArray<TSharedPtr<IDashboardDataViewEntry>>& InSelectedItems, ::Audio::FDeviceId InAudioDeviceId)
+	{
+		FAudioDeviceManager* AudioDeviceManager = FAudioDeviceManager::Get();
+		if (AudioDeviceManager)
+		{
+			AudioDeviceManager->IterateOverAllDevices([this, &InSelectedItems, InElapsed](::Audio::FDeviceId DeviceId, FAudioDevice* Device)
+			{
+				DebugDrawEntries(InElapsed, InSelectedItems, DeviceId);
+			});
+		}
+	}
+
+} // namespace UE::Audio::Insights
diff --git a/Engine/Plugins/AudioInsights/Source/AudioInsightsEditor/Private/Views/VirtualLoopsDebugDraw.h b/Engine/Plugins/AudioInsights/Source/AudioInsightsEditor/Private/Views/VirtualLoopsDebugDraw.h
new file mode 100644
--- /dev/null
+++ b/Engine/Plugins/AudioInsights/Source/AudioInsightsEditor/Private/Views/VirtualLoopsDebugDraw.h
@@ -0,0 +1,24 @@
+// Copyright Epic Games, Inc. All Rights Reserved.
+#pragma once 
+
+#include "AudioDefines.h"
+#include "SoundAttenuationVisualizer.h"
+
+namespace UE::Audio::Insights
+{
+	class IDashboardDataViewEntry;
+
+	class FVirtualLoopsDebugDraw
+	{
+	public:
+		FVirtualLoopsDebugDraw();
+		~FVirtualLoopsDebugDraw();
+
+		void DebugDraw(float InElapsed, const TArray<TSharedPtr<IDashboardDataViewEntry>>& InSelectedItems, ::Audio::FDeviceId InAudioDeviceId);
+
+	private:
+		void DebugDrawEntries(float InElapsed, const TArray<TSharedPtr<IDashboardDataViewEntry>>& InSelectedItems, ::Audio::FDeviceId InAudioDeviceId) const;
+
+		FSoundAttenuationVisualizer AttenuationVisualizer;
+	};
+} // namespace UE::Audio::Insights
diff --git a/Engine/Plugins/AudioInsights/Source/AudioInsightsEditor/Public/AudioInsightsEditorLog.h b/Engine/Plugins/AudioInsights/Source/AudioInsightsEditor/Public/AudioInsightsEditorLog.h
new file mode 100644
--- /dev/null
+++ b/Engine/Plugins/AudioInsights/Source/AudioInsightsEditor/Public/AudioInsightsEditorLog.h
@@ -0,0 +1,7 @@
+// Copyright Epic Games, Inc. All Rights Reserved.
+#pragma once
+
+#include "Logging/LogMacros.h"
+#include "Logging/LogVerbosity.h"
+
+AUDIOINSIGHTSEDITOR_API DECLARE_LOG_CATEGORY_EXTERN(LogAudioInsightsEditor, Log, All);
