[Audio Insights] FTraceModule: Changed IsRunningCookCommandlet to IsRunningCommandlet to make sure AudioInsights is not running while other commandlets are running.

diff --git a/Engine/Plugins/AudioInsights/Source/AudioInsights/Private/AudioInsightsTraceModule.cpp b/Engine/Plugins/AudioInsights/Source/AudioInsights/Private/AudioInsightsTraceModule.cpp
--- a/Engine/Plugins/AudioInsights/Source/AudioInsights/Private/AudioInsightsTraceModule.cpp
+++ b/Engine/Plugins/AudioInsights/Source/AudioInsights/Private/AudioInsightsTraceModule.cpp
@@ -15,8 +15,8 @@ namespace UE::Audio::Insights
 	FTraceModule::FTraceModule()
 		: ChannelManager(MakeShared<FTraceChannelManager>())
 	{
-		// Don't run providers in cook commandlet to avoid additional, unnecessary overhead as audio insights is dormant.
-		if (!IsRunningCookCommandlet())
+		// Don't run providers in any commandlet to avoid additional, unnecessary overhead as audio insights is dormant.
+		if (!IsRunningCommandlet())
 		{
 			TSharedPtr<FMixerSourceTraceProvider> SourceProvider = MakeShared<FMixerSourceTraceProvider>(ChannelManager);
 			TSharedPtr<FVirtualLoopTraceProvider> VirtualLoopProvider = MakeShared<FVirtualLoopTraceProvider>(ChannelManager);
diff --git a/Engine/Plugins/AudioInsights/Source/AudioInsightsEditor/Private/AudioInsightsEditorModule.cpp b/Engine/Plugins/AudioInsights/Source/AudioInsightsEditor/Private/AudioInsightsEditorModule.cpp
--- a/Engine/Plugins/AudioInsights/Source/AudioInsightsEditor/Private/AudioInsightsEditorModule.cpp
+++ b/Engine/Plugins/AudioInsights/Source/AudioInsightsEditor/Private/AudioInsightsEditorModule.cpp
@@ -27,7 +27,7 @@ namespace UE::Audio::Insights
 {
 	void FAudioInsightsEditorModule::StartupModule()
 	{
-		// Don't run providers in commandlet to avoid additional, unnecessary overhead as audio insights is dormant.
+		// Don't run providers in any commandlet to avoid additional, unnecessary overhead as audio insights is dormant.
 		if (!IsRunningCommandlet())
 		{
 			RegisterMenus();
