[Audio Insights] - Fix for: Insights continues to display virtual sounds after they have stopped being active (unified duplicated Bump logic in one BumpEntryFunc + do BumpEntryFunc instead of Get when processing the FVirtualLoopUpdateMessage queue)

diff --git a/Engine/Plugins/AudioInsights/Source/AudioInsights/Private/Providers/VirtualLoopTraceProvider.cpp b/Engine/Plugins/AudioInsights/Source/AudioInsights/Private/Providers/VirtualLoopTraceProvider.cpp
--- a/Engine/Plugins/AudioInsights/Source/AudioInsights/Private/Providers/VirtualLoopTraceProvider.cpp
+++ b/Engine/Plugins/AudioInsights/Source/AudioInsights/Private/Providers/VirtualLoopTraceProvider.cpp
@@ -28,10 +28,7 @@ namespace UE::Audio::Insights
 			return FindDeviceEntry(Msg.DeviceId, Msg.PlayOrder);
 		};
 
-		ProcessMessageQueue<FVirtualLoopStopMessage>(TraceMessages.StopMessages, GetEntryFunc, RemoveEntryFunc);
-
-		ProcessMessageQueue<FVirtualLoopVirtualizeMessage>(TraceMessages.VirtualizeMessages,
-		[this](const FVirtualLoopMessageBase& Msg)
+		auto BumpEntryFunc = [this](const FVirtualLoopMessageBase& Msg)
 		{
 			TSharedPtr<FVirtualLoopDashboardEntry>* ToReturn = nullptr;
 			UpdateDeviceEntry(Msg.DeviceId, Msg.PlayOrder, [&ToReturn, &Msg](TSharedPtr<FVirtualLoopDashboardEntry>& Entry)
@@ -39,15 +36,19 @@ namespace UE::Audio::Insights
 				if (!Entry.IsValid())
 				{
 					Entry = MakeShared<FVirtualLoopDashboardEntry>();
+					Entry->DeviceId  = Msg.DeviceId;
+					Entry->PlayOrder = Msg.PlayOrder;
 				}
-				Entry->DeviceId = Msg.DeviceId;
-				Entry->PlayOrder = Msg.PlayOrder;
 				Entry->Timestamp = Msg.Timestamp;
 
 				ToReturn = &Entry;
 			});
+
 			return ToReturn;
-		},
+		};
+
+		ProcessMessageQueue<FVirtualLoopVirtualizeMessage>(TraceMessages.VirtualizeMessages,
+		BumpEntryFunc,
 		[](const FVirtualLoopVirtualizeMessage& Msg, TSharedPtr<FVirtualLoopDashboardEntry>* OutEntry)
 		{
 			FVirtualLoopDashboardEntry& EntryRef = *OutEntry->Get();
@@ -55,29 +56,10 @@ namespace UE::Audio::Insights
 			EntryRef.ComponentId = Msg.ComponentId;
 		});
 
-		ProcessMessageQueue<FVirtualLoopRealizeMessage>(TraceMessages.RealizeMessages, 
-		[this](const FVirtualLoopMessageBase& Msg)
-		{
-			TSharedPtr<FVirtualLoopDashboardEntry>* ToReturn = nullptr;
-
-			UpdateDeviceEntry(Msg.DeviceId, Msg.PlayOrder, [&ToReturn, &Msg](TSharedPtr<FVirtualLoopDashboardEntry>& Entry)
-			{
-				if (!Entry.IsValid())
-				{
-					Entry = MakeShared<FVirtualLoopDashboardEntry>();
-					Entry->DeviceId = Msg.DeviceId;
-					Entry->PlayOrder = Msg.PlayOrder;
-				}
-				Entry->Timestamp = Msg.Timestamp;
-				ToReturn = &Entry;
-			});
-
-			return ToReturn;
-		}, 
-		RemoveEntryFunc);
+		ProcessMessageQueue<FVirtualLoopStopMessage>(TraceMessages.StopMessages, GetEntryFunc, RemoveEntryFunc);
 
 		ProcessMessageQueue<FVirtualLoopUpdateMessage>(TraceMessages.UpdateMessages, 
-		GetEntryFunc,
+		BumpEntryFunc,
 		[](const FVirtualLoopUpdateMessage& Msg, TSharedPtr<FVirtualLoopDashboardEntry>* OutEntry)
 		{
 			if (OutEntry)
@@ -91,6 +73,8 @@ namespace UE::Audio::Insights
 			}
 		});
 
+		ProcessMessageQueue<FVirtualLoopRealizeMessage>(TraceMessages.RealizeMessages, BumpEntryFunc, RemoveEntryFunc);
+
 		return true;
 	}
 
