[Audio Insights] Added Audio Insights component + enable in Unreal Insights

diff --git a/Engine/Plugins/AudioInsights/Source/AudioInsights/AudioInsights.Build.cs b/Engine/Plugins/AudioInsights/Source/AudioInsights/AudioInsights.Build.cs
--- a/Engine/Plugins/AudioInsights/Source/AudioInsights/AudioInsights.Build.cs
+++ b/Engine/Plugins/AudioInsights/Source/AudioInsights/AudioInsights.Build.cs
@@ -7,7 +7,8 @@ public class AudioInsights : ModuleRules
 	public AudioInsights(ReadOnlyTargetRules Target) : base(Target)
 	{
 		PCHUsage = ModuleRules.PCHUsageMode.UseExplicitOrSharedPCHs;
-		
+		OptimizeCode = CodeOptimization.Never;
+
 		PublicDependencyModuleNames.AddRange
 		(
 			new string[]
@@ -27,7 +28,6 @@ public class AudioInsights : ModuleRules
 				"CoreUObject",
 				"InputCore",
 				"OutputLog",
-				"SessionServices",
 				"Slate",
 				"SlateCore",
 				"TraceAnalysis",
diff --git a/Engine/Plugins/AudioInsights/Source/AudioInsights/Private/AudioInsightsComponent.cpp b/Engine/Plugins/AudioInsights/Source/AudioInsights/Private/AudioInsightsComponent.cpp
new file mode 100644
--- /dev/null
+++ b/Engine/Plugins/AudioInsights/Source/AudioInsights/Private/AudioInsightsComponent.cpp
@@ -0,0 +1,149 @@
+// Copyright Epic Games, Inc. All Rights Reserved.
+#include "AudioInsightsComponent.h"
+
+#include "AudioInsightsModule.h"
+#include "AudioInsightsStyle.h"
+#include "Modules/ModuleManager.h"
+#include "Trace/StoreClient.h"
+#include "Widgets/Docking/SDockTab.h"
+
+#define LOCTEXT_NAMESPACE "AudioInsightsComponent"
+
+namespace UE::Audio::Insights
+{
+	namespace ComponentPrivate
+	{
+		static const FName TabName = "Audio Insights";
+	}
+
+	FAudioInsightsComponent::~FAudioInsightsComponent()
+	{
+		ensure(!bIsInitialized);
+	}
+
+	TSharedPtr<FAudioInsightsComponent> FAudioInsightsComponent::CreateInstance()
+	{
+		ensure(!Instance.IsValid());
+
+		if (Instance.IsValid())
+		{
+			Instance.Reset();
+		}
+
+		Instance = MakeShared<FAudioInsightsComponent>();
+
+		return Instance;
+	}
+
+	void FAudioInsightsComponent::Initialize(IUnrealInsightsModule& InsightsModule)
+	{
+		ensure(!bIsInitialized);
+	
+		if (!bIsInitialized)
+		{
+			bIsInitialized = true;
+
+			OnTick = FTickerDelegate::CreateSP(this, &FAudioInsightsComponent::Tick);
+
+			constexpr float TickDelay = 0.5f; // 500 ms. delay between ticks
+			OnTickHandle = FTSTicker::GetCoreTicker().AddTicker(OnTick, TickDelay);
+		}
+	}
+
+	void FAudioInsightsComponent::Shutdown()
+	{
+		if (!bIsInitialized)
+		{
+			return;
+		}
+	
+		bIsInitialized = false;
+
+		FTSTicker::GetCoreTicker().RemoveTicker(OnTickHandle);
+
+		Instance.Reset();
+	}
+
+	void FAudioInsightsComponent::RegisterMajorTabs(IUnrealInsightsModule& InsightsModule)
+	{
+		using namespace ComponentPrivate;
+
+		const FInsightsMajorTabConfig& Config = InsightsModule.FindMajorTabConfig(TabName);
+
+		if (Config.bIsAvailable)
+		{
+			// Register tab spawner for the Audio Insights.
+			FTabSpawnerEntry& TabSpawnerEntry = FGlobalTabmanager::Get()->RegisterNomadTabSpawner(TabName,
+				FOnSpawnTab::CreateRaw(this,  &FAudioInsightsComponent::SpawnTab), 
+				FCanSpawnTab::CreateRaw(this, &FAudioInsightsComponent::CanSpawnTab))
+				.SetDisplayName(Config.TabLabel.IsSet()   ? Config.TabLabel.GetValue()   : LOCTEXT("AudioInsights_TabTitle", "Audio Insights"))
+				.SetTooltipText(Config.TabTooltip.IsSet() ? Config.TabTooltip.GetValue() : LOCTEXT("AudioInsights_TooltipText", "Open the Audio Insights tab."))
+				.SetIcon(Config.TabIcon.IsSet() ? Config.TabIcon.GetValue() : FSlateStyle::Get().CreateIcon("AudioInsights.Icon.Submix"));
+
+			const TSharedRef<FWorkspaceItem>* FoundWorkspace = FGlobalTabmanager::Get()->GetLocalWorkspaceMenuRoot()->GetChildItems().FindByPredicate(
+				[](const TSharedRef<FWorkspaceItem>& WorkspaceItem)
+				{
+					return WorkspaceItem->GetDisplayName().ToString() == "Insights Tools";
+				});
+
+			if (FoundWorkspace)
+			{
+				TabSpawnerEntry.SetGroup(*FoundWorkspace);
+			}
+		}
+	}
+
+	void FAudioInsightsComponent::UnregisterMajorTabs()
+	{
+		FGlobalTabmanager::Get()->UnregisterNomadTabSpawner(ComponentPrivate::TabName);
+	}
+
+	bool FAudioInsightsComponent::CanSpawnTab(const FSpawnTabArgs& Args) const
+	{
+		return bCanSpawnTab;
+	}
+
+	TSharedRef<SDockTab> FAudioInsightsComponent::SpawnTab(const FSpawnTabArgs& Args)
+	{
+		const TSharedRef<SDockTab> DockTab = FAudioInsightsModule::GetChecked().CreateDashboardTabWidget(Args);
+
+		OnTabSpawn.Broadcast();
+
+		return DockTab;
+	}
+
+	bool FAudioInsightsComponent::Tick(float DeltaTime)
+	{
+		// Audio Insights will be available only if there is an active live session
+		if (bCanCheckForActiveSession && !bCanSpawnTab)
+		{
+			IUnrealInsightsModule& UnrealInsightsModule = FModuleManager::LoadModuleChecked<IUnrealInsightsModule>("TraceInsights");
+
+			TSharedPtr<const TraceServices::IAnalysisSession> Session = UnrealInsightsModule.GetAnalysisSession();
+			if (Session.IsValid())
+			{
+				TraceServices::FAnalysisSessionReadScope SessionReadScope(*Session.Get());
+
+				if (!Session->IsAnalysisComplete())
+				{
+					if (UE::Trace::FStoreClient* StoreClient = UnrealInsightsModule.GetStoreClient())
+					{
+						const UE::Trace::FStoreClient::FSessionInfo* SessionInfo = StoreClient->GetSessionInfoByTraceId(Session->GetTraceId());
+						if (SessionInfo)
+						{
+							bCanSpawnTab = true;
+						}
+					}
+				}
+				else
+				{
+					bCanCheckForActiveSession = false;
+				}
+			}
+		}
+
+		return true;
+	}
+} // namespace UE::Audio::Insights
+
+#undef LOCTEXT_NAMESPACE
diff --git a/Engine/Plugins/AudioInsights/Source/AudioInsights/Private/AudioInsightsComponent.h b/Engine/Plugins/AudioInsights/Source/AudioInsights/Private/AudioInsightsComponent.h
new file mode 100644
--- /dev/null
+++ b/Engine/Plugins/AudioInsights/Source/AudioInsights/Private/AudioInsightsComponent.h
@@ -0,0 +1,46 @@
+// Copyright Epic Games, Inc. All Rights Reserved.
+
+#pragma once
+
+#include "Containers/Ticker.h"
+#include "Insights/IUnrealInsightsModule.h"
+
+namespace UE::Audio::Insights
+{
+	/**
+	 * The component that makes Audio Insights available inside Unreal Insights.
+	 */
+	class FAudioInsightsComponent : public IInsightsComponent, public TSharedFromThis<FAudioInsightsComponent>
+	{
+	public:
+		FAudioInsightsComponent() = default;
+		virtual ~FAudioInsightsComponent();
+
+		static TSharedPtr<FAudioInsightsComponent> CreateInstance();
+
+		// IInsightsComponent overrides
+		virtual void Initialize(IUnrealInsightsModule& InsightsModule) override;
+		virtual void Shutdown() override;
+		virtual void RegisterMajorTabs(IUnrealInsightsModule& InsightsModule) override;
+		virtual void UnregisterMajorTabs() override;
+
+		DECLARE_MULTICAST_DELEGATE(FOnTabSpawn);
+		inline static FOnTabSpawn OnTabSpawn;
+
+	private:
+		bool CanSpawnTab(const FSpawnTabArgs& Args) const;
+		TSharedRef<SDockTab> SpawnTab(const FSpawnTabArgs& Args);
+
+		bool Tick(float DeltaTime);
+
+		bool bIsInitialized = false;
+
+		bool bCanCheckForActiveSession = true;
+		bool bCanSpawnTab = false;
+
+		FTickerDelegate OnTick;
+		FTSTicker::FDelegateHandle OnTickHandle;
+
+		inline static TSharedPtr<FAudioInsightsComponent> Instance;
+	};
+} // namespace UE::Audio::Insights
diff --git a/Engine/Plugins/AudioInsights/Source/AudioInsights/Private/AudioInsightsModule.cpp b/Engine/Plugins/AudioInsights/Source/AudioInsights/Private/AudioInsightsModule.cpp
--- a/Engine/Plugins/AudioInsights/Source/AudioInsights/Private/AudioInsightsModule.cpp
+++ b/Engine/Plugins/AudioInsights/Source/AudioInsights/Private/AudioInsightsModule.cpp
@@ -7,13 +7,13 @@
 #include "AudioInsightsTraceModule.h"
 #include "Features/IModularFeatures.h"
 #include "Framework/Docking/TabManager.h"
-#include "Insights/IUnrealInsightsModule.h"
 #include "Modules/ModuleManager.h"
 #include "Templates/SharedPointer.h"
 #include "TraceServices/ModuleService.h"
 #include "UObject/NameTypes.h"
 
 #if !WITH_EDITOR
+#include "AudioInsightsComponent.h"
 #include "Views/LogDashboardViewFactory.h"
 #include "Views/MixerSourceDashboardViewFactory.h"
 #include "Views/VirtualLoopDashboardViewFactory.h"
@@ -40,6 +40,11 @@ namespace UE::Audio::Insights
 			DashboardFactory->RegisterViewFactory(MakeShared<FLogDashboardViewFactory>());
 			DashboardFactory->RegisterViewFactory(MakeShared<FMixerSourceDashboardViewFactory>());
 			DashboardFactory->RegisterViewFactory(MakeShared<FVirtualLoopDashboardViewFactory>());
+			
+			AudioInsightsComponent = FAudioInsightsComponent::CreateInstance();
+
+			IUnrealInsightsModule& UnrealInsightsModule = FModuleManager::LoadModuleChecked<IUnrealInsightsModule>("TraceInsights");
+			UnrealInsightsModule.RegisterComponent(AudioInsightsComponent);
 #endif // !WITH_EDITOR
 
 			FCoreDelegates::OnFEngineLoopInitComplete.AddLambda([this]
@@ -60,10 +65,18 @@ namespace UE::Audio::Insights
 	{
 		if (!IsRunningCommandlet())
 		{
-			DashboardFactory.Reset();
-			IModularFeatures::Get().UnregisterModularFeature(TraceServices::ModuleFeatureName, &TraceModule);
+#if !WITH_EDITOR
+			IUnrealInsightsModule& UnrealInsightsModule = FModuleManager::LoadModuleChecked<IUnrealInsightsModule>("TraceInsights");
+			UnrealInsightsModule.UnregisterComponent(AudioInsightsComponent);
+
+			AudioInsightsComponent.Reset();
+#endif // !WITH_EDITOR
 
 			FDashboardAssetCommands::Unregister();
+
+			DashboardFactory.Reset();
+
+			IModularFeatures::Get().UnregisterModularFeature(TraceServices::ModuleFeatureName, &TraceModule);
 		}
 	}
 
diff --git a/Engine/Plugins/AudioInsights/Source/AudioInsights/Private/Views/MixerSourceDashboardViewFactory.cpp b/Engine/Plugins/AudioInsights/Source/AudioInsights/Private/Views/MixerSourceDashboardViewFactory.cpp
--- a/Engine/Plugins/AudioInsights/Source/AudioInsights/Private/Views/MixerSourceDashboardViewFactory.cpp
+++ b/Engine/Plugins/AudioInsights/Source/AudioInsights/Private/Views/MixerSourceDashboardViewFactory.cpp
@@ -14,6 +14,8 @@
 
 #if WITH_EDITOR
 #include "Editor.h"
+#else
+#include "AudioInsightsComponent.h"
 #endif // WITH_EDITOR
 
 #define LOCTEXT_NAMESPACE "AudioInsights"
@@ -389,24 +391,31 @@ namespace UE::Audio::Insights
 #if WITH_EDITOR
 	void FMixerSourceDashboardViewFactory::OnPIEStarted(bool bSimulating)
 	{
-		PIEState = EPIEState::Running;
+		GameState = EGameState::Running;
 	}
 
 	void FMixerSourceDashboardViewFactory::OnPIEStopped(bool bSimulating)
 	{
 		ResetPlots();
 
-		PIEState = EPIEState::Stopped;
+		GameState = EGameState::Stopped;
 	}
 
 	void FMixerSourceDashboardViewFactory::OnPIEPaused(bool bSimulating)
 	{
-		PIEState = EPIEState::Paused;
+		GameState = EGameState::Paused;
 	}
 
 	void FMixerSourceDashboardViewFactory::OnPIEResumed(bool bSimulating)
 	{
-		PIEState = EPIEState::Running;
+		GameState = EGameState::Running;
+	}
+#else
+	void FMixerSourceDashboardViewFactory::OnAudioInsightsComponentTabSpawn()
+	{
+		// It is guaranteed to have a game running when the audio insights tab is spawning from Unreal Insights 
+		// (In Unreal Insights the Audio Insights menu item is only active when a live trace is running)
+		GameState = EGameState::Running;
 	}
 #endif // WITH_EDITOR
 
@@ -638,12 +647,12 @@ namespace UE::Audio::Insights
 		// Create plot widgets
 		auto GetViewRange = [this]()
 		{
-			if (PIEState == EPIEState::Stopped || BeginTimestamp == TNumericLimits<double>::Max())
+			if (GameState == EGameState::Stopped || BeginTimestamp == TNumericLimits<double>::Max())
 			{
 				return TRange<double>(0, MaxPlotHistorySeconds);
 			}
 
-			if (PIEState == EPIEState::Running)
+			if (GameState == EGameState::Running)
 			{
 				const FTraceModule& TraceModule = FAudioInsightsModule::GetChecked().GetTraceModule();
 				const double FirstTimestamp = TraceModule.GetFirstTimeStamp();
@@ -832,6 +841,8 @@ namespace UE::Audio::Insights
 		FEditorDelegates::EndPIE.AddSP(this, &FMixerSourceDashboardViewFactory::OnPIEStopped);
 		FEditorDelegates::PausePIE.AddSP(this, &FMixerSourceDashboardViewFactory::OnPIEPaused);
 		FEditorDelegates::ResumePIE.AddSP(this, &FMixerSourceDashboardViewFactory::OnPIEResumed);
+#else
+		FAudioInsightsComponent::OnTabSpawn.AddSP(this, &FMixerSourceDashboardViewFactory::OnAudioInsightsComponentTabSpawn);
 #endif // WITH_EDITOR
 
 		TSharedRef<SWidget> TableDashboardWidget = FTraceTableDashboardViewFactory::MakeWidget();
diff --git a/Engine/Plugins/AudioInsights/Source/AudioInsights/Public/AudioInsightsModule.h b/Engine/Plugins/AudioInsights/Source/AudioInsights/Public/AudioInsightsModule.h
--- a/Engine/Plugins/AudioInsights/Source/AudioInsights/Public/AudioInsightsModule.h
+++ b/Engine/Plugins/AudioInsights/Source/AudioInsights/Public/AudioInsightsModule.h
@@ -5,6 +5,7 @@
 #include "AudioInsightsTraceModule.h"
 #include "Framework/Docking/TabManager.h"
 #include "IAudioInsightsModule.h"
+#include "Insights/IUnrealInsightsModule.h"
 #include "Templates/SharedPointer.h"
 #include "Views/DashboardViewFactory.h"
 #include "Widgets/Docking/SDockTab.h"
@@ -36,5 +37,9 @@ namespace UE::Audio::Insights
 	private:
 		TSharedPtr<FDashboardFactory> DashboardFactory;
 		FTraceModule TraceModule;
+
+#if !WITH_EDITOR
+		TSharedPtr<IInsightsComponent> AudioInsightsComponent;
+#endif // !WITH_EDITOR
 	};
 } // namespace UE::Audio::Insights
diff --git a/Engine/Plugins/AudioInsights/Source/AudioInsights/Public/Views/MixerSourceDashboardViewFactory.h b/Engine/Plugins/AudioInsights/Source/AudioInsights/Public/Views/MixerSourceDashboardViewFactory.h
--- a/Engine/Plugins/AudioInsights/Source/AudioInsights/Public/Views/MixerSourceDashboardViewFactory.h
+++ b/Engine/Plugins/AudioInsights/Source/AudioInsights/Public/Views/MixerSourceDashboardViewFactory.h
@@ -51,6 +51,8 @@ namespace UE::Audio::Insights
 		void OnPIEStopped(bool bSimulating);
 		void OnPIEPaused(bool bSimulating);
 		void OnPIEResumed(bool bSimulating);
+#else
+		void OnAudioInsightsComponentTabSpawn();
 #endif // WITH_EDITOR
 
 		void UpdatePlotsWidgetsData();
@@ -91,14 +93,14 @@ namespace UE::Audio::Insights
 		// Column names for plot selector widget 
 		TArray<FName> ColumnNames;
 		
-		enum class EPIEState : uint8
+		enum class EGameState : uint8
 		{
 			Running,
 			Stopped,
 			Paused
 		};
 
-		EPIEState PIEState = EPIEState::Stopped;
+		EGameState GameState = EGameState::Stopped;
 
 		double BeginTimestamp = TNumericLimits<double>::Max();
 		double CurrentTimestamp = 0.0;
diff --git a/Engine/Programs/UnrealInsights/Config/DefaultEngine.ini b/Engine/Programs/UnrealInsights/Config/DefaultEngine.ini
--- a/Engine/Programs/UnrealInsights/Config/DefaultEngine.ini
+++ b/Engine/Programs/UnrealInsights/Config/DefaultEngine.ini
@@ -1,4 +1,5 @@
 [Plugins]
++ProgramEnabledPlugins=AudioInsights
 +ProgramEnabledPlugins=GameplayInsights
 +ProgramEnabledPlugins=TraceDataFilters
 +ProgramEnabledPlugins=SlateInsights
