[Audio Insights] Removed OptimizeCode = CodeOptimization.Never; from AudioInsights.Build.cs

diff --git a/Engine/Plugins/AudioInsights/Source/AudioInsights/AudioInsights.Build.cs b/Engine/Plugins/AudioInsights/Source/AudioInsights/AudioInsights.Build.cs
--- a/Engine/Plugins/AudioInsights/Source/AudioInsights/AudioInsights.Build.cs
+++ b/Engine/Plugins/AudioInsights/Source/AudioInsights/AudioInsights.Build.cs
@@ -7,7 +7,6 @@ public class AudioInsights : ModuleRules
 	public AudioInsights(ReadOnlyTargetRules Target) : base(Target)
 	{
 		PCHUsage = ModuleRules.PCHUsageMode.UseExplicitOrSharedPCHs;
-		OptimizeCode = CodeOptimization.Never;
 
 		PublicDependencyModuleNames.AddRange
 		(
