[Audio Insights] Removed duplicated AudioMixerCore dependency in AudioInsightsEditor.Build.cs

diff --git a/Engine/Plugins/AudioInsights/Source/AudioInsightsEditor/AudioInsightsEditor.Build.cs b/Engine/Plugins/AudioInsights/Source/AudioInsightsEditor/AudioInsightsEditor.Build.cs
--- a/Engine/Plugins/AudioInsights/Source/AudioInsightsEditor/AudioInsightsEditor.Build.cs
+++ b/Engine/Plugins/AudioInsights/Source/AudioInsightsEditor/AudioInsightsEditor.Build.cs
@@ -23,7 +23,6 @@ public class AudioInsightsEditor : ModuleRules
 			{
 				"AudioInsights",
 				"AudioMixer",
-				"AudioMixerCore",
 				"AudioWidgets",
 				"CoreUObject",
 				"Engine",
