[Audio Insights] Mixer Source plot area improvements: Improved legibility of the Y axis labels + fixing scroll data accumulation of the first 5 seconds.

diff --git a/Engine/Plugins/AudioInsights/Source/AudioInsights/Private/Views/MixerSourceDashboardViewFactory.cpp b/Engine/Plugins/AudioInsights/Source/AudioInsights/Private/Views/MixerSourceDashboardViewFactory.cpp
--- a/Engine/Plugins/AudioInsights/Source/AudioInsights/Private/Views/MixerSourceDashboardViewFactory.cpp
+++ b/Engine/Plugins/AudioInsights/Source/AudioInsights/Private/Views/MixerSourceDashboardViewFactory.cpp
@@ -664,7 +664,7 @@ namespace UE::Audio::Insights
 				CurrentTimestamp += TimestampsDiff;
 			}
 
-			return TRange<double>(FMath::Max(0, CurrentTimestamp - MaxPlotHistorySeconds - BeginTimestamp), CurrentTimestamp - BeginTimestamp);
+			return TRange<double>(CurrentTimestamp - MaxPlotHistorySeconds - BeginTimestamp, CurrentTimestamp - BeginTimestamp);
 		};
 
 		if (PlotWidgets.IsEmpty())
@@ -718,7 +718,7 @@ namespace UE::Audio::Insights
 				SNew(SSimpleTimeSlider)
 				.ViewRange_Lambda(GetViewRange)
 				.ClampRangeHighlightSize(0.0f) // Hide clamp range
-				.ScrubPosition_Lambda([]() { return -1.0; }) // Hide scrub
+				.ScrubPosition_Lambda([]() { return TNumericLimits<double>::Lowest(); }) // Hide scrub
 				.PixelSnappingMethod(EWidgetPixelSnapping::Disabled)
 			]
 			+ SVerticalBox::Slot()
diff --git a/Engine/Plugins/AudioInsights/Source/AudioInsights/Private/Views/SAudioCurveView.cpp b/Engine/Plugins/AudioInsights/Source/AudioInsights/Private/Views/SAudioCurveView.cpp
--- a/Engine/Plugins/AudioInsights/Source/AudioInsights/Private/Views/SAudioCurveView.cpp
+++ b/Engine/Plugins/AudioInsights/Source/AudioInsights/Private/Views/SAudioCurveView.cpp
@@ -8,6 +8,7 @@
 #include "Fonts/FontMeasure.h"
 #include "Framework/Application/SlateApplication.h"
 #include "Rendering/DrawElements.h"
+#include "Styling/AppStyle.h"
 #include "Widgets/Text/STextBlock.h"
 #include "Widgets/SToolTip.h"
 
@@ -29,7 +30,7 @@ void SAudioCurveView::Construct( const SAudioCurveView::FArguments& InArgs )
 	Clipping = EWidgetClipping::ClipToBounds;
 	XValueFormattingOptions.MaximumFractionalDigits = 3;
 	LineDrawEffects = ESlateDrawEffect::NoPixelSnapping;
-	LabelFont = FCoreStyle::GetDefaultFontStyle("Regular", 6);
+	LabelFont = FCoreStyle::GetDefaultFontStyle("Bold", 7);
 
 	SetToolTip(CreateCurveTooltip());
 }
@@ -259,12 +260,11 @@ int32 SAudioCurveView::PaintGridLines(const FGeometry& AllottedGeometry, const F
 		}
 	}
 
-	// Draw and label NumHorizontalGridLines horizontal grid lines 
-	const int32 GridLineLabelLayer = LayerId++;
+	// Draw horizontal grid lines 
 	const float MarginBase = YMargin.Get() * Size.Y;
 	const float MarginTop = (1.0f - YMargin.Get()) * Size.Y;
 	const float GridLineYIncrement = (MarginTop - MarginBase) / (NumHorizontalGridLines - 1);
-	const TSharedRef<FSlateFontMeasure> FontMeasureService = FSlateApplication::Get().GetRenderer()->GetFontMeasureService();
+
 	for (uint32 HorizontalLineIndex = 0; HorizontalLineIndex < NumHorizontalGridLines; ++HorizontalLineIndex)
 	{
 		const float WidgetY = GridLineYIncrement * HorizontalLineIndex + MarginBase;
@@ -282,20 +282,52 @@ int32 SAudioCurveView::PaintGridLines(const FGeometry& AllottedGeometry, const F
 			GridLineColor.Get(),
 			false
 		);
+	}
+	return LayerId;
+}
+
+int32 SAudioCurveView::PaintYAxisLabels(const FGeometry& AllottedGeometry, FSlateWindowElementList& OutDrawElements, int32 LayerId) const
+{
+	// Draw Background rectangle
+	const FVector2D RectangleSize(32.0f, AllottedGeometry.GetLocalSize().Y);
+	const FVector2D RectanglePosition(0.0f, 0.0f);
+	
+	FSlateDrawElement::MakeBox(
+		OutDrawElements,
+		LayerId++,
+		AllottedGeometry.ToPaintGeometry(RectangleSize, FSlateLayoutTransform(RectanglePosition)),
+		FAppStyle::GetBrush("BlackBrush"),
+		ESlateDrawEffect::None,
+		FLinearColor(0.0f, 0.0f, 0.0f, 0.6f)
+	);
+	
+	// Draw Y axis labels
+	const FVector2f Size = AllottedGeometry.GetLocalSize();
+
+	const float MarginBase = YMargin.Get() * Size.Y;
+	const float MarginTop  = (1.0f - YMargin.Get()) * Size.Y;
+
+	const float GridLineYIncrement = (MarginTop - MarginBase) / (NumHorizontalGridLines - 1);
+
+	const TSharedRef<FSlateFontMeasure> FontMeasureService = FSlateApplication::Get().GetRenderer()->GetFontMeasureService();
+
+	for (uint32 HorizontalLineIndex = 0; HorizontalLineIndex < NumHorizontalGridLines; ++HorizontalLineIndex)
+	{
+		const float WidgetY = GridLineYIncrement * HorizontalLineIndex + MarginBase;
 
 		// Draw y axis text label every other grid line
 		if (HorizontalLineIndex % 2 == 1)
 		{
-			const float LabelValue = LocalYToValue(Size, WidgetY);
-			FText LabelString = FText::AsNumber(LabelValue, &YValueFormattingOptions);
+			const float LabelValue  = LocalYToValue(Size, WidgetY);
+			const FText LabelString = FText::AsNumber(LabelValue, &YValueFormattingOptions);
 
 			// Position text slightly above the corresponding horizontal line 
 			const FVector2f TextSize = FontMeasureService->Measure(LabelString, LabelFont);
-			FVector2f TextOffset(TextSize.X * 0.5f, WidgetY - TextSize.Y * 0.85f);
+			const FVector2f TextOffset(5.0f, WidgetY - TextSize.Y * 0.85f);
 
 			FSlateDrawElement::MakeText(
 				OutDrawElements,
-				GridLineLabelLayer,
+				LayerId++,
 				AllottedGeometry.ToPaintGeometry(TextSize, FSlateLayoutTransform(TextOffset)),
 				LabelString,
 				LabelFont,
@@ -380,6 +412,9 @@ int32 SAudioCurveView::PaintCurves(const FGeometry& AllottedGeometry, const FSla
 		);
 	}
 
+	// Draw Y axis labels
+	LayerId = PaintYAxisLabels(AllottedGeometry, OutDrawElements, LayerId);
+
 	return LayerId;
 }
 
diff --git a/Engine/Plugins/AudioInsights/Source/AudioInsights/Public/Views/SAudioCurveView.h b/Engine/Plugins/AudioInsights/Source/AudioInsights/Public/Views/SAudioCurveView.h
--- a/Engine/Plugins/AudioInsights/Source/AudioInsights/Public/Views/SAudioCurveView.h
+++ b/Engine/Plugins/AudioInsights/Source/AudioInsights/Public/Views/SAudioCurveView.h
@@ -30,7 +30,7 @@ public:
 		, _YMargin(0.05f)
 		, _HorizontalAxisIncrement(0.5)
 		, _GridLineColor(FLinearColor(0.5f, 0.5f, 0.5f, 0.25f))
-		, _AxesLabelColor(FLinearColor(0.5f, 0.5f, 0.75f))
+		, _AxesLabelColor(FLinearColor::White)
 		, _DesiredSize(FVector2D(100.f,100.f))
 	{}
     	/** View X axis range (in value space) */
@@ -78,6 +78,7 @@ protected:
 	
 	int32 PaintCurves(const FGeometry& AllottedGeometry, const FSlateRect& MyCullingRect, FSlateWindowElementList& OutDrawElements, int32 LayerId, const FWidgetStyle& InWidgetStyle, bool bParentEnabled) const;
 	int32 PaintGridLines(const FGeometry& AllottedGeometry, const FSlateRect& MyCullingRect, FSlateWindowElementList& OutDrawElements, int32 LayerId, const FWidgetStyle& InWidgetStyle, bool bParentEnabled, const SSimpleTimeSlider::FScrubRangeToScreen RangeToScreen) const;
+	int32 PaintYAxisLabels(const FGeometry& AllottedGeometry, FSlateWindowElementList& OutDrawElements, int32 LayerId) const;
 	void UpdateCurveToolTip(const FGeometry& MyGeometry, const FPointerEvent& MouseEvent);
 	void UpdateYDataRange();
 	TSharedRef<SToolTip> CreateCurveTooltip();
