[Audio Insights] Limit to open Audio Insights tab within Unreal Insights to standalone live sessions.

diff --git a/Engine/Plugins/AudioInsights/Source/AudioInsights/AudioInsights.Build.cs b/Engine/Plugins/AudioInsights/Source/AudioInsights/AudioInsights.Build.cs
--- a/Engine/Plugins/AudioInsights/Source/AudioInsights/AudioInsights.Build.cs
+++ b/Engine/Plugins/AudioInsights/Source/AudioInsights/AudioInsights.Build.cs
@@ -27,6 +27,7 @@ public class AudioInsights : ModuleRules
 				"CoreUObject",
 				"InputCore",
 				"OutputLog",
+				"SessionServices",
 				"Slate",
 				"SlateCore",
 				"TraceAnalysis",
diff --git a/Engine/Plugins/AudioInsights/Source/AudioInsights/Private/AudioInsightsComponent.cpp b/Engine/Plugins/AudioInsights/Source/AudioInsights/Private/AudioInsightsComponent.cpp
--- a/Engine/Plugins/AudioInsights/Source/AudioInsights/Private/AudioInsightsComponent.cpp
+++ b/Engine/Plugins/AudioInsights/Source/AudioInsights/Private/AudioInsightsComponent.cpp
@@ -5,6 +5,7 @@
 #include "AudioInsightsStyle.h"
 #include "Modules/ModuleManager.h"
 #include "Trace/StoreClient.h"
+#include "TraceServices/Model/Diagnostics.h"
 #include "Widgets/Docking/SDockTab.h"
 
 #define LOCTEXT_NAMESPACE "AudioInsightsComponent"
@@ -77,7 +78,7 @@ namespace UE::Audio::Insights
 				FOnSpawnTab::CreateRaw(this,  &FAudioInsightsComponent::SpawnTab), 
 				FCanSpawnTab::CreateRaw(this, &FAudioInsightsComponent::CanSpawnTab))
 				.SetDisplayName(Config.TabLabel.IsSet()   ? Config.TabLabel.GetValue()   : LOCTEXT("AudioInsights_TabTitle", "Audio Insights"))
-				.SetTooltipText(Config.TabTooltip.IsSet() ? Config.TabTooltip.GetValue() : LOCTEXT("AudioInsights_TooltipText", "Open the Audio Insights tab."))
+				.SetTooltipText(Config.TabTooltip.IsSet() ? Config.TabTooltip.GetValue() : LOCTEXT("AudioInsights_TooltipText", "Open the Audio Insights tab (Only available for standalone live traces)."))
 				.SetIcon(Config.TabIcon.IsSet() ? Config.TabIcon.GetValue() : FSlateStyle::Get().CreateIcon("AudioInsights.Icon.Submix"));
 
 			const TSharedRef<FWorkspaceItem>* FoundWorkspace = FGlobalTabmanager::Get()->GetLocalWorkspaceMenuRoot()->GetChildItems().FindByPredicate(
@@ -114,7 +115,7 @@ namespace UE::Audio::Insights
 
 	bool FAudioInsightsComponent::Tick(float DeltaTime)
 	{
-		// Audio Insights will be available only if there is an active live session
+		// Audio Insights will be available only if there is an active standalone (non-editor) live session
 		if (bCanCheckForActiveSession && !bCanSpawnTab)
 		{
 			IUnrealInsightsModule& UnrealInsightsModule = FModuleManager::LoadModuleChecked<IUnrealInsightsModule>("TraceInsights");
@@ -128,10 +129,19 @@ namespace UE::Audio::Insights
 				{
 					if (UE::Trace::FStoreClient* StoreClient = UnrealInsightsModule.GetStoreClient())
 					{
-						const UE::Trace::FStoreClient::FSessionInfo* SessionInfo = StoreClient->GetSessionInfoByTraceId(Session->GetTraceId());
-						if (SessionInfo)
+						const UE::Trace::FStoreClient::FSessionInfo* StoreClientSessionInfo = StoreClient->GetSessionInfoByTraceId(Session->GetTraceId());
+						if (StoreClientSessionInfo)
 						{
-							bCanSpawnTab = true;
+							const TraceServices::IDiagnosticsProvider* DiagnosticsProvider = TraceServices::ReadDiagnosticsProvider(*Session.Get());
+							if (DiagnosticsProvider && DiagnosticsProvider->IsSessionInfoAvailable())
+							{
+								const TraceServices::FSessionInfo& TraceServicesSessionInfo = DiagnosticsProvider->GetSessionInfo();
+
+								if (TraceServicesSessionInfo.TargetType != EBuildTargetType::Editor)
+								{
+									bCanSpawnTab = true;
+								}
+							}
 						}
 					}
 				}
