[Audio Insights] Added option to be able to default World Filter to Client 1

diff --git a/Engine/Plugins/AudioInsights/Source/AudioInsightsEditor/AudioInsightsEditor.Build.cs b/Engine/Plugins/AudioInsights/Source/AudioInsightsEditor/AudioInsightsEditor.Build.cs
--- a/Engine/Plugins/AudioInsights/Source/AudioInsightsEditor/AudioInsightsEditor.Build.cs
+++ b/Engine/Plugins/AudioInsights/Source/AudioInsightsEditor/AudioInsightsEditor.Build.cs
@@ -25,6 +25,7 @@ public class AudioInsightsEditor : ModuleRules
 				"AudioMixer",
 				"AudioWidgets",
 				"CoreUObject",
+				"DeveloperSettings",
 				"Engine",
 				"InputCore",
 				"OutputLog",
diff --git a/Engine/Plugins/AudioInsights/Source/AudioInsightsEditor/Private/AudioInsightsEditorDashboardFactory.cpp b/Engine/Plugins/AudioInsights/Source/AudioInsightsEditor/Private/AudioInsightsEditorDashboardFactory.cpp
--- a/Engine/Plugins/AudioInsights/Source/AudioInsightsEditor/Private/AudioInsightsEditorDashboardFactory.cpp
+++ b/Engine/Plugins/AudioInsights/Source/AudioInsightsEditor/Private/AudioInsightsEditorDashboardFactory.cpp
@@ -3,6 +3,7 @@
 
 #include "AudioDevice.h"
 #include "AudioDeviceManager.h"
+#include "AudioInsightsEditorSettings.h"
 #include "AudioInsightsModule.h"
 #include "AudioInsightsStyle.h"
 #include "Editor.h"
@@ -64,7 +65,15 @@ namespace UE::Audio::Insights
 				const IAudioInsightsTraceModule& TraceModule = FAudioInsightsModule::GetChecked().GetTraceModule();
 				TraceModule.StartTraceAnalysis();
 
-				ActiveDeviceId = InDeviceId;
+				const TObjectPtr<const UAudioInsightsEditorSettings> AudioInsightsEditorSettings = GetDefault<UAudioInsightsEditorSettings>();
+
+				// We don't want to set ActiveDeviceId if bWorldFilterDefaultsToFirstClient is true and more than 2 PIE clients are running
+				if (AudioInsightsEditorSettings == nullptr ||
+					(AudioInsightsEditorSettings && !AudioInsightsEditorSettings->bWorldFilterDefaultsToFirstClient) ||
+					AudioDeviceIds.Num() < 2)
+				{
+					ActiveDeviceId = InDeviceId;
+				}
 			}
 		}
 
diff --git a/Engine/Plugins/AudioInsights/Source/AudioInsightsEditor/Private/AudioInsightsEditorSettings.cpp b/Engine/Plugins/AudioInsights/Source/AudioInsightsEditor/Private/AudioInsightsEditorSettings.cpp
new file mode 100644
--- /dev/null
+++ b/Engine/Plugins/AudioInsights/Source/AudioInsightsEditor/Private/AudioInsightsEditorSettings.cpp
@@ -0,0 +1,24 @@
+// Copyright Epic Games, Inc. All Rights Reserved.
+
+#include "AudioInsightsEditorSettings.h"
+
+#include UE_INLINE_GENERATED_CPP_BY_NAME(AudioInsightsEditorSettings)
+
+#define LOCTEXT_NAMESPACE "UAudioInsightsEditorSettings"
+
+FName UAudioInsightsEditorSettings::GetCategoryName() const
+{
+	return "Plugins";
+}
+
+FText UAudioInsightsEditorSettings::GetSectionText() const
+{
+	return LOCTEXT("AudioInsightsEditorSettings_SectionText", "Audio Insights");
+}
+
+FText UAudioInsightsEditorSettings::GetSectionDescription() const
+{
+	return LOCTEXT("AudioInsightsEditorSettings_SectionDesc", "Configure Audio Insights editor settings.");
+}
+
+#undef LOCTEXT_NAMESPACE
diff --git a/Engine/Plugins/AudioInsights/Source/AudioInsightsEditor/Public/AudioInsightsEditorSettings.h b/Engine/Plugins/AudioInsights/Source/AudioInsightsEditor/Public/AudioInsightsEditorSettings.h
new file mode 100644
--- /dev/null
+++ b/Engine/Plugins/AudioInsights/Source/AudioInsightsEditor/Public/AudioInsightsEditorSettings.h
@@ -0,0 +1,22 @@
+// Copyright Epic Games, Inc. All Rights Reserved.
+
+#pragma once
+
+#include "Engine/DeveloperSettings.h"
+
+#include "AudioInsightsEditorSettings.generated.h"
+
+UCLASS(config = EditorPerProjectUserSettings)
+class AUDIOINSIGHTSEDITOR_API UAudioInsightsEditorSettings : public UDeveloperSettings
+{
+	GENERATED_BODY()
+	
+public:
+	virtual FName GetCategoryName() const override;
+	virtual FText GetSectionText() const override;
+	virtual FText GetSectionDescription() const override;
+
+	/** Whether to automatically set the first PIE client in Audio Insights World Filter. */
+	UPROPERTY(Config, EditAnywhere, Category="World Filter")
+	bool bWorldFilterDefaultsToFirstClient = false;
+};
