[AudioInsights] Implemented Save Layout for Audio Insights

diff --git a/Engine/Plugins/AudioInsights/Source/AudioInsightsEditor/Private/AudioInsightsEditorDashboardFactory.cpp b/Engine/Plugins/AudioInsights/Source/AudioInsightsEditor/Private/AudioInsightsEditorDashboardFactory.cpp
--- a/Engine/Plugins/AudioInsights/Source/AudioInsightsEditor/Private/AudioInsightsEditorDashboardFactory.cpp
+++ b/Engine/Plugins/AudioInsights/Source/AudioInsightsEditor/Private/AudioInsightsEditorDashboardFactory.cpp
@@ -8,6 +8,7 @@
 #include "AudioInsightsStyle.h"
 #include "Editor.h"
 #include "Engine/World.h"
+#include "Framework/Docking/LayoutService.h"
 #include "Framework/Docking/TabManager.h"
 #include "Internationalization/Text.h"
 #include "IPropertyTypeCustomization.h"
@@ -211,12 +212,21 @@ namespace UE::Audio::Insights
 
 		DashboardTabManager = FGlobalTabmanager::Get()->NewTabManager(DockTab);
 
+		DashboardTabManager->SetOnPersistLayout(FTabManager::FOnPersistLayout::CreateStatic([](const TSharedRef<FTabManager::FLayout>& InLayout)
+		{
+			if (InLayout->GetPrimaryArea().Pin().IsValid())
+			{
+				FLayoutSaveRestore::SaveToConfig(GEditorLayoutIni, InLayout);
+			}
+		}));
+
 		InitDelegates();
-		TabLayout = GetDefaultTabLayout();
 
 		RegisterTabSpawners();
 		RefreshDeviceSelector();
 
+		const TSharedRef<FTabManager::FLayout> TabLayout = LoadLayoutFromConfig();
+
 		const TSharedRef<SWidget> TabContent = SNew(SVerticalBox)
 			+ SVerticalBox::Slot()
 			.AutoHeight()
@@ -236,7 +246,7 @@ namespace UE::Audio::Insights
 			]
 			+ SVerticalBox::Slot()
 			[
-				DashboardTabManager->RestoreFrom(TabLayout->AsShared(), Args.GetOwnerWindow()).ToSharedRef()
+				DashboardTabManager->RestoreFrom(TabLayout, Args.GetOwnerWindow()).ToSharedRef()
 			];
 
 		DockTab->SetContent(TabContent);
@@ -245,6 +255,12 @@ namespace UE::Audio::Insights
 		{
 			ResetDelegates();
 			UnregisterTabSpawners();
+			SaveLayoutToConfig();
+
+			DashboardTabManager->CloseAllAreas();
+
+			DashboardTabManager.Reset();
+			DashboardWorkspace.Reset();
 		}));
 
 		return DockTab;
@@ -490,7 +506,7 @@ namespace UE::Audio::Insights
 		}
 	}
 
-	TSharedPtr<FTabManager::FLayout> FEditorDashboardFactory::GetDefaultTabLayout()
+	TSharedRef<FTabManager::FLayout> FEditorDashboardFactory::GetDefaultTabLayout()
 	{
 		using namespace EditorDashboardFactoryPrivate;
 
@@ -657,15 +673,6 @@ namespace UE::Audio::Insights
 		}
 	}
 
-	void FEditorDashboardFactory::RegisterViewFactory(TSharedRef<IDashboardViewFactory> InFactory)
-	{
-		if (const FName Name = InFactory->GetName(); 
-			ensureAlwaysMsgf(!DashboardViewFactories.Contains(Name), TEXT("Failed to register Audio Insights Dashboard '%s': Dashboard with name already registered"), *Name.ToString()))
-		{
-			DashboardViewFactories.Add(Name, InFactory);
-		}
-	}
-
 	void FEditorDashboardFactory::UnregisterTabSpawners()
 	{
 		if (DashboardTabManager.IsValid())
@@ -675,17 +682,35 @@ namespace UE::Audio::Insights
 				const FName& FactoryName = KVP.Key;
 				DashboardTabManager->UnregisterTabSpawner(FactoryName);
 			}
-
-			DashboardTabManager.Reset();
 		}
+	}
 
-		DashboardWorkspace.Reset();
+	void FEditorDashboardFactory::RegisterViewFactory(TSharedRef<IDashboardViewFactory> InFactory)
+	{
+		if (const FName Name = InFactory->GetName(); 
+			ensureAlwaysMsgf(!DashboardViewFactories.Contains(Name), TEXT("Failed to register Audio Insights Dashboard '%s': Dashboard with name already registered"), *Name.ToString()))
+		{
+			DashboardViewFactories.Add(Name, InFactory);
+		}
 	}
 
 	void FEditorDashboardFactory::UnregisterViewFactory(FName InName)
 	{
 		DashboardViewFactories.Remove(InName);
 	}
+
+	TSharedRef<FTabManager::FLayout> FEditorDashboardFactory::LoadLayoutFromConfig()
+	{
+		return FLayoutSaveRestore::LoadFromConfig(GEditorLayoutIni, GetDefaultTabLayout());
+	}
+
+	void FEditorDashboardFactory::SaveLayoutToConfig()
+	{
+		if (DashboardTabManager.IsValid())
+		{
+			FLayoutSaveRestore::SaveToConfig(GEditorLayoutIni, DashboardTabManager->PersistLayout());
+		}
+	}
 } // namespace UE::Audio::Insights
 
 #undef LOCTEXT_NAMESPACE
diff --git a/Engine/Plugins/AudioInsights/Source/AudioInsightsEditor/Private/AudioInsightsEditorDashboardFactory.h b/Engine/Plugins/AudioInsights/Source/AudioInsightsEditor/Private/AudioInsightsEditorDashboardFactory.h
--- a/Engine/Plugins/AudioInsights/Source/AudioInsightsEditor/Private/AudioInsightsEditorDashboardFactory.h
+++ b/Engine/Plugins/AudioInsights/Source/AudioInsightsEditor/Private/AudioInsightsEditorDashboardFactory.h
@@ -44,11 +44,14 @@ namespace UE::Audio::Insights
 		TSharedRef<SWidget> MakeMainToolbarWidget();
 
 		void InitDelegates();
-		TSharedPtr<FTabManager::FLayout> GetDefaultTabLayout();
+		TSharedRef<FTabManager::FLayout> GetDefaultTabLayout();
 
 		void RegisterTabSpawners();
 		void UnregisterTabSpawners();
 
+		TSharedRef<FTabManager::FLayout> LoadLayoutFromConfig();
+		void SaveLayoutToConfig();
+
 		FDelegateHandle OnDeviceDestroyedHandle;
 		FDelegateHandle OnWorldRegisteredToAudioDeviceHandle;
 		FDelegateHandle OnWorldUnregisteredFromAudioDeviceHandle;
@@ -62,7 +65,6 @@ namespace UE::Audio::Insights
 
 		TSharedPtr<FTabManager> DashboardTabManager;
 		TSharedPtr<FWorkspaceItem> DashboardWorkspace;
-		TSharedPtr<FTabManager::FLayout> TabLayout;
 
 		TArray<TSharedPtr<::Audio::FDeviceId>> AudioDeviceIds;
 		TSharedPtr<SComboBox<TSharedPtr<::Audio::FDeviceId>>> AudioDeviceComboBox;
