Move AudioInsightsModule.h back to private

diff --git a/Engine/Plugins/AudioInsights/Source/AudioInsights/Private/AudioInsightsModule.cpp b/Engine/Plugins/AudioInsights/Source/AudioInsights/Private/AudioInsightsModule.cpp
--- a/Engine/Plugins/AudioInsights/Source/AudioInsights/Private/AudioInsightsModule.cpp
+++ b/Engine/Plugins/AudioInsights/Source/AudioInsights/Private/AudioInsightsModule.cpp
@@ -95,6 +95,11 @@ namespace UE::Audio::Insights
 		return DashboardFactory->GetDeviceId();
 	}
 
+	FAudioInsightsModule& FAudioInsightsModule::GetChecked()
+	{
+		return static_cast<FAudioInsightsModule&>(FModuleManager::GetModuleChecked<IAudioInsightsModule>("AudioInsights"));
+	}
+
 	TSharedRef<FDashboardFactory> FAudioInsightsModule::GetDashboardFactory()
 	{
 		return DashboardFactory->AsShared();
diff --git a/Engine/Plugins/AudioInsights/Source/AudioInsights/Private/AudioInsightsModule.h b/Engine/Plugins/AudioInsights/Source/AudioInsights/Private/AudioInsightsModule.h
new file mode 100644
--- /dev/null
+++ b/Engine/Plugins/AudioInsights/Source/AudioInsights/Private/AudioInsightsModule.h
@@ -0,0 +1,45 @@
+// Copyright Epic Games, Inc. All Rights Reserved.
+#pragma once 
+
+#include "AudioInsightsDashboardFactory.h"
+#include "AudioInsightsTraceModule.h"
+#include "Framework/Docking/TabManager.h"
+#include "IAudioInsightsModule.h"
+#include "Insights/IUnrealInsightsModule.h"
+#include "Templates/SharedPointer.h"
+#include "Views/DashboardViewFactory.h"
+#include "Widgets/Docking/SDockTab.h"
+
+
+namespace UE::Audio::Insights
+{
+	class FAudioInsightsModule final : public IAudioInsightsModule
+	{
+	public:
+		FAudioInsightsModule() = default;
+
+		virtual ~FAudioInsightsModule() = default;
+
+		virtual void StartupModule() override;
+		virtual void ShutdownModule() override;
+		virtual void RegisterDashboardViewFactory(TSharedRef<IDashboardViewFactory> InDashboardFactory) override;
+		virtual void UnregisterDashboardViewFactory(FName InName) override;
+		virtual ::Audio::FDeviceId GetDeviceId() const override;
+
+		static FAudioInsightsModule& GetChecked();
+		virtual IAudioInsightsTraceModule& GetTraceModule() override;
+
+		TSharedRef<FDashboardFactory> GetDashboardFactory();
+		const TSharedRef<FDashboardFactory> GetDashboardFactory() const;
+
+		virtual TSharedRef<SDockTab> CreateDashboardTabWidget(const FSpawnTabArgs& Args) override;
+
+	private:
+		TSharedPtr<FDashboardFactory> DashboardFactory;
+		FTraceModule TraceModule;
+
+#if !WITH_EDITOR
+		TSharedPtr<IInsightsComponent> AudioInsightsComponent;
+#endif // !WITH_EDITOR
+	};
+} // namespace UE::Audio::Insights
diff --git a/Engine/Plugins/AudioInsights/Source/AudioInsights/Private/IAudioInsightsModule.cpp b/Engine/Plugins/AudioInsights/Source/AudioInsights/Private/IAudioInsightsModule.cpp
--- a/Engine/Plugins/AudioInsights/Source/AudioInsights/Private/IAudioInsightsModule.cpp
+++ b/Engine/Plugins/AudioInsights/Source/AudioInsights/Private/IAudioInsightsModule.cpp
@@ -5,11 +5,18 @@
 
 IAudioInsightsTraceModule& IAudioInsightsModule::GetTraceModule()
 {
-	return GetChecked().GetTraceModule();
+	IAudioInsightsModule& AudioInsightsModule = static_cast<IAudioInsightsModule&>(UE::Audio::Insights::FAudioInsightsModule::GetChecked());
+	return AudioInsightsModule.GetTraceModule();
 }
 
-UE::Audio::Insights::FAudioInsightsModule& IAudioInsightsModule::GetChecked()
+IAudioInsightsModule& IAudioInsightsModule::GetChecked()
 {
-	using namespace UE::Audio::Insights;
-	return static_cast<FAudioInsightsModule&>(FModuleManager::GetModuleChecked<IAudioInsightsModule>(FAudioInsightsModule::GetName()));
+	return static_cast<IAudioInsightsModule&>(UE::Audio::Insights::FAudioInsightsModule::GetChecked());
 }
+
+#ifdef WITH_EDITOR
+IAudioInsightsModule& IAudioInsightsModule::GetEditorChecked()
+{
+	return static_cast<IAudioInsightsModule&>(FModuleManager::GetModuleChecked<IAudioInsightsModule>("AudioInsightsEditor"));
+}
+#endif
diff --git a/Engine/Plugins/AudioInsights/Source/AudioInsights/Public/AudioInsightsModule.h b/Engine/Plugins/AudioInsights/Source/AudioInsights/Public/AudioInsightsModule.h
deleted file mode 100644
diff --git a/Engine/Plugins/AudioInsights/Source/AudioInsights/Public/AudioInsightsTraceProviderBase.h b/Engine/Plugins/AudioInsights/Source/AudioInsights/Public/AudioInsightsTraceProviderBase.h
--- a/Engine/Plugins/AudioInsights/Source/AudioInsights/Public/AudioInsightsTraceProviderBase.h
+++ b/Engine/Plugins/AudioInsights/Source/AudioInsights/Public/AudioInsightsTraceProviderBase.h
@@ -122,10 +122,10 @@ namespace UE::Audio::Insights
 		const FDeviceData* FindFilteredDeviceData() const
 		{
 #if WITH_EDITOR
-			const IAudioInsightsModule& AudioInsightsEditorModule = FModuleManager::GetModuleChecked<IAudioInsightsModule>(IAudioInsightsModule::GetEditorName());
-			const ::Audio::FDeviceId AudioDeviceId = AudioInsightsEditorModule.GetDeviceId();
+			const IAudioInsightsModule& AudioInsightsModule = IAudioInsightsModule::GetEditorChecked();
+			const ::Audio::FDeviceId AudioDeviceId = AudioInsightsModule.GetDeviceId();
 #else
-			const IAudioInsightsModule& AudioInsightsModule = FModuleManager::GetModuleChecked<IAudioInsightsModule>(IAudioInsightsModule::GetName());
+			const IAudioInsightsModule& AudioInsightsModule = IAudioInsightsModule::GetChecked();
 			const ::Audio::FDeviceId AudioDeviceId = AudioInsightsModule.GetDeviceId();
 #endif // WITH_EDITOR
 
@@ -181,10 +181,10 @@ namespace UE::Audio::Insights
 		FDeviceData* FindFilteredDeviceData()
 		{
 #if WITH_EDITOR
-			const IAudioInsightsModule& AudioInsightsEditorModule = FModuleManager::GetModuleChecked<IAudioInsightsModule>(IAudioInsightsModule::GetEditorName());
-			const ::Audio::FDeviceId AudioDeviceId = AudioInsightsEditorModule.GetDeviceId();
+			const IAudioInsightsModule& AudioInsightsModule = IAudioInsightsModule::GetEditorChecked();
+			const ::Audio::FDeviceId AudioDeviceId = AudioInsightsModule.GetDeviceId();
 #else
-			const IAudioInsightsModule& AudioInsightsModule = FModuleManager::GetModuleChecked<IAudioInsightsModule>(IAudioInsightsModule::GetName());
+			const IAudioInsightsModule& AudioInsightsModule = IAudioInsightsModule::GetChecked();
 			const ::Audio::FDeviceId AudioDeviceId = AudioInsightsModule.GetDeviceId();
 #endif // WITH_EDITOR
 
diff --git a/Engine/Plugins/AudioInsights/Source/AudioInsights/Public/IAudioInsightsModule.h b/Engine/Plugins/AudioInsights/Source/AudioInsights/Public/IAudioInsightsModule.h
--- a/Engine/Plugins/AudioInsights/Source/AudioInsights/Public/IAudioInsightsModule.h
+++ b/Engine/Plugins/AudioInsights/Source/AudioInsights/Public/IAudioInsightsModule.h
@@ -2,20 +2,22 @@
 #pragma once
 
 #include "AudioDefines.h"
+#include "Framework/Docking/TabManager.h"
+#include "IAudioInsightsModuleInterface.h"
 #include "Modules/ModuleInterface.h"
 #include "Templates/SharedPointer.h"
 #include "UObject/NameTypes.h"
+#include "Widgets/Docking/SDockTab.h"
 
 class IAudioInsightsTraceModule;
 
 namespace UE::Audio::Insights
 {
 	class IDashboardViewFactory;
-	class FAudioInsightsModule;
 	class FTraceProviderBase;
 } // namespace UE::Audio::Insights
 
-class AUDIOINSIGHTS_API IAudioInsightsModule : public IModuleInterface
+class AUDIOINSIGHTS_API IAudioInsightsModule : public IAudioInsightsModuleInterface
 {
 public:
 	virtual void RegisterDashboardViewFactory(TSharedRef<UE::Audio::Insights::IDashboardViewFactory> InDashboardFactory) = 0;
@@ -23,19 +25,14 @@ public:
 
 	virtual ::Audio::FDeviceId GetDeviceId() const = 0;
 
-	virtual IAudioInsightsTraceModule& GetTraceModule();
+	virtual IAudioInsightsTraceModule& GetTraceModule() override;
 
-	static UE::Audio::Insights::FAudioInsightsModule& GetChecked();
+	virtual TSharedRef<SDockTab> CreateDashboardTabWidget(const FSpawnTabArgs& Args) = 0;
 
-	static FName GetName()
-	{
-		const FLazyName ModuleName = "AudioInsights";
-		return ModuleName.Resolve();
-	}
+	static IAudioInsightsModule& GetChecked();
+
+#ifdef WITH_EDITOR
+	static IAudioInsightsModule& GetEditorChecked();
+#endif
 
-	static FName GetEditorName()
-	{
-		const FLazyName EditorModuleName = "AudioInsightsEditor";
-		return EditorModuleName.Resolve();
-	}
 };
diff --git a/Engine/Plugins/AudioInsights/Source/AudioInsights/Public/IAudioInsightsModuleInterface.h b/Engine/Plugins/AudioInsights/Source/AudioInsights/Public/IAudioInsightsModuleInterface.h
new file mode 100644
--- /dev/null
+++ b/Engine/Plugins/AudioInsights/Source/AudioInsights/Public/IAudioInsightsModuleInterface.h
@@ -0,0 +1,26 @@
+// Copyright Epic Games, Inc. All Rights Reserved.
+#pragma once
+
+#include "AudioDefines.h"
+#include "Modules/ModuleInterface.h"
+#include "Templates/SharedPointer.h"
+#include "UObject/NameTypes.h"
+
+class IAudioInsightsTraceModule;
+
+namespace UE::Audio::Insights
+{
+	class IDashboardViewFactory;
+	class FTraceProviderBase;
+} // namespace UE::Audio::Insights
+
+class AUDIOINSIGHTS_API IAudioInsightsModuleInterface : public IModuleInterface
+{
+public:
+	virtual void RegisterDashboardViewFactory(TSharedRef<UE::Audio::Insights::IDashboardViewFactory> InDashboardFactory) = 0;
+	virtual void UnregisterDashboardViewFactory(FName InName) = 0;
+
+	virtual ::Audio::FDeviceId GetDeviceId() const = 0;
+
+	virtual IAudioInsightsTraceModule& GetTraceModule() = 0;
+};
diff --git a/Engine/Plugins/AudioInsights/Source/AudioInsightsEditor/Private/AudioInsightsEditorDashboardFactory.cpp b/Engine/Plugins/AudioInsights/Source/AudioInsightsEditor/Private/AudioInsightsEditorDashboardFactory.cpp
--- a/Engine/Plugins/AudioInsights/Source/AudioInsightsEditor/Private/AudioInsightsEditorDashboardFactory.cpp
+++ b/Engine/Plugins/AudioInsights/Source/AudioInsightsEditor/Private/AudioInsightsEditorDashboardFactory.cpp
@@ -3,8 +3,8 @@
 
 #include "AudioDevice.h"
 #include "AudioDeviceManager.h"
+#include "AudioInsightsEditorModule.h"
 #include "AudioInsightsEditorSettings.h"
-#include "AudioInsightsModule.h"
 #include "AudioInsightsStyle.h"
 #include "Editor.h"
 #include "Engine/World.h"
@@ -63,7 +63,7 @@ namespace UE::Audio::Insights
 		{
 			if (bStartWithPIE)
 			{
-				const IAudioInsightsTraceModule& TraceModule = FAudioInsightsModule::GetChecked().GetTraceModule();
+				const IAudioInsightsTraceModule& TraceModule = FAudioInsightsEditorModule::GetChecked().GetTraceModule();
 				TraceModule.StartTraceAnalysis();
 
 				const TObjectPtr<const UAudioInsightsEditorSettings> AudioInsightsEditorSettings = GetDefault<UAudioInsightsEditorSettings>();
@@ -85,7 +85,7 @@ namespace UE::Audio::Insights
 	{
 		if (bStartWithPIE)
 		{
-			const IAudioInsightsTraceModule& TraceModule = FAudioInsightsModule::GetChecked().GetTraceModule();
+			const IAudioInsightsTraceModule& TraceModule = FAudioInsightsEditorModule::GetChecked().GetTraceModule();
 			TraceModule.StartTraceAnalysis();
 		}
 	}
@@ -99,7 +99,7 @@ namespace UE::Audio::Insights
 	{
 		if (bStopWithPIE)
 		{
-			const IAudioInsightsTraceModule& TraceModule = FAudioInsightsModule::GetChecked().GetTraceModule();
+			const IAudioInsightsTraceModule& TraceModule = FAudioInsightsEditorModule::GetChecked().GetTraceModule();
 			TraceModule.StopTraceAnalysis();
 		}
 
diff --git a/Engine/Plugins/AudioInsights/Source/AudioInsightsEditor/Private/AudioInsightsEditorModule.cpp b/Engine/Plugins/AudioInsights/Source/AudioInsightsEditor/Private/AudioInsightsEditorModule.cpp
--- a/Engine/Plugins/AudioInsights/Source/AudioInsightsEditor/Private/AudioInsightsEditorModule.cpp
+++ b/Engine/Plugins/AudioInsights/Source/AudioInsightsEditor/Private/AudioInsightsEditorModule.cpp
@@ -4,6 +4,7 @@
 #include "AudioInsightsEditorLog.h"
 #include "AudioInsightsStyle.h"
 #include "Framework/Docking/TabManager.h"
+#include "IAudioInsightsModule.h"
 #include "Modules/ModuleManager.h"
 #include "Templates/SharedPointer.h"
 #include "UObject/NameTypes.h"
@@ -71,12 +72,12 @@ namespace UE::Audio::Insights
 
 	FAudioInsightsEditorModule& FAudioInsightsEditorModule::GetChecked()
 	{
-		return static_cast<FAudioInsightsEditorModule&>(FModuleManager::GetModuleChecked<IAudioInsightsModule>(GetEditorName()));
+		return static_cast<FAudioInsightsEditorModule&>(FModuleManager::GetModuleChecked<IAudioInsightsEditorModule>("AudioInsightsEditor"));
 	}
 
 	IAudioInsightsTraceModule& FAudioInsightsEditorModule::GetTraceModule()
 	{
-		IAudioInsightsModule& InsightsModule = FModuleManager::LoadModuleChecked<IAudioInsightsModule>(IAudioInsightsModule::GetName());
+		IAudioInsightsModule& InsightsModule = IAudioInsightsModule::GetChecked();
 		return InsightsModule.GetTraceModule();
 
 	}
diff --git a/Engine/Plugins/AudioInsights/Source/AudioInsightsEditor/Private/AudioInsightsEditorModule.h b/Engine/Plugins/AudioInsights/Source/AudioInsightsEditor/Private/AudioInsightsEditorModule.h
--- a/Engine/Plugins/AudioInsights/Source/AudioInsightsEditor/Private/AudioInsightsEditorModule.h
+++ b/Engine/Plugins/AudioInsights/Source/AudioInsightsEditor/Private/AudioInsightsEditorModule.h
@@ -4,7 +4,7 @@
 #include "AudioInsightsEditorDashboardFactory.h"
 #include "AudioInsightsTraceModule.h"
 #include "Framework/Docking/TabManager.h"
-#include "IAudioInsightsModule.h"
+#include "IAudioInsightsEditorModule.h"
 #include "Templates/SharedPointer.h"
 #include "Views/DashboardViewFactory.h"
 #include "Views/MixerSourceMuteSoloFilter.h"
@@ -13,7 +13,7 @@
 
 namespace UE::Audio::Insights
 {
-	class FAudioInsightsEditorModule final : public IAudioInsightsModule
+	class FAudioInsightsEditorModule final : public IAudioInsightsEditorModule
 	{
 	public:
 		FAudioInsightsEditorModule() = default;
@@ -30,7 +30,7 @@ namespace UE::Audio::Insights
 		const TSharedRef<FEditorDashboardFactory> GetDashboardFactory() const;
 
 		static FAudioInsightsEditorModule& GetChecked();
-		IAudioInsightsTraceModule& GetTraceModule() override;
+		virtual IAudioInsightsTraceModule& GetTraceModule() override;
 
 	private:
 		void RegisterMenus();
diff --git a/Engine/Plugins/AudioInsights/Source/AudioInsightsEditor/Private/IAudioInsightsEditorModule.cpp b/Engine/Plugins/AudioInsights/Source/AudioInsightsEditor/Private/IAudioInsightsEditorModule.cpp
new file mode 100644
--- /dev/null
+++ b/Engine/Plugins/AudioInsights/Source/AudioInsightsEditor/Private/IAudioInsightsEditorModule.cpp
@@ -0,0 +1,15 @@
+// Copyright Epic Games, Inc. All Rights Reserved.
+#include "IAudioInsightsEditorModule.h"
+
+#include "AudioInsightsEditorModule.h"
+
+IAudioInsightsTraceModule& IAudioInsightsEditorModule::GetTraceModule()
+{
+	IAudioInsightsEditorModule& AudioInsightsModule = static_cast<IAudioInsightsEditorModule&>(UE::Audio::Insights::FAudioInsightsEditorModule::GetChecked());
+	return AudioInsightsModule.GetTraceModule();
+}
+
+IAudioInsightsEditorModule& IAudioInsightsEditorModule::GetChecked()
+{
+	return static_cast<IAudioInsightsEditorModule&>(UE::Audio::Insights::FAudioInsightsEditorModule::GetChecked());
+}
diff --git a/Engine/Plugins/AudioInsights/Source/AudioInsightsEditor/Public/IAudioInsightsEditorModule.h b/Engine/Plugins/AudioInsights/Source/AudioInsightsEditor/Public/IAudioInsightsEditorModule.h
new file mode 100644
--- /dev/null
+++ b/Engine/Plugins/AudioInsights/Source/AudioInsightsEditor/Public/IAudioInsightsEditorModule.h
@@ -0,0 +1,34 @@
+// Copyright Epic Games, Inc. All Rights Reserved.
+#pragma once
+
+#include "AudioDefines.h"
+#include "Framework/Docking/TabManager.h"
+#include "IAudioInsightsModuleInterface.h"
+#include "Modules/ModuleInterface.h"
+#include "Templates/SharedPointer.h"
+#include "UObject/NameTypes.h"
+#include "Widgets/Docking/SDockTab.h"
+
+class IAudioInsightsTraceModule;
+class SDockTab;
+
+namespace UE::Audio::Insights
+{
+	class IDashboardViewFactory;
+	class FTraceProviderBase;
+} // namespace UE::Audio::Insights
+
+class IAudioInsightsEditorModule : public IAudioInsightsModuleInterface
+{
+public:
+	AUDIOINSIGHTSEDITOR_API virtual void RegisterDashboardViewFactory(TSharedRef<UE::Audio::Insights::IDashboardViewFactory> InDashboardFactory) = 0;
+	AUDIOINSIGHTSEDITOR_API virtual void UnregisterDashboardViewFactory(FName InName) = 0;
+
+	AUDIOINSIGHTSEDITOR_API virtual ::Audio::FDeviceId GetDeviceId() const = 0;
+
+	AUDIOINSIGHTSEDITOR_API virtual IAudioInsightsTraceModule& GetTraceModule() override;
+
+	AUDIOINSIGHTSEDITOR_API virtual TSharedRef<SDockTab> CreateDashboardTabWidget(const FSpawnTabArgs& Args) = 0;
+
+	AUDIOINSIGHTSEDITOR_API static IAudioInsightsEditorModule& GetChecked();
+};
diff --git a/Engine/Plugins/Runtime/AudioModulation/Source/AudioModulationEditor/Private/AudioModulationEditor.cpp b/Engine/Plugins/Runtime/AudioModulation/Source/AudioModulationEditor/Private/AudioModulationEditor.cpp
--- a/Engine/Plugins/Runtime/AudioModulation/Source/AudioModulationEditor/Private/AudioModulationEditor.cpp
+++ b/Engine/Plugins/Runtime/AudioModulation/Source/AudioModulationEditor/Private/AudioModulationEditor.cpp
@@ -13,7 +13,7 @@
 #include "Framework/Commands/UIAction.h"
 #include "Framework/MultiBox/MultiBoxBuilder.h"
 #include "Framework/MultiBox/MultiBoxExtender.h"
-#include "IAudioInsightsModule.h"
+#include "IAudioInsightsEditorModule.h"
 #include "IAudioModulation.h"
 #include "ICurveEditorModule.h"
 #include "Insights/Views/ControlBusDashboardViewFactory.h"
@@ -61,7 +61,7 @@ namespace
 		if (!bIsModulationRegisteredInAudioInsights)
 		{
 
-			IAudioInsightsModule& InsightsModule = FModuleManager::LoadModuleChecked<IAudioInsightsModule>(IAudioInsightsModule::GetEditorName());
+			IAudioInsightsEditorModule& InsightsModule = IAudioInsightsEditorModule::GetChecked();
 			InsightsModule.RegisterDashboardViewFactory(MakeShared<AudioModulationEditor::FControlBusDashboardViewFactory>());
 
 			bIsModulationRegisteredInAudioInsights = true;
@@ -72,7 +72,7 @@ namespace
 	{
 		if (bIsModulationRegisteredInAudioInsights)
 		{
-			IAudioInsightsModule& InsightsModule = FModuleManager::GetModuleChecked<IAudioInsightsModule>(IAudioInsightsModule::GetEditorName());
+			IAudioInsightsEditorModule& InsightsModule = IAudioInsightsEditorModule::GetChecked();
 			InsightsModule.UnregisterDashboardViewFactory("ControlBuses");
 			bIsModulationRegisteredInAudioInsights = false;
 		}
diff --git a/Engine/Plugins/Runtime/AudioModulation/Source/AudioModulationEditor/Private/Insights/Views/ControlBusDashboardViewFactory.cpp b/Engine/Plugins/Runtime/AudioModulation/Source/AudioModulationEditor/Private/Insights/Views/ControlBusDashboardViewFactory.cpp
--- a/Engine/Plugins/Runtime/AudioModulation/Source/AudioModulationEditor/Private/Insights/Views/ControlBusDashboardViewFactory.cpp
+++ b/Engine/Plugins/Runtime/AudioModulation/Source/AudioModulationEditor/Private/Insights/Views/ControlBusDashboardViewFactory.cpp
@@ -5,7 +5,7 @@
 #include "AssetRegistry/AssetRegistryModule.h"
 #include "AudioDefines.h"
 #include "AudioDeviceManager.h"
-#include "AudioInsightsModule.h"
+#include "IAudioInsightsEditorModule.h"
 #include "AudioInsightsStyle.h"
 #include "AudioInsightsTraceModule.h"
 #include "AudioInsightsTraceProviderBase.h"
@@ -30,7 +30,7 @@ namespace AudioModulationEditor
 
 	FControlBusDashboardViewFactory::FControlBusDashboardViewFactory()
 	{
-		IAudioInsightsModule& InsightsModule = FModuleManager::LoadModuleChecked<IAudioInsightsModule>(IAudioInsightsModule::GetEditorName());
+		IAudioInsightsEditorModule& InsightsModule = IAudioInsightsEditorModule::GetChecked();
 		IAudioInsightsTraceModule& InsightsTraceModule = InsightsModule.GetTraceModule();
 
 		TSharedPtr<FControlBusTraceProvider> ControlBusProvider = MakeShared<FControlBusTraceProvider>(InsightsTraceModule.GetChannelManager());
