[Audio Insights] Mixer Source tab: separate plots graph to its own tab

diff --git a/Engine/Plugins/AudioInsights/Source/AudioInsights/Private/Views/MixerSourceDashboardViewFactory.cpp b/Engine/Plugins/AudioInsights/Source/AudioInsights/Private/Views/MixerSourceDashboardViewFactory.cpp
--- a/Engine/Plugins/AudioInsights/Source/AudioInsights/Private/Views/MixerSourceDashboardViewFactory.cpp
+++ b/Engine/Plugins/AudioInsights/Source/AudioInsights/Private/Views/MixerSourceDashboardViewFactory.cpp
@@ -4,13 +4,14 @@
 #include "AudioInsightsModule.h"
 #include "AudioInsightsStyle.h"
 #include "DSP/Dsp.h"
+#include "Framework/Docking/LayoutService.h"
 #include "Internationalization/Text.h"
 #include "Providers/MixerSourceTraceProvider.h"
 #include "SSimpleTimeSlider.h"
+#include "Styling/AppStyle.h"
 #include "Templates/SharedPointer.h"
 #include "UObject/SoftObjectPath.h"
 #include "Widgets/Input/SComboBox.h"
-#include "Widgets/Layout/SSplitter.h"
 
 #if WITH_EDITOR
 #include "Editor.h"
@@ -43,7 +44,14 @@ namespace UE::Audio::Insights
 			return LastValue;
 		};
 
-		static const FText PlotColumnSelectDescription = LOCTEXT("AudioDashboard_MixerSources_SelectPlotColumnDescription", "Select a column from the table to plot.");
+		const FText PlotColumnSelectDescription = LOCTEXT("AudioDashboard_MixerSources_SelectPlotColumnDescription", "Select a column from the table to plot.");
+		const FText PlotsIconDescription = LOCTEXT("AudioDashboard_MixerSources_PlotsIconDescription", "Show/Hides the Mixer Sources Plots section.");
+
+		const FText MixerSourcesWorkspaceName = LOCTEXT("MixerSourcesWorkspace_Name", "MixerSourcesWorkspace");
+
+		const FName MixerSourcesTableTabName = "MixerSourcesTableTab";
+		const FName MixerSourcesPlotsTabName = "MixerSourcesPlotsTab";
+
 	} // namespace MixerSourcePrivate
 
 	const double FMixerSourceDashboardViewFactory::MaxPlotHistorySeconds = 5.0;
@@ -832,10 +840,190 @@ namespace UE::Audio::Insights
 				]
 			];
 	}
+
+	TSharedRef<SWidget> FMixerSourceDashboardViewFactory::MakePlotsButtonWidget()
+	{
+		return SAssignNew(PlotsButton, SCheckBox)
+			.Style(&FAppStyle::Get().GetWidgetStyle<FCheckBoxStyle>("ToggleButtonCheckBox"))
+			.OnCheckStateChanged(this, &FMixerSourceDashboardViewFactory::TogglePlotsTabVisibility)
+			.ToolTipText(MixerSourcePrivate::PlotsIconDescription)
+			[
+				SNew(SHorizontalBox)
+				+ SHorizontalBox::Slot()
+				.AutoWidth()
+				.VAlign(VAlign_Center)
+				[
+					SNew(SImage)
+					.ColorAndOpacity(FSlateColor::UseForeground())
+					.Image(FAppStyle::GetBrush("GenericCurveEditor.SetViewModeStacked"))
+				]
+			];
+	}
+
+	TSharedRef<FTabManager::FLayout> FMixerSourceDashboardViewFactory::LoadLayoutFromConfig()
+	{
+		return FLayoutSaveRestore::LoadFromConfig(GEditorLayoutIni, GetDefaultTabLayout());
+	}
+
+	void FMixerSourceDashboardViewFactory::SaveLayoutToConfig()
+	{
+		if (MixerSourcesTabManager.IsValid())
+		{
+			FLayoutSaveRestore::SaveToConfig(GEditorLayoutIni, MixerSourcesTabManager->PersistLayout());
+		}
+	}
+
+	void FMixerSourceDashboardViewFactory::TogglePlotsTabVisibility(ECheckBoxState InCheckboxState)
+	{
+		using namespace MixerSourcePrivate;
+
+		if (!MixerSourcesTabManager.IsValid())
+		{
+			return;
+		}
+
+		if (InCheckboxState == ECheckBoxState::Checked)
+		{
+			MixerSourcesTabManager->TryInvokeTab(MixerSourcesPlotsTabName);
+		}
+		else if (InCheckboxState == ECheckBoxState::Unchecked)
+		{
+			const TSharedPtr<SDockTab> PlotsTab = MixerSourcesTabManager->FindExistingLiveTab(MixerSourcesPlotsTabName);
+			if (PlotsTab.IsValid())
+			{
+				PlotsTab->RequestCloseTab();
+			}
+		}
+
+		SaveLayoutToConfig();
+	}
+#endif // WITH_EDITOR
+
+	TSharedRef<SDockTab> FMixerSourceDashboardViewFactory::CreateMixerSourcesTab(const FSpawnTabArgs& Args)
+	{
+		return SNew(SDockTab)
+			.Clipping(EWidgetClipping::ClipToBounds)
+			[
+				SNew(SVerticalBox)
+#if WITH_EDITOR
+				+SVerticalBox::Slot()
+				.AutoHeight()
+				.HAlign(HAlign_Fill)
+				.Padding(0.0f, 0.0f, 0.0f, 6.0f)
+				[
+					SNew(SHorizontalBox)
+					+SHorizontalBox::Slot()
+					.HAlign(HAlign_Left)
+					[
+						MakeMuteSoloWidget()
+					]
+					+ SHorizontalBox::Slot()
+					.HAlign(HAlign_Right)
+					[
+						SNullWidget::NullWidget
+					]
+					+ SHorizontalBox::Slot()
+					.AutoWidth()
+					.HAlign(HAlign_Right)
+					[
+						MakePlotsButtonWidget()
+					]
+				]
 #endif // WITH_EDITOR
+				+ SVerticalBox::Slot()
+				.HAlign(HAlign_Fill)
+				[
+					FTraceTableDashboardViewFactory::MakeWidget()
+				]
+			];
+	}
+
+	TSharedRef<SDockTab> FMixerSourceDashboardViewFactory::CreatePlotsTab(const FSpawnTabArgs& Args)
+	{
+		if (PlotsButton.IsValid())
+		{
+			PlotsButton->SetIsChecked(ECheckBoxState::Checked);
+		}
+
+		return SNew(SDockTab)
+			.Clipping(EWidgetClipping::ClipToBounds)
+			[
+				MakePlotsWidget()
+			];
+	}
+
+	void FMixerSourceDashboardViewFactory::RegisterTabSpawners()
+	{
+		using namespace MixerSourcePrivate;
+
+		if (!MixerSourcesTabManager.IsValid())
+		{
+			return;
+		}
+
+		MixerSourcesTabManager->RegisterTabSpawner(MixerSourcesTableTabName, FOnSpawnTab::CreateSP(this, &FMixerSourceDashboardViewFactory::CreateMixerSourcesTab))
+			.SetDisplayName(LOCTEXT("MixerSourceTab_MixerSourcesTable_Name", "Mixer Sources"))
+			.SetGroup(MixerSourcesWorkspace.ToSharedRef())
+			.SetMenuType(ETabSpawnerMenuType::Hidden);
+#if WITH_EDITOR
+		MixerSourcesTabManager->RegisterTabSpawner(MixerSourcesPlotsTabName, FOnSpawnTab::CreateSP(this, &FMixerSourceDashboardViewFactory::CreatePlotsTab))
+			.SetDisplayName(LOCTEXT("MixerSourceTab_PlotsTab_Name", "Mixer Sources Plots"))
+			.SetGroup(MixerSourcesWorkspace.ToSharedRef())
+			.SetMenuType(ETabSpawnerMenuType::Hidden);
+#endif // WITH_EDITOR
+	}
+
+	void FMixerSourceDashboardViewFactory::UnregisterTabSpawners()
+	{
+		using namespace MixerSourcePrivate;
+		
+		if (MixerSourcesTabManager.IsValid())
+		{
+			MixerSourcesTabManager->UnregisterTabSpawner(MixerSourcesTableTabName);
+#if WITH_EDITOR
+			MixerSourcesTabManager->UnregisterTabSpawner(MixerSourcesPlotsTabName);
+#endif // WITH_EDITOR
+		}
+	}
+
+	TSharedRef<FTabManager::FLayout> FMixerSourceDashboardViewFactory::GetDefaultTabLayout()
+	{
+		using namespace MixerSourcePrivate;
+
+		return FTabManager::NewLayout("MixerSourceTabsLayout_v1")
+		->AddArea
+		(
+			FTabManager::NewPrimaryArea()
+			->SetOrientation(Orient_Vertical)
+			->Split
+			(
+				FTabManager::NewSplitter()
+				->SetOrientation(Orient_Vertical)
+				->SetSizeCoefficient(0.7f)
+				->Split
+				(
+					FTabManager::NewStack()
+					->SetSizeCoefficient(0.7f)
+					->SetHideTabWell(true)
+					->AddTab(MixerSourcesTableTabName, ETabState::OpenedTab)
+				)
+#if WITH_EDITOR
+				->Split
+				(
+					FTabManager::NewStack()
+					->SetSizeCoefficient(0.3f)
+					->AddTab(MixerSourcesPlotsTabName, ETabState::OpenedTab)
+					->SetHideTabWell(true)
+				)
+#endif // WITH_EDITOR
+			)
+		);
+	}
 
 	TSharedRef<SWidget> FMixerSourceDashboardViewFactory::MakeWidget()
 	{
+		using namespace MixerSourcePrivate;
+
 #if WITH_EDITOR
 		FEditorDelegates::PostPIEStarted.AddSP(this, &FMixerSourceDashboardViewFactory::OnPIEStarted);
 		FEditorDelegates::EndPIE.AddSP(this, &FMixerSourceDashboardViewFactory::OnPIEStopped);
@@ -845,36 +1033,48 @@ namespace UE::Audio::Insights
 		FAudioInsightsComponent::OnTabSpawn.AddSP(this, &FMixerSourceDashboardViewFactory::OnAudioInsightsComponentTabSpawn);
 #endif // WITH_EDITOR
 
-		TSharedRef<SWidget> TableDashboardWidget = FTraceTableDashboardViewFactory::MakeWidget();
-		TSharedRef<SWidget> PlotsWidget = MakePlotsWidget();
+		const TSharedRef<SDockTab> DockTab = SNew(SDockTab);
+		MixerSourcesTabManager = FGlobalTabmanager::Get()->NewTabManager(DockTab);
 
-		return SNew(SVerticalBox)
 #if WITH_EDITOR
-			+ SVerticalBox::Slot()
-			.AutoHeight()
-			.HAlign(HAlign_Fill)
-			.Padding(0.0f, 0.0f, 0.0f, 6.0f)
-			[
-				MakeMuteSoloWidget()
-			]
+		MixerSourcesTabManager->SetOnPersistLayout(FTabManager::FOnPersistLayout::CreateStatic([](const TSharedRef<FTabManager::FLayout>& InLayout)
+		{
+			if (InLayout->GetPrimaryArea().Pin().IsValid())
+			{
+				FLayoutSaveRestore::SaveToConfig(GEditorLayoutIni, InLayout);
+			}
+		}));
 #endif // WITH_EDITOR
-			+ SVerticalBox::Slot()
-			.HAlign(HAlign_Fill)
-			[
-				// Dashboard and plots area
-				SNew(SSplitter)
-				.Orientation(Orient_Vertical)
-				+ SSplitter::Slot()
-				.Value(0.68f)
-				[
-					TableDashboardWidget
-				]
-				+ SSplitter::Slot()
-				.Value(0.32f)
-				[
-					PlotsWidget
-				]
-			];
+
+		DockTab->SetOnTabClosed(SDockTab::FOnTabClosedCallback::CreateLambda([this](TSharedRef<SDockTab> TabClosed)
+		{
+			UnregisterTabSpawners();
+
+#if WITH_EDITOR
+			SaveLayoutToConfig();
+#endif // WITH_EDITOR
+
+			if (MixerSourcesTabManager.IsValid())
+			{
+				MixerSourcesTabManager->CloseAllAreas();
+
+				MixerSourcesTabManager.Reset();
+				MixerSourcesWorkspace.Reset();
+			}
+		}));
+
+		
+		MixerSourcesWorkspace = MixerSourcesTabManager->AddLocalWorkspaceMenuCategory(MixerSourcesWorkspaceName);
+
+		RegisterTabSpawners();
+
+#if WITH_EDITOR
+		const TSharedRef<FTabManager::FLayout> TabLayout = LoadLayoutFromConfig();
+#else
+		const TSharedRef<FTabManager::FLayout> TabLayout = GetDefaultTabLayout();
+#endif // WITH_EDITOR
+
+		return MixerSourcesTabManager->RestoreFrom(TabLayout, TSharedPtr<SWindow>()).ToSharedRef();
 	}	
 
 	void FMixerSourceDashboardViewFactory::ProcessEntries(FTraceTableDashboardViewFactory::EProcessReason InReason)
diff --git a/Engine/Plugins/AudioInsights/Source/AudioInsights/Public/Views/MixerSourceDashboardViewFactory.h b/Engine/Plugins/AudioInsights/Source/AudioInsights/Public/Views/MixerSourceDashboardViewFactory.h
--- a/Engine/Plugins/AudioInsights/Source/AudioInsights/Public/Views/MixerSourceDashboardViewFactory.h
+++ b/Engine/Plugins/AudioInsights/Source/AudioInsights/Public/Views/MixerSourceDashboardViewFactory.h
@@ -1,6 +1,7 @@
 // Copyright Epic Games, Inc. All Rights Reserved.
 #pragma once
 
+#include "Framework/Docking/TabManager.h"
 #include "Messages/MixerSourceTraceMessages.h"
 #include "Views/SAudioCurveView.h"
 #include "Views/TableDashboardViewFactory.h"
@@ -69,11 +70,25 @@ namespace UE::Audio::Insights
 
 #if WITH_EDITOR
 		TSharedRef<SWidget> MakeMuteSoloWidget();
+		TSharedRef<SWidget> MakePlotsButtonWidget();
 
 		void ToggleMuteForAllItems(ECheckBoxState NewState);
 		void ToggleSoloForAllItems(ECheckBoxState NewState);
+
+		TSharedRef<FTabManager::FLayout> LoadLayoutFromConfig();
+		void SaveLayoutToConfig();
+
+		void TogglePlotsTabVisibility(ECheckBoxState InCheckboxState);
 #endif // WITH EDITOR
 
+		TSharedRef<SDockTab> CreateMixerSourcesTab(const FSpawnTabArgs& Args);
+		TSharedRef<SDockTab> CreatePlotsTab(const FSpawnTabArgs& Args);
+
+		void RegisterTabSpawners();
+		void UnregisterTabSpawners();
+
+		TSharedRef<FTabManager::FLayout> GetDefaultTabLayout();
+
 		using FPlotCurvePoint = SAudioCurveView::FCurvePoint;
 		// Map of source id to data point array 
 		using FPointDataPerCurveMap = TMap<int32, TArray<FPlotCurvePoint>>;
@@ -85,6 +100,8 @@ namespace UE::Audio::Insights
 		TSharedPtr<SCheckBox> MuteToggleButton;
 		TSharedPtr<SCheckBox> SoloToggleButton;
 
+		TSharedPtr<SCheckBox> PlotsButton;
+
 		// Curve points per timestamp per source id per column name 
 		TMap<FName, TSharedPtr<FPointDataPerCurveMap>> PlotWidgetCurveIdToPointDataMapPerColumn;
 		// SourceId to metadata for the corresponding curve
@@ -109,6 +126,9 @@ namespace UE::Audio::Insights
 		TArray<FName> SelectedPlotColumnNames;
 		TArray<TSharedPtr<SAudioCurveView>> PlotWidgets;
 
+		TSharedPtr<FTabManager> MixerSourcesTabManager;
+		TSharedPtr<FWorkspaceItem> MixerSourcesWorkspace;
+
 #if WITH_EDITOR
 		// State of the mute and solo buttons
 		ECheckBoxState MuteState = ECheckBoxState::Unchecked;
