[Audio Insights] Mixer Source tab: plot widget needs to be created before restoring the tabs layout (for the case when Mixer Source tab starts with the plots tab hidden) + old missing TSharedPtr check

diff --git a/Engine/Plugins/AudioInsights/Source/AudioInsights/Private/Views/MixerSourceDashboardViewFactory.cpp b/Engine/Plugins/AudioInsights/Source/AudioInsights/Private/Views/MixerSourceDashboardViewFactory.cpp
--- a/Engine/Plugins/AudioInsights/Source/AudioInsights/Private/Views/MixerSourceDashboardViewFactory.cpp
+++ b/Engine/Plugins/AudioInsights/Source/AudioInsights/Private/Views/MixerSourceDashboardViewFactory.cpp
@@ -390,7 +390,10 @@ namespace UE::Audio::Insights
 			PointDataPerCurveMap->Empty();
 		}
 
-		PlotWidgetMetadataPerCurve->Empty();
+		if (PlotWidgetMetadataPerCurve.IsValid())
+		{
+			PlotWidgetMetadataPerCurve->Empty();
+		}
 
 		BeginTimestamp = TNumericLimits<double>::Max();
 		CurrentTimestamp = 0;
@@ -948,7 +951,7 @@ namespace UE::Audio::Insights
 		return SNew(SDockTab)
 			.Clipping(EWidgetClipping::ClipToBounds)
 			[
-				MakePlotsWidget()
+				PlotsWidget ? PlotsWidget.ToSharedRef() : SNullWidget::NullWidget
 			];
 	}
 
@@ -1033,6 +1036,8 @@ namespace UE::Audio::Insights
 		FAudioInsightsComponent::OnTabSpawn.AddSP(this, &FMixerSourceDashboardViewFactory::OnAudioInsightsComponentTabSpawn);
 #endif // WITH_EDITOR
 
+		PlotsWidget = MakePlotsWidget();
+
 		const TSharedRef<SDockTab> DockTab = SNew(SDockTab);
 		MixerSourcesTabManager = FGlobalTabmanager::Get()->NewTabManager(DockTab);
 
diff --git a/Engine/Plugins/AudioInsights/Source/AudioInsights/Public/Views/MixerSourceDashboardViewFactory.h b/Engine/Plugins/AudioInsights/Source/AudioInsights/Public/Views/MixerSourceDashboardViewFactory.h
--- a/Engine/Plugins/AudioInsights/Source/AudioInsights/Public/Views/MixerSourceDashboardViewFactory.h
+++ b/Engine/Plugins/AudioInsights/Source/AudioInsights/Public/Views/MixerSourceDashboardViewFactory.h
@@ -126,6 +126,8 @@ namespace UE::Audio::Insights
 		TArray<FName> SelectedPlotColumnNames;
 		TArray<TSharedPtr<SAudioCurveView>> PlotWidgets;
 
+		TSharedPtr<SWidget> PlotsWidget;
+
 		TSharedPtr<FTabManager> MixerSourcesTabManager;
 		TSharedPtr<FWorkspaceItem> MixerSourcesWorkspace;
 
