[Audio Insights] Fix for crash on reopening the editor if any Audio Insights tab have been undocked

diff --git a/Engine/Plugins/AudioInsights/Source/AudioInsights/Private/AudioInsightsDashboardFactory.cpp b/Engine/Plugins/AudioInsights/Source/AudioInsights/Private/AudioInsightsDashboardFactory.cpp
--- a/Engine/Plugins/AudioInsights/Source/AudioInsights/Private/AudioInsightsDashboardFactory.cpp
+++ b/Engine/Plugins/AudioInsights/Source/AudioInsights/Private/AudioInsightsDashboardFactory.cpp
@@ -52,7 +52,7 @@ namespace UE::Audio::Insights
 			]
 			+ SVerticalBox::Slot()
 			[
-				DashboardTabManager->RestoreFrom(TabLayout->AsShared(), Args.GetOwnerWindow()).ToSharedRef()
+				DashboardTabManager->RestoreFrom(TabLayout->AsShared(), TSharedPtr<SWindow>()).ToSharedRef()
 			];
 
 		DockTab->SetContent(TabContent);
diff --git a/Engine/Plugins/AudioInsights/Source/AudioInsightsEditor/Private/AudioInsightsEditorDashboardFactory.cpp b/Engine/Plugins/AudioInsights/Source/AudioInsightsEditor/Private/AudioInsightsEditorDashboardFactory.cpp
--- a/Engine/Plugins/AudioInsights/Source/AudioInsightsEditor/Private/AudioInsightsEditorDashboardFactory.cpp
+++ b/Engine/Plugins/AudioInsights/Source/AudioInsightsEditor/Private/AudioInsightsEditorDashboardFactory.cpp
@@ -246,7 +246,7 @@ namespace UE::Audio::Insights
 			]
 			+ SVerticalBox::Slot()
 			[
-				DashboardTabManager->RestoreFrom(TabLayout, Args.GetOwnerWindow()).ToSharedRef()
+				DashboardTabManager->RestoreFrom(TabLayout, TSharedPtr<SWindow>()).ToSharedRef()
 			];
 
 		DockTab->SetContent(TabContent);
