[Audio Insights] Update Audio Insights data when scrubbing in Unreal Insights (part 2): Added TimingViewExtender (to be able to get notified when scrubbing unreal insights timeline view)

diff --git a/Engine/Plugins/AudioInsights/Source/AudioInsights/Private/AudioInsightsModule.cpp b/Engine/Plugins/AudioInsights/Source/AudioInsights/Private/AudioInsightsModule.cpp
--- a/Engine/Plugins/AudioInsights/Source/AudioInsights/Private/AudioInsightsModule.cpp
+++ b/Engine/Plugins/AudioInsights/Source/AudioInsights/Private/AudioInsightsModule.cpp
@@ -30,13 +30,16 @@ namespace UE::Audio::Insights
 		// Don't run providers in any commandlet to avoid additional, unnecessary overhead as audio insights is dormant.
 		if (!IsRunningCommandlet())
 		{
-			IModularFeatures::Get().RegisterModularFeature(TraceServices::ModuleFeatureName, &TraceModule);
+			TraceModule = MakeUnique<FTraceModule>();
+			IModularFeatures::Get().RegisterModularFeature(TraceServices::ModuleFeatureName, TraceModule.Get());
 
 			DashboardFactory = MakeShared<FDashboardFactory>();
 			
 			FDashboardAssetCommands::Register();
 
 #if !WITH_EDITOR
+			IModularFeatures::Get().RegisterModularFeature(UE::Insights::Timing::TimingViewExtenderFeatureName, &AudioInsightsTimingViewExtender);
+
 			DashboardFactory->RegisterViewFactory(MakeShared<FLogDashboardViewFactory>());
 			DashboardFactory->RegisterViewFactory(MakeShared<FMixerSourceDashboardViewFactory>());
 			DashboardFactory->RegisterViewFactory(MakeShared<FVirtualLoopDashboardViewFactory>());
@@ -70,13 +73,15 @@ namespace UE::Audio::Insights
 			UnrealInsightsModule.UnregisterComponent(AudioInsightsComponent);
 
 			AudioInsightsComponent.Reset();
+
+			IModularFeatures::Get().UnregisterModularFeature(TraceServices::ModuleFeatureName, &AudioInsightsTimingViewExtender);
 #endif // !WITH_EDITOR
 
 			FDashboardAssetCommands::Unregister();
 
 			DashboardFactory.Reset();
 
-			IModularFeatures::Get().UnregisterModularFeature(TraceServices::ModuleFeatureName, &TraceModule);
+			IModularFeatures::Get().UnregisterModularFeature(TraceServices::ModuleFeatureName, TraceModule.Get());
 		}
 	}
 
@@ -112,7 +117,7 @@ namespace UE::Audio::Insights
 
 	IAudioInsightsTraceModule& FAudioInsightsModule::GetTraceModule()
 	{
-		return TraceModule;
+		return *TraceModule;
 	}
 
 	TSharedRef<SDockTab> FAudioInsightsModule::CreateDashboardTabWidget(const FSpawnTabArgs& Args)
diff --git a/Engine/Plugins/AudioInsights/Source/AudioInsights/Private/AudioInsightsModule.h b/Engine/Plugins/AudioInsights/Source/AudioInsights/Private/AudioInsightsModule.h
--- a/Engine/Plugins/AudioInsights/Source/AudioInsights/Private/AudioInsightsModule.h
+++ b/Engine/Plugins/AudioInsights/Source/AudioInsights/Private/AudioInsightsModule.h
@@ -7,9 +7,13 @@
 #include "IAudioInsightsModule.h"
 #include "Insights/IUnrealInsightsModule.h"
 #include "Templates/SharedPointer.h"
+#include "Templates/UniquePtr.h"
 #include "Views/DashboardViewFactory.h"
 #include "Widgets/Docking/SDockTab.h"
 
+#if !WITH_EDITOR
+#include "AudioInsightsTimingViewExtender.h"
+#endif // !WITH_EDITOR
 
 namespace UE::Audio::Insights
 {
@@ -29,6 +33,10 @@ namespace UE::Audio::Insights
 		static FAudioInsightsModule& GetChecked();
 		virtual IAudioInsightsTraceModule& GetTraceModule() override;
 
+#if !WITH_EDITOR
+		FAudioInsightsTimingViewExtender& GetTimingViewExtender() { return AudioInsightsTimingViewExtender; };
+#endif // !WITH_EDITOR
+
 		TSharedRef<FDashboardFactory> GetDashboardFactory();
 		const TSharedRef<FDashboardFactory> GetDashboardFactory() const;
 
@@ -36,10 +44,11 @@ namespace UE::Audio::Insights
 
 	private:
 		TSharedPtr<FDashboardFactory> DashboardFactory;
-		FTraceModule TraceModule;
+		TUniquePtr<FTraceModule> TraceModule;
 
 #if !WITH_EDITOR
 		TSharedPtr<IInsightsComponent> AudioInsightsComponent;
+		FAudioInsightsTimingViewExtender AudioInsightsTimingViewExtender;
 #endif // !WITH_EDITOR
 	};
 } // namespace UE::Audio::Insights
diff --git a/Engine/Plugins/AudioInsights/Source/AudioInsights/Private/AudioInsightsTimingViewExtender.cpp b/Engine/Plugins/AudioInsights/Source/AudioInsights/Private/AudioInsightsTimingViewExtender.cpp
new file mode 100644
--- /dev/null
+++ b/Engine/Plugins/AudioInsights/Source/AudioInsights/Private/AudioInsightsTimingViewExtender.cpp
@@ -0,0 +1,40 @@
+#include "AudioInsightsTimingViewExtender.h"
+
+#include "Insights/ITimingViewSession.h"
+#include "TraceServices/Model/AnalysisSession.h"
+
+namespace UE::Audio::Insights
+{
+	void FAudioInsightsTimingViewExtender::OnBeginSession(UE::Insights::Timing::ITimingViewSession& InSession)
+	{
+		InSession.OnTimeMarkerChanged().AddRaw(this, &FAudioInsightsTimingViewExtender::OnTimeMarkerChanged);
+	}
+
+	void FAudioInsightsTimingViewExtender::OnEndSession(UE::Insights::Timing::ITimingViewSession& InSession)
+	{
+		InSession.OnTimeMarkerChanged().RemoveAll(this);
+	}
+
+	void FAudioInsightsTimingViewExtender::Tick(UE::Insights::Timing::ITimingViewSession& InSession, const TraceServices::IAnalysisSession& InAnalysisSession)
+	{
+		AnalysisSession = &InAnalysisSession;
+	}
+
+	void FAudioInsightsTimingViewExtender::OnTimeMarkerChanged(UE::Insights::Timing::ETimeChangedFlags InFlags, double InTimeMarker)
+	{
+		if (AnalysisSession == nullptr)
+		{
+			return;
+		}
+		
+		TraceServices::FAnalysisSessionReadScope SessionReadScope(*AnalysisSession);
+
+		const bool bIsValidTime = InTimeMarker >= 0.0 && InTimeMarker <= AnalysisSession->GetDurationSeconds();
+
+		// Scrub Audio Insights when all data is loaded in non-live sessions
+		if (AnalysisSession->IsAnalysisComplete() && bIsValidTime)
+		{
+			OnTimingViewTimeMarkerChanged.Broadcast(InTimeMarker);
+		}
+	}
+} // namespace UE::Audio::Insights
diff --git a/Engine/Plugins/AudioInsights/Source/AudioInsights/Private/AudioInsightsTimingViewExtender.h b/Engine/Plugins/AudioInsights/Source/AudioInsights/Private/AudioInsightsTimingViewExtender.h
new file mode 100644
--- /dev/null
+++ b/Engine/Plugins/AudioInsights/Source/AudioInsights/Private/AudioInsightsTimingViewExtender.h
@@ -0,0 +1,33 @@
+// Copyright Epic Games, Inc. All Rights Reserved.
+#pragma once
+
+#include "Delegates/Delegate.h"
+#include "Insights/ITimingViewExtender.h"
+
+namespace Insights
+{
+	class ITimingViewSession;
+	enum class ETimeChangedFlags : int32;
+}
+
+class IAnalysisSession;
+
+namespace UE::Audio::Insights
+{
+	class  FAudioInsightsTimingViewExtender : public UE::Insights::Timing::ITimingViewExtender
+	{
+	public:
+		// Insights::ITimingViewExtender interface
+		virtual void OnBeginSession(UE::Insights::Timing::ITimingViewSession& InSession) override;
+		virtual void OnEndSession(UE::Insights::Timing::ITimingViewSession& InSession) override;
+		virtual void Tick(UE::Insights::Timing::ITimingViewSession& InSession, const TraceServices::IAnalysisSession& InAnalysisSession) override;
+
+		DECLARE_MULTICAST_DELEGATE_OneParam(FOnTimingViewTimeMarkerChanged, double /*TimeMarker*/);
+		FOnTimingViewTimeMarkerChanged OnTimingViewTimeMarkerChanged;
+
+	private:
+		void OnTimeMarkerChanged(UE::Insights::Timing::ETimeChangedFlags InFlags, double InTimeMarker);
+
+		const TraceServices::IAnalysisSession* AnalysisSession = nullptr;
+	};
+} // namespace UE::Audio::Insights
diff --git a/Engine/Plugins/AudioInsights/Source/AudioInsights/Private/AudioInsightsTraceProviderBase.cpp b/Engine/Plugins/AudioInsights/Source/AudioInsights/Private/AudioInsightsTraceProviderBase.cpp
--- a/Engine/Plugins/AudioInsights/Source/AudioInsights/Private/AudioInsightsTraceProviderBase.cpp
+++ b/Engine/Plugins/AudioInsights/Source/AudioInsights/Private/AudioInsightsTraceProviderBase.cpp
@@ -2,12 +2,34 @@
 
 #include "AudioInsightsTraceProviderBase.h"
 
+#if !WITH_EDITOR
+#include "AudioInsightsModule.h"
+#include "AudioInsightsTimingViewExtender.h"
+#include "Modules/ModuleManager.h"
+#endif // !WITH_EDITOR
 
 namespace UE::Audio::Insights
 {
 	FTraceProviderBase::FTraceProviderBase(FName InName)
 		: Name(InName)
 	{
+#if !WITH_EDITOR
+		FAudioInsightsModule& AudioInsightsModule = FAudioInsightsModule::GetChecked();
+		FAudioInsightsTimingViewExtender& AudioInsightsTimingViewExtender = AudioInsightsModule.GetTimingViewExtender();
+		AudioInsightsTimingViewExtender.OnTimingViewTimeMarkerChanged.AddRaw(this, &FTraceProviderBase::OnTimingViewTimeMarkerChanged);
+#endif // !WITH_EDITOR
+	}
+
+	FTraceProviderBase::~FTraceProviderBase()
+	{
+#if !WITH_EDITOR
+		if (FModuleManager::Get().IsModuleLoaded("AudioInsights"))
+		{
+			FAudioInsightsModule& AudioInsightsModule = FAudioInsightsModule::GetChecked();
+			FAudioInsightsTimingViewExtender& AudioInsightsTimingViewExtender = AudioInsightsModule.GetTimingViewExtender();
+			AudioInsightsTimingViewExtender.OnTimingViewTimeMarkerChanged.RemoveAll(this);
+		}
+#endif // !WITH_EDITOR
 	}
 
 	FName FTraceProviderBase::GetName() const
diff --git a/Engine/Plugins/AudioInsights/Source/AudioInsights/Private/Providers/MixerSourceTraceProvider.cpp b/Engine/Plugins/AudioInsights/Source/AudioInsights/Private/Providers/MixerSourceTraceProvider.cpp
--- a/Engine/Plugins/AudioInsights/Source/AudioInsights/Private/Providers/MixerSourceTraceProvider.cpp
+++ b/Engine/Plugins/AudioInsights/Source/AudioInsights/Private/Providers/MixerSourceTraceProvider.cpp
@@ -14,6 +14,13 @@ namespace UE::Audio::Insights
 		return "MixerSourceProvider";
 	}
 
+#if !WITH_EDITOR
+	void FMixerSourceTraceProvider::OnTimingViewTimeMarkerChanged(double TimeMarker)
+	{
+		// TODO alex.perez: to simplify code review, the implementation of this method will be submitted after this CL (in part 3)
+	}
+#endif // !WITH_EDITOR
+
 	bool FMixerSourceTraceProvider::ProcessMessages()
 	{
 		auto BumpEntryFunc = [this](const FMixerSourceMessageBase& Msg)
diff --git a/Engine/Plugins/AudioInsights/Source/AudioInsights/Private/Providers/MixerSourceTraceProvider.h b/Engine/Plugins/AudioInsights/Source/AudioInsights/Private/Providers/MixerSourceTraceProvider.h
--- a/Engine/Plugins/AudioInsights/Source/AudioInsights/Private/Providers/MixerSourceTraceProvider.h
+++ b/Engine/Plugins/AudioInsights/Source/AudioInsights/Private/Providers/MixerSourceTraceProvider.h
@@ -28,6 +28,10 @@ namespace UE::Audio::Insights
 		static FName GetName_Static();
 
 	private:
+#if !WITH_EDITOR
+		virtual void OnTimingViewTimeMarkerChanged(double TimeMarker) override;
+#endif // !WITH_EDITOR
+
 		FMixerSourceMessages TraceMessages;
 
 		TSet<FTraceChannelHandle> Channels;
diff --git a/Engine/Plugins/AudioInsights/Source/AudioInsights/Private/Providers/VirtualLoopTraceProvider.cpp b/Engine/Plugins/AudioInsights/Source/AudioInsights/Private/Providers/VirtualLoopTraceProvider.cpp
--- a/Engine/Plugins/AudioInsights/Source/AudioInsights/Private/Providers/VirtualLoopTraceProvider.cpp
+++ b/Engine/Plugins/AudioInsights/Source/AudioInsights/Private/Providers/VirtualLoopTraceProvider.cpp
@@ -13,6 +13,13 @@ namespace UE::Audio::Insights
 		return "AudioVirtualLoopProvider";
 	}
 
+#if !WITH_EDITOR
+	void FVirtualLoopTraceProvider::OnTimingViewTimeMarkerChanged(double TimeMarker)
+	{
+		// TODO alex.perez: to simplify code review, the implementation of this method will be submitted after this CL (in part 3)
+	}
+#endif // !WITH_EDITOR
+
 	bool FVirtualLoopTraceProvider::ProcessMessages()
 	{
 		auto RemoveEntryFunc = [this](const FVirtualLoopRealizeMessage& Msg, TSharedPtr<FVirtualLoopDashboardEntry>* OutEntry)
diff --git a/Engine/Plugins/AudioInsights/Source/AudioInsights/Private/Providers/VirtualLoopTraceProvider.h b/Engine/Plugins/AudioInsights/Source/AudioInsights/Private/Providers/VirtualLoopTraceProvider.h
--- a/Engine/Plugins/AudioInsights/Source/AudioInsights/Private/Providers/VirtualLoopTraceProvider.h
+++ b/Engine/Plugins/AudioInsights/Source/AudioInsights/Private/Providers/VirtualLoopTraceProvider.h
@@ -25,6 +25,10 @@ namespace UE::Audio::Insights
 		static FName GetName_Static();
 
 	private:
+#if !WITH_EDITOR
+		virtual void OnTimingViewTimeMarkerChanged(double TimeMarker) override;
+#endif // !WITH_EDITOR
+
 		virtual bool ProcessMessages() override;
 
 		FVirtualLoopMessages TraceMessages;
diff --git a/Engine/Plugins/AudioInsights/Source/AudioInsights/Public/AudioInsightsTraceProviderBase.h b/Engine/Plugins/AudioInsights/Source/AudioInsights/Public/AudioInsightsTraceProviderBase.h
--- a/Engine/Plugins/AudioInsights/Source/AudioInsights/Public/AudioInsightsTraceProviderBase.h
+++ b/Engine/Plugins/AudioInsights/Source/AudioInsights/Public/AudioInsightsTraceProviderBase.h
@@ -19,8 +19,10 @@ namespace UE::Audio::Insights
 	class AUDIOINSIGHTS_API FTraceProviderBase : public TraceServices::IProvider, public TraceServices::IEditableProvider
 	{
 	public:
+		FTraceProviderBase() = delete;
 		explicit FTraceProviderBase(FName InName);
-		virtual ~FTraceProviderBase() = default;
+
+		virtual ~FTraceProviderBase();
 
 		virtual Trace::IAnalyzer* ConstructAnalyzer(TraceServices::IAnalysisSession& InSession) = 0;
 		FName GetName() const;
@@ -47,6 +49,10 @@ namespace UE::Audio::Insights
 			return GetLastMessageId() == LastUpdateId;
 		}
 
+#if !WITH_EDITOR
+		virtual void OnTimingViewTimeMarkerChanged(double TimeMarker) {};
+#endif // !WITH_EDITOR
+
 	protected:
 		class AUDIOINSIGHTS_API FTraceAnalyzerBase : public Trace::IAnalyzer
 		{
