[Audio Insights] Update Audio Insights data when scrubbing in Unreal Insights (part 3): Being able to retrieve cached information in Audio Insights (in mixer source and virtual loops tabs) while scrubbing in Unreal Insights timeline view.

diff --git a/Engine/Plugins/AudioInsights/Source/AudioInsights/Private/AudioInsightsComponent.cpp b/Engine/Plugins/AudioInsights/Source/AudioInsights/Private/AudioInsightsComponent.cpp
--- a/Engine/Plugins/AudioInsights/Source/AudioInsights/Private/AudioInsightsComponent.cpp
+++ b/Engine/Plugins/AudioInsights/Source/AudioInsights/Private/AudioInsightsComponent.cpp
@@ -78,7 +78,7 @@ namespace UE::Audio::Insights
 				FOnSpawnTab::CreateRaw(this,  &FAudioInsightsComponent::SpawnTab), 
 				FCanSpawnTab::CreateRaw(this, &FAudioInsightsComponent::CanSpawnTab))
 				.SetDisplayName(Config.TabLabel.IsSet()   ? Config.TabLabel.GetValue()   : LOCTEXT("AudioInsights_TabTitle", "Audio Insights"))
-				.SetTooltipText(Config.TabTooltip.IsSet() ? Config.TabTooltip.GetValue() : LOCTEXT("AudioInsights_TooltipText", "Open the Audio Insights tab (Only available for standalone live traces)."))
+				.SetTooltipText(Config.TabTooltip.IsSet() ? Config.TabTooltip.GetValue() : LOCTEXT("AudioInsights_TooltipText", "Open the Audio Insights tab (Only available for standalone game traces)."))
 				.SetIcon(Config.TabIcon.IsSet() ? Config.TabIcon.GetValue() : FSlateStyle::Get().CreateIcon("AudioInsights.Icon.Submix"));
 
 			const TSharedRef<FWorkspaceItem>* FoundWorkspace = FGlobalTabmanager::Get()->GetLocalWorkspaceMenuRoot()->GetChildItems().FindByPredicate(
@@ -115,8 +115,8 @@ namespace UE::Audio::Insights
 
 	bool FAudioInsightsComponent::Tick(float DeltaTime)
 	{
-		// Audio Insights will be available only if there is an active standalone (non-editor) live session
-		if (bCanCheckForActiveSession && !bCanSpawnTab)
+		// Audio Insights will be available in non-editor file traces or if there is an active standalone game live session
+		if (!bCanSpawnTab)
 		{
 			IUnrealInsightsModule& UnrealInsightsModule = FModuleManager::LoadModuleChecked<IUnrealInsightsModule>("TraceInsights");
 
@@ -125,29 +125,25 @@ namespace UE::Audio::Insights
 			{
 				TraceServices::FAnalysisSessionReadScope SessionReadScope(*Session.Get());
 
-				if (!Session->IsAnalysisComplete())
+				// Set bIsLiveSession
+				const UE::Trace::FStoreClient* StoreClient = UnrealInsightsModule.GetStoreClient();
+				const UE::Trace::FStoreClient::FSessionInfo* SessionInfo = StoreClient ? StoreClient->GetSessionInfoByTraceId(Session->GetTraceId()) : nullptr;
+
+				bIsLiveSession = !Session->IsAnalysisComplete() && StoreClient != nullptr && SessionInfo != nullptr;
+
+				// Set bIsEditorTrace
+				const TraceServices::IDiagnosticsProvider* DiagnosticsProvider = TraceServices::ReadDiagnosticsProvider(*Session.Get());
+				if (DiagnosticsProvider && DiagnosticsProvider->IsSessionInfoAvailable())
 				{
-					if (UE::Trace::FStoreClient* StoreClient = UnrealInsightsModule.GetStoreClient())
-					{
-						const UE::Trace::FStoreClient::FSessionInfo* StoreClientSessionInfo = StoreClient->GetSessionInfoByTraceId(Session->GetTraceId());
-						if (StoreClientSessionInfo)
-						{
-							const TraceServices::IDiagnosticsProvider* DiagnosticsProvider = TraceServices::ReadDiagnosticsProvider(*Session.Get());
-							if (DiagnosticsProvider && DiagnosticsProvider->IsSessionInfoAvailable())
-							{
-								const TraceServices::FSessionInfo& TraceServicesSessionInfo = DiagnosticsProvider->GetSessionInfo();
-
-								if (TraceServicesSessionInfo.TargetType != EBuildTargetType::Editor)
-								{
-									bCanSpawnTab = true;
-								}
-							}
-						}
-					}
+					const TraceServices::FSessionInfo& TraceServicesSessionInfo = DiagnosticsProvider->GetSessionInfo();
+
+					bIsEditorTrace = TraceServicesSessionInfo.TargetType == EBuildTargetType::Editor;
 				}
-				else
+
+				// Allow to spawn tab if non-editor trace
+				if (!bIsEditorTrace)
 				{
-					bCanCheckForActiveSession = false;
+					bCanSpawnTab = true;
 				}
 			}
 		}
diff --git a/Engine/Plugins/AudioInsights/Source/AudioInsights/Private/AudioInsightsComponent.h b/Engine/Plugins/AudioInsights/Source/AudioInsights/Private/AudioInsightsComponent.h
--- a/Engine/Plugins/AudioInsights/Source/AudioInsights/Private/AudioInsightsComponent.h
+++ b/Engine/Plugins/AudioInsights/Source/AudioInsights/Private/AudioInsightsComponent.h
@@ -24,6 +24,9 @@ namespace UE::Audio::Insights
 		virtual void RegisterMajorTabs(IUnrealInsightsModule& InsightsModule) override;
 		virtual void UnregisterMajorTabs() override;
 
+		bool GetIsLiveSession() const { return bIsLiveSession; }
+		bool GetIsEditorTrace() const { return bIsEditorTrace; }
+
 		DECLARE_MULTICAST_DELEGATE(FOnTabSpawn);
 		inline static FOnTabSpawn OnTabSpawn;
 
@@ -35,7 +38,9 @@ namespace UE::Audio::Insights
 
 		bool bIsInitialized = false;
 
-		bool bCanCheckForActiveSession = true;
+		bool bIsLiveSession = false;
+		bool bIsEditorTrace = false;
+
 		bool bCanSpawnTab = false;
 
 		FTickerDelegate OnTick;
diff --git a/Engine/Plugins/AudioInsights/Source/AudioInsights/Private/AudioInsightsModule.h b/Engine/Plugins/AudioInsights/Source/AudioInsights/Private/AudioInsightsModule.h
--- a/Engine/Plugins/AudioInsights/Source/AudioInsights/Private/AudioInsightsModule.h
+++ b/Engine/Plugins/AudioInsights/Source/AudioInsights/Private/AudioInsightsModule.h
@@ -17,6 +17,10 @@
 
 namespace UE::Audio::Insights
 {
+#if !WITH_EDITOR
+	class FAudioInsightsComponent;
+#endif // !WITH_EDITOR
+
 	class FAudioInsightsModule final : public IAudioInsightsModule
 	{
 	public:
@@ -34,6 +38,7 @@ namespace UE::Audio::Insights
 		virtual IAudioInsightsTraceModule& GetTraceModule() override;
 
 #if !WITH_EDITOR
+		TSharedPtr<FAudioInsightsComponent> GetAudioInsightsComponent() { return AudioInsightsComponent; };
 		FAudioInsightsTimingViewExtender& GetTimingViewExtender() { return AudioInsightsTimingViewExtender; };
 #endif // !WITH_EDITOR
 
@@ -47,7 +52,7 @@ namespace UE::Audio::Insights
 		TUniquePtr<FTraceModule> TraceModule;
 
 #if !WITH_EDITOR
-		TSharedPtr<IInsightsComponent> AudioInsightsComponent;
+		TSharedPtr<FAudioInsightsComponent> AudioInsightsComponent;
 		FAudioInsightsTimingViewExtender AudioInsightsTimingViewExtender;
 #endif // !WITH_EDITOR
 	};
diff --git a/Engine/Plugins/AudioInsights/Source/AudioInsights/Private/AudioInsightsTraceModule.cpp b/Engine/Plugins/AudioInsights/Source/AudioInsights/Private/AudioInsightsTraceModule.cpp
--- a/Engine/Plugins/AudioInsights/Source/AudioInsights/Private/AudioInsightsTraceModule.cpp
+++ b/Engine/Plugins/AudioInsights/Source/AudioInsights/Private/AudioInsightsTraceModule.cpp
@@ -71,6 +71,10 @@ namespace UE::Audio::Insights
 	{
 		for (const auto& [ProviderName, Provider] : TraceProviders)
 		{
+#if !WITH_EDITOR
+			Provider->InitSessionCachedMessages(InSession);
+#endif // !WITH_EDITOR
+
 			InSession.AddProvider(ProviderName, Provider, Provider);
 			InSession.AddAnalyzer(Provider->ConstructAnalyzer(InSession));
 		}
diff --git a/Engine/Plugins/AudioInsights/Source/AudioInsights/Private/Providers/MixerSourceTraceProvider.cpp b/Engine/Plugins/AudioInsights/Source/AudioInsights/Private/Providers/MixerSourceTraceProvider.cpp
--- a/Engine/Plugins/AudioInsights/Source/AudioInsights/Private/Providers/MixerSourceTraceProvider.cpp
+++ b/Engine/Plugins/AudioInsights/Source/AudioInsights/Private/Providers/MixerSourceTraceProvider.cpp
@@ -6,9 +6,34 @@
 #include "TraceServices/Model/AnalysisSession.h"
 #include "TraceServices/ModuleService.h"
 
+#if !WITH_EDITOR
+#include "Common/PagedArray.h"
+#endif // !WITH_EDITOR
 
 namespace UE::Audio::Insights
 {
+	namespace FMixerSourceTraceProviderPrivate
+	{
+#if !WITH_EDITOR
+		template<typename T>
+		const T* FindClosestMessageToTimestamp(const TraceServices::TPagedArray<T>& InCachedMessages, const double InTimeMarker, const uint32 InPlayOrder)
+		{
+			const int32 ClosestMessageToTimeStampIndex = TraceServices::PagedArrayAlgo::BinarySearchClosestBy(InCachedMessages, InTimeMarker, [](const T& Msg) { return Msg.Timestamp; });
+
+			// Iterate backwards from TimeMarker until we find the matching PlayOrder
+			for (auto It = InCachedMessages.GetIteratorFromItem(ClosestMessageToTimeStampIndex); It; --It)
+			{
+				if (It->PlayOrder == InPlayOrder)
+				{
+					return &(*It);
+				}
+			}
+
+			return nullptr;
+		}
+#endif // !WITH_EDITOR
+	}
+
 	FName FMixerSourceTraceProvider::GetName_Static()
 	{
 		return "MixerSourceProvider";
@@ -17,7 +42,211 @@ namespace UE::Audio::Insights
 #if !WITH_EDITOR
 	void FMixerSourceTraceProvider::OnTimingViewTimeMarkerChanged(double TimeMarker)
 	{
-		// TODO alex.perez: to simplify code review, the implementation of this method will be submitted after this CL (in part 3)
+		using namespace FMixerSourceTraceProviderPrivate;
+
+		if (!SessionCachedMessages.IsValid())
+		{
+			return;
+		}
+
+		DeviceDataMap.Empty();
+
+		// Collect all the start messages registered until this point in time 
+		for (const FMixerSourceStartMessage& StartCachedMessage : SessionCachedMessages->StartCachedMessages)
+		{
+			if (StartCachedMessage.Timestamp > TimeMarker)
+			{
+				break;
+			}
+
+			UpdateDeviceEntry(StartCachedMessage.DeviceId, StartCachedMessage.PlayOrder, [&StartCachedMessage](TSharedPtr<FMixerSourceDashboardEntry>& Entry)
+			{
+				if (!Entry.IsValid())
+				{
+					Entry = MakeShared<FMixerSourceDashboardEntry>();
+					Entry->DeviceId  = StartCachedMessage.DeviceId;
+					Entry->PlayOrder = StartCachedMessage.PlayOrder;
+				}
+				Entry->Timestamp = StartCachedMessage.Timestamp;
+
+				Entry->Name        = *StartCachedMessage.Name;
+				Entry->ComponentId = StartCachedMessage.ComponentId;
+				Entry->SourceId    = StartCachedMessage.SourceId;
+			});
+		}
+
+		// Selectively remove start messages collected in the step above by knowing which sounds were stopped.
+		// With this we will know what are the active sounds at this point in time.
+		for (const FMixerSourceStopMessage& StopCachedMessage : SessionCachedMessages->StopCachedMessages)
+		{
+			if (StopCachedMessage.Timestamp > TimeMarker)
+			{
+				break;
+			}
+
+			auto* OutEntry = FindDeviceEntry(StopCachedMessage.DeviceId, StopCachedMessage.PlayOrder);
+
+			if (OutEntry && (*OutEntry)->Timestamp < StopCachedMessage.Timestamp)
+			{
+				RemoveDeviceEntry(StopCachedMessage.DeviceId, StopCachedMessage.PlayOrder);
+			}
+		}
+
+		// For now we only retrieve information from AudioDeviceId 1 (main device in standalone games)
+		const FDeviceData* DeviceData = DeviceDataMap.Find(1);
+		if (DeviceData)
+		{
+			// Collect messages (volume, pitch, etc.) from active sounds (based on active sounds's PlayOrder)
+			struct CachedEntryInfo
+			{
+				FMixerSourceVolumeMessage VolumeMessage;
+				FMixerSourcePitchMessage PitchMessage;
+				FMixerSourceLPFFreqMessage LPFFreqMessage;
+				FMixerSourceHPFFreqMessage HPFFreqMessage;
+				FMixerSourceEnvelopeMessage EnvelopeMessage;
+				FMixerSourceDistanceAttenuationMessage DistanceAttenuationMessage;
+			};
+
+			TArray<uint32> PlayOrderArray;
+			(*DeviceData).GenerateKeyArray(PlayOrderArray);
+
+			TArray<CachedEntryInfo> CachedEntryInfos;
+			CachedEntryInfos.SetNumUninitialized(PlayOrderArray.Num());
+
+			// Using ParallelFor to speed-up the cached messages retrieval, using a traditional for loop is unacceptably slower, specially in large traces.
+			ParallelFor(PlayOrderArray.Num(), 
+			[&PlayOrderArray, &CachedEntryInfos, TimeMarker, this](const int32 Index)
+			{
+				const uint32 PlayOrder = PlayOrderArray[Index];
+
+				// Volume
+				const FMixerSourceVolumeMessage* FoundVolumeCachedMessage = FindClosestMessageToTimestamp(SessionCachedMessages->VolumeCachedMessages, TimeMarker, PlayOrder);
+				if (FoundVolumeCachedMessage)
+				{
+					CachedEntryInfos[Index].VolumeMessage = *FoundVolumeCachedMessage;
+				}
+
+				// Pitch
+				const FMixerSourcePitchMessage* FoundPitchCachedMessage = FindClosestMessageToTimestamp(SessionCachedMessages->PitchCachedMessages, TimeMarker, PlayOrder);
+				if (FoundPitchCachedMessage)
+				{
+					CachedEntryInfos[Index].PitchMessage = *FoundPitchCachedMessage;
+				}
+
+				// LPF
+				const FMixerSourceLPFFreqMessage* FoundLPFFreqCachedMessage = FindClosestMessageToTimestamp(SessionCachedMessages->LPFFreqCachedMessages, TimeMarker, PlayOrder);
+				if (FoundLPFFreqCachedMessage)
+				{
+					CachedEntryInfos[Index].LPFFreqMessage = *FoundLPFFreqCachedMessage;
+				}
+
+				// HPF
+				const FMixerSourceHPFFreqMessage* FoundHPFFreqCachedMessage = FindClosestMessageToTimestamp(SessionCachedMessages->HPFFreqCachedMessages, TimeMarker, PlayOrder);
+				if (FoundHPFFreqCachedMessage)
+				{
+					CachedEntryInfos[Index].HPFFreqMessage = *FoundHPFFreqCachedMessage;
+				}
+
+				// Envelope
+				const FMixerSourceEnvelopeMessage* FoundEnvelopeCachedMessage = FindClosestMessageToTimestamp(SessionCachedMessages->EnvelopeCachedMessages, TimeMarker, PlayOrder);
+				if (FoundEnvelopeCachedMessage)
+				{
+					CachedEntryInfos[Index].EnvelopeMessage = *FoundEnvelopeCachedMessage;
+				}
+
+				// Distance Attenuation
+				const FMixerSourceDistanceAttenuationMessage* FoundDistanceAttenuationCachedMessage = FindClosestMessageToTimestamp(SessionCachedMessages->DistanceAttenuationCachedMessages, TimeMarker, PlayOrder);
+				if (FoundDistanceAttenuationCachedMessage)
+				{
+					CachedEntryInfos[Index].DistanceAttenuationMessage = *FoundDistanceAttenuationCachedMessage;
+				}
+			});
+
+			// Update the device entries with the collected info
+			for (const CachedEntryInfo& CachedEntryInfo : CachedEntryInfos)
+			{
+				UpdateDeviceEntry(CachedEntryInfo.VolumeMessage.DeviceId, CachedEntryInfo.VolumeMessage.PlayOrder, [&CachedEntryInfo](TSharedPtr<FMixerSourceDashboardEntry>& Entry)
+				{
+					if (!Entry.IsValid())
+					{
+						Entry = MakeShared<FMixerSourceDashboardEntry>();
+						Entry->DeviceId  = CachedEntryInfo.VolumeMessage.DeviceId;
+						Entry->PlayOrder = CachedEntryInfo.VolumeMessage.PlayOrder;
+					}
+
+					Entry->Timestamp = CachedEntryInfo.VolumeMessage.Timestamp;
+					Entry->VolumeDataPoints.Push({ CachedEntryInfo.VolumeMessage.Timestamp, CachedEntryInfo.VolumeMessage.Volume });
+				});
+
+				UpdateDeviceEntry(CachedEntryInfo.PitchMessage.DeviceId, CachedEntryInfo.PitchMessage.PlayOrder, [&CachedEntryInfo](TSharedPtr<FMixerSourceDashboardEntry>& Entry)
+				{
+					if (!Entry.IsValid())
+					{
+						Entry = MakeShared<FMixerSourceDashboardEntry>();
+						Entry->DeviceId  = CachedEntryInfo.PitchMessage.DeviceId;
+						Entry->PlayOrder = CachedEntryInfo.PitchMessage.PlayOrder;
+					}
+
+					Entry->Timestamp = CachedEntryInfo.PitchMessage.Timestamp;
+					Entry->PitchDataPoints.Push({ CachedEntryInfo.PitchMessage.Timestamp, CachedEntryInfo.PitchMessage.Pitch });
+				});
+
+				UpdateDeviceEntry(CachedEntryInfo.LPFFreqMessage.DeviceId, CachedEntryInfo.LPFFreqMessage.PlayOrder, [&CachedEntryInfo](TSharedPtr<FMixerSourceDashboardEntry>& Entry)
+				{
+					if (!Entry.IsValid())
+					{
+						Entry = MakeShared<FMixerSourceDashboardEntry>();
+						Entry->DeviceId  = CachedEntryInfo.LPFFreqMessage.DeviceId;
+						Entry->PlayOrder = CachedEntryInfo.LPFFreqMessage.PlayOrder;
+					}
+
+					Entry->Timestamp = CachedEntryInfo.LPFFreqMessage.Timestamp;
+					Entry->LPFFreqDataPoints.Push({ CachedEntryInfo.LPFFreqMessage.Timestamp, CachedEntryInfo.LPFFreqMessage.LPFFrequency });
+				});
+
+				UpdateDeviceEntry(CachedEntryInfo.HPFFreqMessage.DeviceId, CachedEntryInfo.HPFFreqMessage.PlayOrder, [&CachedEntryInfo](TSharedPtr<FMixerSourceDashboardEntry>& Entry)
+				{
+					if (!Entry.IsValid())
+					{
+						Entry = MakeShared<FMixerSourceDashboardEntry>();
+						Entry->DeviceId  = CachedEntryInfo.HPFFreqMessage.DeviceId;
+						Entry->PlayOrder = CachedEntryInfo.HPFFreqMessage.PlayOrder;
+					}
+
+					Entry->Timestamp = CachedEntryInfo.HPFFreqMessage.Timestamp;
+					Entry->HPFFreqDataPoints.Push({ CachedEntryInfo.HPFFreqMessage.Timestamp, CachedEntryInfo.HPFFreqMessage.HPFFrequency });
+				});
+
+				UpdateDeviceEntry(CachedEntryInfo.EnvelopeMessage.DeviceId, CachedEntryInfo.EnvelopeMessage.PlayOrder, [&CachedEntryInfo](TSharedPtr<FMixerSourceDashboardEntry>& Entry)
+				{
+					if (!Entry.IsValid())
+					{
+						Entry = MakeShared<FMixerSourceDashboardEntry>();
+						Entry->DeviceId  = CachedEntryInfo.EnvelopeMessage.DeviceId;
+						Entry->PlayOrder = CachedEntryInfo.EnvelopeMessage.PlayOrder;
+					}
+
+					Entry->Timestamp = CachedEntryInfo.EnvelopeMessage.Timestamp;
+					Entry->EnvelopeDataPoints.Push({ CachedEntryInfo.EnvelopeMessage.Timestamp, CachedEntryInfo.EnvelopeMessage.Envelope });
+				});
+
+				UpdateDeviceEntry(CachedEntryInfo.DistanceAttenuationMessage.DeviceId, CachedEntryInfo.DistanceAttenuationMessage.PlayOrder, [&CachedEntryInfo](TSharedPtr<FMixerSourceDashboardEntry>& Entry)
+				{
+					if (!Entry.IsValid())
+					{
+						Entry = MakeShared<FMixerSourceDashboardEntry>();
+						Entry->DeviceId  = CachedEntryInfo.DistanceAttenuationMessage.DeviceId;
+						Entry->PlayOrder = CachedEntryInfo.DistanceAttenuationMessage.PlayOrder;
+					}
+
+					Entry->Timestamp = CachedEntryInfo.DistanceAttenuationMessage.Timestamp;
+					Entry->DistanceAttenuationDataPoints.Push({ CachedEntryInfo.DistanceAttenuationMessage.Timestamp, CachedEntryInfo.DistanceAttenuationMessage.DistanceAttenuation });
+				});
+			}
+		}
+
+		// Call parent method to update LastMessageId
+		FTraceProviderBase::OnTimingViewTimeMarkerChanged(TimeMarker);
 	}
 #endif // !WITH_EDITOR
 
@@ -42,8 +271,15 @@ namespace UE::Audio::Insights
 		};
 
 		ProcessMessageQueue<FMixerSourceStartMessage>(TraceMessages.StartMessages, BumpEntryFunc,
-		[](const FMixerSourceStartMessage& Msg, TSharedPtr<FMixerSourceDashboardEntry>* OutEntry)
+		[this](const FMixerSourceStartMessage& Msg, TSharedPtr<FMixerSourceDashboardEntry>* OutEntry)
 		{
+#if !WITH_EDITOR
+			if (SessionCachedMessages.IsValid())
+			{
+				SessionCachedMessages->StartCachedMessages.EmplaceBack(Msg);
+			}
+#endif // !WITH_EDITOR
+
 			FMixerSourceDashboardEntry& EntryRef = *OutEntry->Get();
 			EntryRef.Name = *Msg.Name;
 			EntryRef.ComponentId = Msg.ComponentId;
@@ -53,8 +289,15 @@ namespace UE::Audio::Insights
 		TArray<int32, TInlineAllocator<64>> EntriesWithPoppedDataPoints;
 
 		ProcessMessageQueue<FMixerSourceVolumeMessage>(TraceMessages.VolumeMessages, BumpEntryFunc,
-		[&EntriesWithPoppedDataPoints](const FMixerSourceVolumeMessage& Msg, TSharedPtr<FMixerSourceDashboardEntry>* OutEntry)
+		[this, &EntriesWithPoppedDataPoints](const FMixerSourceVolumeMessage& Msg, TSharedPtr<FMixerSourceDashboardEntry>* OutEntry)
 		{
+#if !WITH_EDITOR
+			if (SessionCachedMessages.IsValid())
+			{
+				SessionCachedMessages->VolumeCachedMessages.EmplaceBack(Msg);
+			}
+#endif // !WITH_EDITOR
+
 			if (!EntriesWithPoppedDataPoints.Contains(Msg.PlayOrder))
 			{
 				(*OutEntry)->VolumeDataPoints.Pop((*OutEntry)->VolumeDataPoints.Num());
@@ -66,8 +309,15 @@ namespace UE::Audio::Insights
 
 		EntriesWithPoppedDataPoints.Reset();
 		ProcessMessageQueue<FMixerSourcePitchMessage>(TraceMessages.PitchMessages, BumpEntryFunc,
-		[&EntriesWithPoppedDataPoints](const FMixerSourcePitchMessage& Msg, TSharedPtr<FMixerSourceDashboardEntry>* OutEntry)
+		[this, &EntriesWithPoppedDataPoints](const FMixerSourcePitchMessage& Msg, TSharedPtr<FMixerSourceDashboardEntry>* OutEntry)
 		{
+#if !WITH_EDITOR
+			if (SessionCachedMessages.IsValid())
+			{
+				SessionCachedMessages->PitchCachedMessages.EmplaceBack(Msg);
+			}
+#endif // !WITH_EDITOR
+
 			if (!EntriesWithPoppedDataPoints.Contains(Msg.PlayOrder))
 			{
 				(*OutEntry)->PitchDataPoints.Pop((*OutEntry)->PitchDataPoints.Num());
@@ -79,8 +329,15 @@ namespace UE::Audio::Insights
 
 		EntriesWithPoppedDataPoints.Reset();
 		ProcessMessageQueue<FMixerSourceLPFFreqMessage>(TraceMessages.LPFFreqMessages, BumpEntryFunc,
-		[&EntriesWithPoppedDataPoints](const FMixerSourceLPFFreqMessage& Msg, TSharedPtr<FMixerSourceDashboardEntry>* OutEntry)
+		[this, &EntriesWithPoppedDataPoints](const FMixerSourceLPFFreqMessage& Msg, TSharedPtr<FMixerSourceDashboardEntry>* OutEntry)
 		{
+#if !WITH_EDITOR
+			if (SessionCachedMessages.IsValid())
+			{
+				SessionCachedMessages->LPFFreqCachedMessages.EmplaceBack(Msg);
+			}
+#endif // !WITH_EDITOR
+
 			if (!EntriesWithPoppedDataPoints.Contains(Msg.PlayOrder))
 			{
 				(*OutEntry)->LPFFreqDataPoints.Pop((*OutEntry)->LPFFreqDataPoints.Num());
@@ -92,8 +349,15 @@ namespace UE::Audio::Insights
 
 		EntriesWithPoppedDataPoints.Reset();
 		ProcessMessageQueue<FMixerSourceHPFFreqMessage>(TraceMessages.HPFFreqMessages, BumpEntryFunc,
-		[&EntriesWithPoppedDataPoints](const FMixerSourceHPFFreqMessage& Msg, TSharedPtr<FMixerSourceDashboardEntry>* OutEntry)
+		[this, &EntriesWithPoppedDataPoints](const FMixerSourceHPFFreqMessage& Msg, TSharedPtr<FMixerSourceDashboardEntry>* OutEntry)
 		{
+#if !WITH_EDITOR
+			if (SessionCachedMessages.IsValid())
+			{
+				SessionCachedMessages->HPFFreqCachedMessages.EmplaceBack(Msg);
+			}
+#endif // !WITH_EDITOR
+
 			if (!EntriesWithPoppedDataPoints.Contains(Msg.PlayOrder))
 			{
 				(*OutEntry)->HPFFreqDataPoints.Pop((*OutEntry)->HPFFreqDataPoints.Num());
@@ -105,8 +369,15 @@ namespace UE::Audio::Insights
 
 		EntriesWithPoppedDataPoints.Reset();
 		ProcessMessageQueue<FMixerSourceEnvelopeMessage>(TraceMessages.EnvelopeMessages, BumpEntryFunc,
-		[&EntriesWithPoppedDataPoints](const FMixerSourceEnvelopeMessage& Msg, TSharedPtr<FMixerSourceDashboardEntry>* OutEntry)
+		[this, &EntriesWithPoppedDataPoints](const FMixerSourceEnvelopeMessage& Msg, TSharedPtr<FMixerSourceDashboardEntry>* OutEntry)
 		{
+#if !WITH_EDITOR
+			if (SessionCachedMessages.IsValid())
+			{
+				SessionCachedMessages->EnvelopeCachedMessages.EmplaceBack(Msg);
+			}
+#endif // !WITH_EDITOR
+
 			if (!EntriesWithPoppedDataPoints.Contains(Msg.PlayOrder))
 			{
 				(*OutEntry)->EnvelopeDataPoints.Pop((*OutEntry)->EnvelopeDataPoints.Num());
@@ -118,8 +389,15 @@ namespace UE::Audio::Insights
 
 		EntriesWithPoppedDataPoints.Reset();
 		ProcessMessageQueue<FMixerSourceDistanceAttenuationMessage>(TraceMessages.DistanceAttenuationMessages, BumpEntryFunc,
-		[&EntriesWithPoppedDataPoints](const FMixerSourceDistanceAttenuationMessage& Msg, TSharedPtr<FMixerSourceDashboardEntry>* OutEntry)
+		[this, &EntriesWithPoppedDataPoints](const FMixerSourceDistanceAttenuationMessage& Msg, TSharedPtr<FMixerSourceDashboardEntry>* OutEntry)
 		{
+#if !WITH_EDITOR
+			if (SessionCachedMessages.IsValid())
+			{
+				SessionCachedMessages->DistanceAttenuationCachedMessages.EmplaceBack(Msg);
+			}
+#endif // !WITH_EDITOR
+
 			if (!EntriesWithPoppedDataPoints.Contains(Msg.PlayOrder))
 			{
 				(*OutEntry)->DistanceAttenuationDataPoints.Pop((*OutEntry)->DistanceAttenuationDataPoints.Num());
@@ -137,6 +415,13 @@ namespace UE::Audio::Insights
 		ProcessMessageQueue<FMixerSourceStopMessage>(TraceMessages.StopMessages, GetEntry,
 		[this](const FMixerSourceStopMessage& Msg, TSharedPtr<FMixerSourceDashboardEntry>* OutEntry)
 		{
+#if !WITH_EDITOR
+			if (SessionCachedMessages.IsValid())
+			{
+				SessionCachedMessages->StopCachedMessages.EmplaceBack(Msg);
+			}
+#endif // !WITH_EDITOR
+
 			if (OutEntry && (*OutEntry)->Timestamp < Msg.Timestamp)
 			{
 				RemoveDeviceEntry(Msg.DeviceId, Msg.PlayOrder);
@@ -146,6 +431,13 @@ namespace UE::Audio::Insights
 		return true;
 	}
 
+#if !WITH_EDITOR
+	void FMixerSourceTraceProvider::InitSessionCachedMessages(TraceServices::IAnalysisSession& InSession)
+	{
+		SessionCachedMessages = MakeUnique<FMixerSourceSessionCachedMessages>(InSession);
+	}
+#endif // !WITH_EDITOR
+
 	UE::Trace::IAnalyzer* FMixerSourceTraceProvider::ConstructAnalyzer(TraceServices::IAnalysisSession& InSession)
 	{
 		class FMixerSourceTraceAnalyzer : public FTraceAnalyzerBase
diff --git a/Engine/Plugins/AudioInsights/Source/AudioInsights/Private/Providers/MixerSourceTraceProvider.h b/Engine/Plugins/AudioInsights/Source/AudioInsights/Private/Providers/MixerSourceTraceProvider.h
--- a/Engine/Plugins/AudioInsights/Source/AudioInsights/Private/Providers/MixerSourceTraceProvider.h
+++ b/Engine/Plugins/AudioInsights/Source/AudioInsights/Private/Providers/MixerSourceTraceProvider.h
@@ -27,9 +27,15 @@ namespace UE::Audio::Insights
 
 		static FName GetName_Static();
 
+#if !WITH_EDITOR
+		virtual void InitSessionCachedMessages(TraceServices::IAnalysisSession& InSession) override;
+#endif // !WITH_EDITOR
+
 	private:
 #if !WITH_EDITOR
 		virtual void OnTimingViewTimeMarkerChanged(double TimeMarker) override;
+
+		TUniquePtr<FMixerSourceSessionCachedMessages> SessionCachedMessages;
 #endif // !WITH_EDITOR
 
 		FMixerSourceMessages TraceMessages;
diff --git a/Engine/Plugins/AudioInsights/Source/AudioInsights/Private/Providers/VirtualLoopTraceProvider.cpp b/Engine/Plugins/AudioInsights/Source/AudioInsights/Private/Providers/VirtualLoopTraceProvider.cpp
--- a/Engine/Plugins/AudioInsights/Source/AudioInsights/Private/Providers/VirtualLoopTraceProvider.cpp
+++ b/Engine/Plugins/AudioInsights/Source/AudioInsights/Private/Providers/VirtualLoopTraceProvider.cpp
@@ -5,9 +5,34 @@
 #include "TraceServices/Model/AnalysisSession.h"
 #include "TraceServices/ModuleService.h"
 
+#if !WITH_EDITOR
+#include "Common/PagedArray.h"
+#endif // !WITH_EDITOR
 
 namespace UE::Audio::Insights
 {
+	namespace FVirtualLoopTraceProviderPrivate
+	{
+#if !WITH_EDITOR
+		template<typename T>
+		const T* FindClosestMessageToTimestamp(const TraceServices::TPagedArray<T>& InCachedMessages, const double InTimeMarker, const uint32 InPlayOrder)
+		{
+			const int32 ClosestMessageToTimeStampIndex = TraceServices::PagedArrayAlgo::BinarySearchClosestBy(InCachedMessages, InTimeMarker, [](const T& Msg) { return Msg.Timestamp; });
+
+			// Iterate backwards from TimeMarker until we find the matching PlayOrder
+			for (auto It = InCachedMessages.GetIteratorFromItem(ClosestMessageToTimeStampIndex); It; --It)
+			{
+				if (It->PlayOrder == InPlayOrder)
+				{
+					return &(*It);
+				}
+			}
+
+			return nullptr;
+		}
+#endif // !WITH_EDITOR
+	}
+
 	FName FVirtualLoopTraceProvider::GetName_Static()
 	{
 		return "AudioVirtualLoopProvider";
@@ -16,7 +41,109 @@ namespace UE::Audio::Insights
 #if !WITH_EDITOR
 	void FVirtualLoopTraceProvider::OnTimingViewTimeMarkerChanged(double TimeMarker)
 	{
-		// TODO alex.perez: to simplify code review, the implementation of this method will be submitted after this CL (in part 3)
+		using namespace FVirtualLoopTraceProviderPrivate;
+
+		if (!SessionCachedMessages.IsValid())
+		{
+			return;
+		}
+
+		DeviceDataMap.Empty();
+
+		// Collect all the virtualize messages registered until this point in time 
+		for (const FVirtualLoopVirtualizeMessage& VirtualizeCachedMessage : SessionCachedMessages->VirtualizeCachedMessages)
+		{
+			if (VirtualizeCachedMessage.Timestamp > TimeMarker)
+			{
+				break;
+			}
+
+			UpdateDeviceEntry(VirtualizeCachedMessage.DeviceId, VirtualizeCachedMessage.PlayOrder, [&VirtualizeCachedMessage](TSharedPtr<FVirtualLoopDashboardEntry>& Entry)
+			{
+				if (!Entry.IsValid())
+				{
+					Entry = MakeShared<FVirtualLoopDashboardEntry>();
+					Entry->DeviceId  = VirtualizeCachedMessage.DeviceId;
+					Entry->PlayOrder = VirtualizeCachedMessage.PlayOrder;
+				}
+				Entry->Timestamp = VirtualizeCachedMessage.Timestamp;
+
+				Entry->Name        = *VirtualizeCachedMessage.Name;
+				Entry->ComponentId = VirtualizeCachedMessage.ComponentId;
+			});
+		}
+
+		// Selectively remove virtualize messages collected in the step above by knowing which sounds were stopped/realized.
+		// With this we will know what are the virtualized sounds at this point in time.
+		for (const FVirtualLoopMessageBase& StopOrRealizeCachedMessage : SessionCachedMessages->StopOrRealizeCachedMessages)
+		{
+			if (StopOrRealizeCachedMessage.Timestamp > TimeMarker)
+			{
+				break;
+			}
+
+			auto* OutEntry = FindDeviceEntry(StopOrRealizeCachedMessage.DeviceId, StopOrRealizeCachedMessage.PlayOrder);
+
+			if (OutEntry && (*OutEntry)->Timestamp < StopOrRealizeCachedMessage.Timestamp)
+			{
+				RemoveDeviceEntry(StopOrRealizeCachedMessage.DeviceId, StopOrRealizeCachedMessage.PlayOrder);
+			}
+		}
+
+		// For now we only retrieve information from AudioDeviceId 1 (main device in standalone games)
+		const FDeviceData* DeviceData = DeviceDataMap.Find(1);
+		if (DeviceData)
+		{
+			// Collect update messages from virtualized sounds (based on active sounds's PlayOrder)
+			struct CachedEntryInfo
+			{
+				FVirtualLoopUpdateMessage UpdateMessage;
+			};
+
+			TArray<uint32> PlayOrderArray;
+			(*DeviceData).GenerateKeyArray(PlayOrderArray);
+
+			TArray<CachedEntryInfo> CachedEntryInfos;
+			CachedEntryInfos.SetNumUninitialized(PlayOrderArray.Num());
+
+			ParallelFor(PlayOrderArray.Num(), 
+			[&PlayOrderArray, &CachedEntryInfos, TimeMarker, this](const int32 Index)
+			{
+				const uint32 PlayOrder = PlayOrderArray[Index];
+
+				// Update
+				const FVirtualLoopUpdateMessage* FoundUpdateCachedMessage = FindClosestMessageToTimestamp(SessionCachedMessages->UpdateCachedMessages, TimeMarker, PlayOrder);
+				if (FoundUpdateCachedMessage)
+				{
+					CachedEntryInfos[Index].UpdateMessage = *FoundUpdateCachedMessage;
+				}
+			});
+
+			// Update the device entries with the collected info
+			for (const CachedEntryInfo& CachedEntryInfo : CachedEntryInfos)
+			{
+				UpdateDeviceEntry(CachedEntryInfo.UpdateMessage.DeviceId, CachedEntryInfo.UpdateMessage.PlayOrder, [&CachedEntryInfo](TSharedPtr<FVirtualLoopDashboardEntry>& Entry)
+				{
+					if (!Entry.IsValid())
+					{
+						Entry = MakeShared<FVirtualLoopDashboardEntry>();
+						Entry->DeviceId  = CachedEntryInfo.UpdateMessage.DeviceId;
+						Entry->PlayOrder = CachedEntryInfo.UpdateMessage.PlayOrder;
+					}
+
+					Entry->Timestamp = CachedEntryInfo.UpdateMessage.Timestamp;
+
+					Entry->PlaybackTime    = CachedEntryInfo.UpdateMessage.PlaybackTime;
+					Entry->TimeVirtualized = CachedEntryInfo.UpdateMessage.TimeVirtualized;
+					Entry->UpdateInterval  = CachedEntryInfo.UpdateMessage.UpdateInterval;
+					Entry->Location        = FVector{ CachedEntryInfo.UpdateMessage.LocationX, CachedEntryInfo.UpdateMessage.LocationY, CachedEntryInfo.UpdateMessage.LocationZ };
+					Entry->Rotator         = FRotator{ CachedEntryInfo.UpdateMessage.RotatorPitch, CachedEntryInfo.UpdateMessage.RotatorYaw, CachedEntryInfo.UpdateMessage.RotatorRoll };
+				});
+			}
+		}
+
+		// Call parent method to update LastMessageId
+		FTraceProviderBase::OnTimingViewTimeMarkerChanged(TimeMarker);
 	}
 #endif // !WITH_EDITOR
 
@@ -24,6 +151,13 @@ namespace UE::Audio::Insights
 	{
 		auto RemoveEntryFunc = [this](const FVirtualLoopRealizeMessage& Msg, TSharedPtr<FVirtualLoopDashboardEntry>* OutEntry)
 		{
+#if !WITH_EDITOR
+			if (SessionCachedMessages.IsValid())
+			{
+				SessionCachedMessages->StopOrRealizeCachedMessages.EmplaceBack(Msg);
+			}
+#endif // !WITH_EDITOR
+
 			if (OutEntry && (*OutEntry)->Timestamp < Msg.Timestamp)
 			{
 				RemoveDeviceEntry(Msg.DeviceId, Msg.PlayOrder);
@@ -56,8 +190,15 @@ namespace UE::Audio::Insights
 
 		ProcessMessageQueue<FVirtualLoopVirtualizeMessage>(TraceMessages.VirtualizeMessages,
 		BumpEntryFunc,
-		[](const FVirtualLoopVirtualizeMessage& Msg, TSharedPtr<FVirtualLoopDashboardEntry>* OutEntry)
+		[this](const FVirtualLoopVirtualizeMessage& Msg, TSharedPtr<FVirtualLoopDashboardEntry>* OutEntry)
 		{
+#if !WITH_EDITOR
+			if (SessionCachedMessages.IsValid())
+			{
+				SessionCachedMessages->VirtualizeCachedMessages.EmplaceBack(Msg);
+			}
+#endif // !WITH_EDITOR
+
 			FVirtualLoopDashboardEntry& EntryRef = *OutEntry->Get();
 			EntryRef.Name = Msg.Name;
 			EntryRef.ComponentId = Msg.ComponentId;
@@ -67,8 +208,15 @@ namespace UE::Audio::Insights
 
 		ProcessMessageQueue<FVirtualLoopUpdateMessage>(TraceMessages.UpdateMessages, 
 		BumpEntryFunc,
-		[](const FVirtualLoopUpdateMessage& Msg, TSharedPtr<FVirtualLoopDashboardEntry>* OutEntry)
+		[this](const FVirtualLoopUpdateMessage& Msg, TSharedPtr<FVirtualLoopDashboardEntry>* OutEntry)
 		{
+#if !WITH_EDITOR
+			if (SessionCachedMessages.IsValid())
+			{
+				SessionCachedMessages->UpdateCachedMessages.EmplaceBack(Msg);
+			}
+#endif // !WITH_EDITOR
+
 			if (OutEntry)
 			{
 				FVirtualLoopDashboardEntry& EntryRef = *OutEntry->Get();
@@ -85,6 +233,13 @@ namespace UE::Audio::Insights
 		return true;
 	}
 
+#if !WITH_EDITOR
+	void FVirtualLoopTraceProvider::InitSessionCachedMessages(TraceServices::IAnalysisSession& InSession)
+	{
+		SessionCachedMessages = MakeUnique<FVirtualLoopSessionCachedMessages>(InSession);
+	}
+#endif // !WITH_EDITOR
+
 	UE::Trace::IAnalyzer* FVirtualLoopTraceProvider::ConstructAnalyzer(TraceServices::IAnalysisSession& InSession)
 	{
 		class FVirtualLoopTraceAnalyzer : public FTraceAnalyzerBase
diff --git a/Engine/Plugins/AudioInsights/Source/AudioInsights/Private/Providers/VirtualLoopTraceProvider.h b/Engine/Plugins/AudioInsights/Source/AudioInsights/Private/Providers/VirtualLoopTraceProvider.h
--- a/Engine/Plugins/AudioInsights/Source/AudioInsights/Private/Providers/VirtualLoopTraceProvider.h
+++ b/Engine/Plugins/AudioInsights/Source/AudioInsights/Private/Providers/VirtualLoopTraceProvider.h
@@ -24,9 +24,15 @@ namespace UE::Audio::Insights
 
 		static FName GetName_Static();
 
+#if !WITH_EDITOR
+		virtual void InitSessionCachedMessages(TraceServices::IAnalysisSession& InSession) override;
+#endif // !WITH_EDITOR
+
 	private:
 #if !WITH_EDITOR
 		virtual void OnTimingViewTimeMarkerChanged(double TimeMarker) override;
+
+		TUniquePtr<FVirtualLoopSessionCachedMessages> SessionCachedMessages;
 #endif // !WITH_EDITOR
 
 		virtual bool ProcessMessages() override;
diff --git a/Engine/Plugins/AudioInsights/Source/AudioInsights/Private/Views/MixerSourceDashboardViewFactory.cpp b/Engine/Plugins/AudioInsights/Source/AudioInsights/Private/Views/MixerSourceDashboardViewFactory.cpp
--- a/Engine/Plugins/AudioInsights/Source/AudioInsights/Private/Views/MixerSourceDashboardViewFactory.cpp
+++ b/Engine/Plugins/AudioInsights/Source/AudioInsights/Private/Views/MixerSourceDashboardViewFactory.cpp
@@ -424,9 +424,11 @@ namespace UE::Audio::Insights
 #else
 	void FMixerSourceDashboardViewFactory::OnAudioInsightsComponentTabSpawn()
 	{
-		// It is guaranteed to have a game running when the audio insights tab is spawning from Unreal Insights 
-		// (In Unreal Insights the Audio Insights menu item is only active when a live trace is running)
-		GameState = EGameState::Running;
+		const TSharedPtr<const FAudioInsightsComponent> AudioInsightsComponent = FAudioInsightsModule::GetChecked().GetAudioInsightsComponent();
+		if (AudioInsightsComponent.IsValid())
+		{
+			GameState = AudioInsightsComponent->GetIsLiveSession() ? EGameState::Running : EGameState::Stopped;
+		}
 	}
 #endif // WITH_EDITOR
 
diff --git a/Engine/Plugins/AudioInsights/Source/AudioInsights/Public/AudioInsightsTraceProviderBase.h b/Engine/Plugins/AudioInsights/Source/AudioInsights/Public/AudioInsightsTraceProviderBase.h
--- a/Engine/Plugins/AudioInsights/Source/AudioInsights/Public/AudioInsightsTraceProviderBase.h
+++ b/Engine/Plugins/AudioInsights/Source/AudioInsights/Public/AudioInsightsTraceProviderBase.h
@@ -50,7 +50,8 @@ namespace UE::Audio::Insights
 		}
 
 #if !WITH_EDITOR
-		virtual void OnTimingViewTimeMarkerChanged(double TimeMarker) {};
+		virtual void InitSessionCachedMessages(TraceServices::IAnalysisSession& InSession) {}
+		virtual void OnTimingViewTimeMarkerChanged(double TimeMarker) { ++LastMessageId; };
 #endif // !WITH_EDITOR
 
 	protected:
diff --git a/Engine/Plugins/AudioInsights/Source/AudioInsights/Public/Messages/MixerSourceTraceMessages.h b/Engine/Plugins/AudioInsights/Source/AudioInsights/Public/Messages/MixerSourceTraceMessages.h
--- a/Engine/Plugins/AudioInsights/Source/AudioInsights/Public/Messages/MixerSourceTraceMessages.h
+++ b/Engine/Plugins/AudioInsights/Source/AudioInsights/Public/Messages/MixerSourceTraceMessages.h
@@ -7,6 +7,10 @@
 #include "Trace/Analyzer.h"
 #include "Views/TableDashboardViewFactory.h"
 
+#if !WITH_EDITOR
+#include "Common/PagedArray.h"
+#include "TraceServices/Model/AnalysisSession.h"
+#endif // !WITH_EDITOR
 
 namespace UE::Audio::Insights
 {
@@ -108,4 +112,32 @@ namespace UE::Audio::Insights
 		TAnalyzerMessageQueue<FMixerSourceStopMessage> StopMessages { 0.1 };
 		TAnalyzerMessageQueue<FMixerSourceVolumeMessage> VolumeMessages { 0.1 };
 	};
+
+#if !WITH_EDITOR
+	struct FMixerSourceSessionCachedMessages
+	{
+		FMixerSourceSessionCachedMessages(TraceServices::IAnalysisSession& InSession)
+			: DistanceAttenuationCachedMessages(InSession.GetLinearAllocator(), 16384)
+			, EnvelopeCachedMessages(InSession.GetLinearAllocator(), 16384)
+			, HPFFreqCachedMessages(InSession.GetLinearAllocator(), 16384)
+			, LPFFreqCachedMessages(InSession.GetLinearAllocator(), 16384)
+			, PitchCachedMessages(InSession.GetLinearAllocator(), 16384)
+			, StartCachedMessages(InSession.GetLinearAllocator(), 16384)
+			, StopCachedMessages(InSession.GetLinearAllocator(), 4096)
+			, VolumeCachedMessages(InSession.GetLinearAllocator(), 16384)
+		{
+
+		}
+
+		TraceServices::TPagedArray<FMixerSourceDistanceAttenuationMessage> DistanceAttenuationCachedMessages;
+		TraceServices::TPagedArray<FMixerSourceEnvelopeMessage> EnvelopeCachedMessages;
+		TraceServices::TPagedArray<FMixerSourceHPFFreqMessage> HPFFreqCachedMessages;
+		TraceServices::TPagedArray<FMixerSourceLPFFreqMessage> LPFFreqCachedMessages;
+		TraceServices::TPagedArray<FMixerSourcePitchMessage> PitchCachedMessages;
+		TraceServices::TPagedArray<FMixerSourceStartMessage> StartCachedMessages;
+		TraceServices::TPagedArray<FMixerSourceStopMessage> StopCachedMessages;
+		TraceServices::TPagedArray<FMixerSourceVolumeMessage> VolumeCachedMessages;
+	};
+#endif // !WITH_EDITOR
+
 } // namespace UE::Audio::Insights
diff --git a/Engine/Plugins/AudioInsights/Source/AudioInsights/Public/Messages/VirtualLoopTraceMessages.h b/Engine/Plugins/AudioInsights/Source/AudioInsights/Public/Messages/VirtualLoopTraceMessages.h
--- a/Engine/Plugins/AudioInsights/Source/AudioInsights/Public/Messages/VirtualLoopTraceMessages.h
+++ b/Engine/Plugins/AudioInsights/Source/AudioInsights/Public/Messages/VirtualLoopTraceMessages.h
@@ -7,6 +7,10 @@
 #include "Trace/Analyzer.h"
 #include "Views/TableDashboardViewFactory.h"
 
+#if !WITH_EDITOR
+#include "Common/PagedArray.h"
+#include "TraceServices/Model/AnalysisSession.h"
+#endif // !WITH_EDITOR
 
 namespace UE::Audio::Insights
 {
@@ -108,4 +112,22 @@ namespace UE::Audio::Insights
 
 		friend class FVirtualLoopTraceProvider;
 	};
+
+#if !WITH_EDITOR
+	struct FVirtualLoopSessionCachedMessages
+	{
+		FVirtualLoopSessionCachedMessages(TraceServices::IAnalysisSession& InSession)
+			: VirtualizeCachedMessages(InSession.GetLinearAllocator(), 4096)
+			, StopOrRealizeCachedMessages(InSession.GetLinearAllocator(), 4096)
+			, UpdateCachedMessages(InSession.GetLinearAllocator(), 16384)
+		{
+
+		}
+
+		TraceServices::TPagedArray<FVirtualLoopVirtualizeMessage> VirtualizeCachedMessages;
+		TraceServices::TPagedArray<FVirtualLoopStopMessage> StopOrRealizeCachedMessages;
+		TraceServices::TPagedArray<FVirtualLoopUpdateMessage> UpdateCachedMessages;
+	};
+#endif // !WITH_EDITOR
+
 } // namespace UE::Audio::Insights
