[Audio Insights] Be able to reset the tabs layout via: View > Reset Layout

diff --git a/Engine/Plugins/AudioInsights/Source/AudioInsightsEditor/Private/AudioInsightsEditorDashboardFactory.cpp b/Engine/Plugins/AudioInsights/Source/AudioInsightsEditor/Private/AudioInsightsEditorDashboardFactory.cpp
--- a/Engine/Plugins/AudioInsights/Source/AudioInsightsEditor/Private/AudioInsightsEditorDashboardFactory.cpp
+++ b/Engine/Plugins/AudioInsights/Source/AudioInsightsEditor/Private/AudioInsightsEditorDashboardFactory.cpp
@@ -1,6 +1,7 @@
 // Copyright Epic Games, Inc. All Rights Reserved.
 #include "AudioInsightsEditorDashboardFactory.h"
 
+#include "Async/Async.h"
 #include "AudioDevice.h"
 #include "AudioDeviceManager.h"
 #include "AudioInsightsEditorModule.h"
@@ -347,6 +348,31 @@ namespace UE::Audio::Insights
 						MenuBuilder.AddMenuSeparator();
 					}
 				}
+
+				MenuBuilder.AddMenuSeparator();
+
+				MenuBuilder.AddMenuEntry(LOCTEXT("ViewMenu_ResetLayoutText", "Reset Layout"), 
+					FText::GetEmpty(), 
+					FSlateIcon(),
+					FUIAction(FExecuteAction::CreateLambda([this]()
+					{
+						if (DashboardTabManager)
+						{
+							if (TSharedPtr<SDockTab> OwnerTab = DashboardTabManager->GetOwnerTab())
+							{
+								// Wipe all the persisted areas and close tab
+								DashboardTabManager->CloseAllAreas();
+								OwnerTab->RequestCloseTab();
+
+								// Can't invoke the tab immediately (it won't show up), needs to be done a bit later
+								AsyncTask(ENamedThreads::GameThread, [AudioInsightsTabId = OwnerTab->GetLayoutIdentifier()]()
+								{
+									FGlobalTabmanager::Get()->TryInvokeTab(AudioInsightsTabId);
+								});
+							}
+						}
+					}),
+					FCanExecuteAction()));
 			}),
 			"View"
 		);
