[Audio Insights] Closing the Mixer Source Plots tabs explicitly should turn off the plots button + layout persistency adjustments for this case

diff --git a/Engine/Plugins/AudioInsights/Source/AudioInsights/Private/Views/MixerSourceDashboardViewFactory.cpp b/Engine/Plugins/AudioInsights/Source/AudioInsights/Private/Views/MixerSourceDashboardViewFactory.cpp
--- a/Engine/Plugins/AudioInsights/Source/AudioInsights/Private/Views/MixerSourceDashboardViewFactory.cpp
+++ b/Engine/Plugins/AudioInsights/Source/AudioInsights/Private/Views/MixerSourceDashboardViewFactory.cpp
@@ -1,6 +1,7 @@
 // Copyright Epic Games, Inc. All Rights Reserved.
 #include "Views/MixerSourceDashboardViewFactory.h"
 
+#include "Async/Async.h"
 #include "AudioInsightsModule.h"
 #include "AudioInsightsStyle.h"
 #include "DSP/Dsp.h"
@@ -952,6 +953,19 @@ namespace UE::Audio::Insights
 
 		return SNew(SDockTab)
 			.Clipping(EWidgetClipping::ClipToBounds)
+			.OnTabClosed_Lambda([this](TSharedRef<SDockTab> InDockTab)
+			{				
+				if (PlotsButton.IsValid())
+				{
+					PlotsButton->SetIsChecked(ECheckBoxState::Unchecked);
+
+					// Can't save layout immediately (it won't save the tab closed state), needs to be done a bit later
+					AsyncTask(ENamedThreads::GameThread, [this]()
+					{
+						SaveLayoutToConfig();
+					});
+				}
+			})
 			[
 				PlotsWidget ? PlotsWidget.ToSharedRef() : SNullWidget::NullWidget
 			];
