[Audio Insights] Remove Log tab in Audio Insights component for Unreal Insights until further notice (currently this tab is not showing live audio information like when running Audio Insights from the editor, a potential option is to see if we can adapt and use the Console tab from Session Frontend inside the Audio Insights tab. Needs further discussion).

diff --git a/Engine/Plugins/AudioInsights/Source/AudioInsights/Private/AudioInsightsModule.cpp b/Engine/Plugins/AudioInsights/Source/AudioInsights/Private/AudioInsightsModule.cpp
--- a/Engine/Plugins/AudioInsights/Source/AudioInsights/Private/AudioInsightsModule.cpp
+++ b/Engine/Plugins/AudioInsights/Source/AudioInsights/Private/AudioInsightsModule.cpp
@@ -14,7 +14,6 @@
 
 #if !WITH_EDITOR
 #include "AudioInsightsComponent.h"
-#include "Views/LogDashboardViewFactory.h"
 #include "Views/MixerSourceDashboardViewFactory.h"
 #include "Views/VirtualLoopDashboardViewFactory.h"
 #endif // !WITH_EDITOR
@@ -40,7 +39,6 @@ namespace UE::Audio::Insights
 #if !WITH_EDITOR
 			IModularFeatures::Get().RegisterModularFeature(UE::Insights::Timing::TimingViewExtenderFeatureName, &AudioInsightsTimingViewExtender);
 
-			DashboardFactory->RegisterViewFactory(MakeShared<FLogDashboardViewFactory>());
 			DashboardFactory->RegisterViewFactory(MakeShared<FMixerSourceDashboardViewFactory>());
 			DashboardFactory->RegisterViewFactory(MakeShared<FVirtualLoopDashboardViewFactory>());
 			
