Added missing WITH_EDITOR pragma!

diff --git a/Engine/Plugins/AudioInsights/Source/AudioInsights/Private/Views/MixerSourceDashboardViewFactory.cpp b/Engine/Plugins/AudioInsights/Source/AudioInsights/Private/Views/MixerSourceDashboardViewFactory.cpp
--- a/Engine/Plugins/AudioInsights/Source/AudioInsights/Private/Views/MixerSourceDashboardViewFactory.cpp
+++ b/Engine/Plugins/AudioInsights/Source/AudioInsights/Private/Views/MixerSourceDashboardViewFactory.cpp
@@ -953,6 +953,7 @@ namespace UE::Audio::Insights
 
 		return SNew(SDockTab)
 			.Clipping(EWidgetClipping::ClipToBounds)
+#if WITH_EDITOR
 			.OnTabClosed_Lambda([this](TSharedRef<SDockTab> InDockTab)
 			{				
 				if (PlotsButton.IsValid())
@@ -966,6 +967,7 @@ namespace UE::Audio::Insights
 					});
 				}
 			})
+#endif // WITH_EDITOR
 			[
 				PlotsWidget ? PlotsWidget.ToSharedRef() : SNullWidget::NullWidget
 			];
