Fix nopch compile of UnrealInsights

diff --git a/Engine/Plugins/AudioInsights/Source/AudioInsights/Private/Providers/MixerSourceTraceProvider.cpp b/Engine/Plugins/AudioInsights/Source/AudioInsights/Private/Providers/MixerSourceTraceProvider.cpp
--- a/Engine/Plugins/AudioInsights/Source/AudioInsights/Private/Providers/MixerSourceTraceProvider.cpp
+++ b/Engine/Plugins/AudioInsights/Source/AudioInsights/Private/Providers/MixerSourceTraceProvider.cpp
@@ -8,6 +8,7 @@
 
 #if !WITH_EDITOR
 #include "Common/PagedArray.h"
+#include "Async/ParallelFor.h"
 #endif // !WITH_EDITOR
 
 namespace UE::Audio::Insights
diff --git a/Engine/Plugins/AudioInsights/Source/AudioInsights/Private/Providers/VirtualLoopTraceProvider.cpp b/Engine/Plugins/AudioInsights/Source/AudioInsights/Private/Providers/VirtualLoopTraceProvider.cpp
--- a/Engine/Plugins/AudioInsights/Source/AudioInsights/Private/Providers/VirtualLoopTraceProvider.cpp
+++ b/Engine/Plugins/AudioInsights/Source/AudioInsights/Private/Providers/VirtualLoopTraceProvider.cpp
@@ -7,6 +7,7 @@
 
 #if !WITH_EDITOR
 #include "Common/PagedArray.h"
+#include "Async/ParallelFor.h"
 #endif // !WITH_EDITOR
 
 namespace UE::Audio::Insights
