[Audio Insights] Mixer Source tab: Re-enable plots in Audio Insights for Unreal insights (only for live traces as it was originally, for offline traces a different implementation is required since it is not based on stream trace data)

diff --git a/Engine/Plugins/AudioInsights/Source/AudioInsights/Private/AudioInsightsComponent.cpp b/Engine/Plugins/AudioInsights/Source/AudioInsights/Private/AudioInsightsComponent.cpp
--- a/Engine/Plugins/AudioInsights/Source/AudioInsights/Private/AudioInsightsComponent.cpp
+++ b/Engine/Plugins/AudioInsights/Source/AudioInsights/Private/AudioInsightsComponent.cpp
@@ -113,25 +113,35 @@ namespace UE::Audio::Insights
 		return DockTab;
 	}
 
+	bool FAudioInsightsComponent::GetIsLiveSession() const
+	{
+		IUnrealInsightsModule& UnrealInsightsModule = FModuleManager::LoadModuleChecked<IUnrealInsightsModule>("TraceInsights");
+
+		TSharedPtr<const TraceServices::IAnalysisSession> Session = UnrealInsightsModule.GetAnalysisSession();
+		if (Session.IsValid())
+		{
+			TraceServices::FAnalysisSessionReadScope SessionReadScope(*Session.Get());
+
+			const UE::Trace::FStoreClient* StoreClient = UnrealInsightsModule.GetStoreClient();
+			const UE::Trace::FStoreClient::FSessionInfo* SessionInfo = StoreClient ? StoreClient->GetSessionInfoByTraceId(Session->GetTraceId()) : nullptr;
+
+			return !Session->IsAnalysisComplete() && StoreClient != nullptr && SessionInfo != nullptr;
+		}
+
+		return false;
+	}
+
 	bool FAudioInsightsComponent::Tick(float DeltaTime)
 	{
-		// Audio Insights will be available in non-editor file traces or if there is an active standalone game live session
 		if (!bCanSpawnTab)
 		{
-			IUnrealInsightsModule& UnrealInsightsModule = FModuleManager::LoadModuleChecked<IUnrealInsightsModule>("TraceInsights");
+			const IUnrealInsightsModule& UnrealInsightsModule = FModuleManager::LoadModuleChecked<IUnrealInsightsModule>("TraceInsights");
 
 			TSharedPtr<const TraceServices::IAnalysisSession> Session = UnrealInsightsModule.GetAnalysisSession();
 			if (Session.IsValid())
 			{
 				TraceServices::FAnalysisSessionReadScope SessionReadScope(*Session.Get());
 
-				// Set bIsLiveSession
-				const UE::Trace::FStoreClient* StoreClient = UnrealInsightsModule.GetStoreClient();
-				const UE::Trace::FStoreClient::FSessionInfo* SessionInfo = StoreClient ? StoreClient->GetSessionInfoByTraceId(Session->GetTraceId()) : nullptr;
-
-				bIsLiveSession = !Session->IsAnalysisComplete() && StoreClient != nullptr && SessionInfo != nullptr;
-
-				// Set bIsEditorTrace
 				const TraceServices::IDiagnosticsProvider* DiagnosticsProvider = TraceServices::ReadDiagnosticsProvider(*Session.Get());
 				if (DiagnosticsProvider && DiagnosticsProvider->IsSessionInfoAvailable())
 				{
diff --git a/Engine/Plugins/AudioInsights/Source/AudioInsights/Private/AudioInsightsComponent.h b/Engine/Plugins/AudioInsights/Source/AudioInsights/Private/AudioInsightsComponent.h
--- a/Engine/Plugins/AudioInsights/Source/AudioInsights/Private/AudioInsightsComponent.h
+++ b/Engine/Plugins/AudioInsights/Source/AudioInsights/Private/AudioInsightsComponent.h
@@ -24,7 +24,7 @@ namespace UE::Audio::Insights
 		virtual void RegisterMajorTabs(IUnrealInsightsModule& InsightsModule) override;
 		virtual void UnregisterMajorTabs() override;
 
-		bool GetIsLiveSession() const { return bIsLiveSession; }
+		bool GetIsLiveSession() const;
 		bool GetIsEditorTrace() const { return bIsEditorTrace; }
 
 		DECLARE_MULTICAST_DELEGATE(FOnTabSpawn);
@@ -37,11 +37,8 @@ namespace UE::Audio::Insights
 		bool Tick(float DeltaTime);
 
 		bool bIsInitialized = false;
-
-		bool bIsLiveSession = false;
 		bool bIsEditorTrace = false;
-
-		bool bCanSpawnTab = false;
+		bool bCanSpawnTab   = false;
 
 		FTickerDelegate OnTick;
 		FTSTicker::FDelegateHandle OnTickHandle;
diff --git a/Engine/Plugins/AudioInsights/Source/AudioInsights/Private/Views/MixerSourceDashboardViewFactory.cpp b/Engine/Plugins/AudioInsights/Source/AudioInsights/Private/Views/MixerSourceDashboardViewFactory.cpp
--- a/Engine/Plugins/AudioInsights/Source/AudioInsights/Private/Views/MixerSourceDashboardViewFactory.cpp
+++ b/Engine/Plugins/AudioInsights/Source/AudioInsights/Private/Views/MixerSourceDashboardViewFactory.cpp
@@ -6,6 +6,7 @@
 #include "AudioInsightsStyle.h"
 #include "DSP/Dsp.h"
 #include "Framework/Docking/LayoutService.h"
+#include "Framework/Docking/TabManager.h"
 #include "Internationalization/Text.h"
 #include "Providers/MixerSourceTraceProvider.h"
 #include "SSimpleTimeSlider.h"
@@ -847,6 +848,20 @@ namespace UE::Audio::Insights
 			];
 	}
 
+	TSharedRef<FTabManager::FLayout> FMixerSourceDashboardViewFactory::LoadLayoutFromConfig()
+	{
+		return FLayoutSaveRestore::LoadFromConfig(GEditorLayoutIni, GetDefaultTabLayout());
+	}
+
+	void FMixerSourceDashboardViewFactory::SaveLayoutToConfig()
+	{
+		if (MixerSourcesTabManager.IsValid())
+		{
+			FLayoutSaveRestore::SaveToConfig(GEditorLayoutIni, MixerSourcesTabManager->PersistLayout());
+		}
+	}
+#endif // WITH_EDITOR
+
 	TSharedRef<SWidget> FMixerSourceDashboardViewFactory::MakePlotsButtonWidget()
 	{
 		return SAssignNew(PlotsButton, SCheckBox)
@@ -861,24 +876,11 @@ namespace UE::Audio::Insights
 				[
 					SNew(SImage)
 					.ColorAndOpacity(FSlateColor::UseForeground())
-					.Image(FAppStyle::GetBrush("GenericCurveEditor.SetViewModeStacked"))
+					.Image(FSlateStyle::Get().GetBrush("AudioInsights.Icon.Sources.Plots"))
 				]
 			];
 	}
 
-	TSharedRef<FTabManager::FLayout> FMixerSourceDashboardViewFactory::LoadLayoutFromConfig()
-	{
-		return FLayoutSaveRestore::LoadFromConfig(GEditorLayoutIni, GetDefaultTabLayout());
-	}
-
-	void FMixerSourceDashboardViewFactory::SaveLayoutToConfig()
-	{
-		if (MixerSourcesTabManager.IsValid())
-		{
-			FLayoutSaveRestore::SaveToConfig(GEditorLayoutIni, MixerSourcesTabManager->PersistLayout());
-		}
-	}
-
 	void FMixerSourceDashboardViewFactory::TogglePlotsTabVisibility(ECheckBoxState InCheckboxState)
 	{
 		using namespace MixerSourcePrivate;
@@ -901,28 +903,34 @@ namespace UE::Audio::Insights
 			}
 		}
 
+#if WITH_EDITOR
 		SaveLayoutToConfig();
-	}
 #endif // WITH_EDITOR
+	}
 
 	TSharedRef<SDockTab> FMixerSourceDashboardViewFactory::CreateMixerSourcesTab(const FSpawnTabArgs& Args)
 	{
+#if !WITH_EDITOR
+		const TSharedPtr<const FAudioInsightsComponent> AudioInsightsComponent = FAudioInsightsModule::GetChecked().GetAudioInsightsComponent();
+#endif // !WITH_EDITOR
+
 		return SNew(SDockTab)
 			.Clipping(EWidgetClipping::ClipToBounds)
 			[
 				SNew(SVerticalBox)
-#if WITH_EDITOR
 				+SVerticalBox::Slot()
 				.AutoHeight()
 				.HAlign(HAlign_Fill)
 				.Padding(0.0f, 0.0f, 0.0f, 6.0f)
 				[
 					SNew(SHorizontalBox)
+#if WITH_EDITOR
 					+SHorizontalBox::Slot()
 					.HAlign(HAlign_Left)
 					[
 						MakeMuteSoloWidget()
 					]
+#endif // WITH_EDITOR
 					+ SHorizontalBox::Slot()
 					.HAlign(HAlign_Right)
 					[
@@ -932,10 +940,13 @@ namespace UE::Audio::Insights
 					.AutoWidth()
 					.HAlign(HAlign_Right)
 					[
+#if WITH_EDITOR
 						MakePlotsButtonWidget()
+#else
+						AudioInsightsComponent.IsValid() && AudioInsightsComponent->GetIsLiveSession() ? MakePlotsButtonWidget() : SNullWidget::NullWidget
+#endif // WITH_EDITOR
 					]
 				]
-#endif // WITH_EDITOR
 				+ SVerticalBox::Slot()
 				.HAlign(HAlign_Fill)
 				[
@@ -953,21 +964,20 @@ namespace UE::Audio::Insights
 
 		return SNew(SDockTab)
 			.Clipping(EWidgetClipping::ClipToBounds)
-#if WITH_EDITOR
 			.OnTabClosed_Lambda([this](TSharedRef<SDockTab> InDockTab)
 			{				
 				if (PlotsButton.IsValid())
 				{
 					PlotsButton->SetIsChecked(ECheckBoxState::Unchecked);
-
+#if WITH_EDITOR
 					// Can't save layout immediately (it won't save the tab closed state), needs to be done a bit later
 					AsyncTask(ENamedThreads::GameThread, [this]()
 					{
 						SaveLayoutToConfig();
 					});
+#endif // WITH_EDITOR
 				}
 			})
-#endif // WITH_EDITOR
 			[
 				PlotsWidget ? PlotsWidget.ToSharedRef() : SNullWidget::NullWidget
 			];
@@ -986,12 +996,16 @@ namespace UE::Audio::Insights
 			.SetDisplayName(LOCTEXT("MixerSourceTab_MixerSourcesTable_Name", "Mixer Sources"))
 			.SetGroup(MixerSourcesWorkspace.ToSharedRef())
 			.SetMenuType(ETabSpawnerMenuType::Hidden);
-#if WITH_EDITOR
-		MixerSourcesTabManager->RegisterTabSpawner(MixerSourcesPlotsTabName, FOnSpawnTab::CreateSP(this, &FMixerSourceDashboardViewFactory::CreatePlotsTab))
-			.SetDisplayName(LOCTEXT("MixerSourceTab_PlotsTab_Name", "Mixer Sources Plots"))
-			.SetGroup(MixerSourcesWorkspace.ToSharedRef())
-			.SetMenuType(ETabSpawnerMenuType::Hidden);
-#endif // WITH_EDITOR
+#if !WITH_EDITOR
+		const TSharedPtr<const FAudioInsightsComponent> AudioInsightsComponent = FAudioInsightsModule::GetChecked().GetAudioInsightsComponent();
+		if (AudioInsightsComponent.IsValid() && AudioInsightsComponent->GetIsLiveSession())
+#endif // !WITH_EDITOR
+		{
+			MixerSourcesTabManager->RegisterTabSpawner(MixerSourcesPlotsTabName, FOnSpawnTab::CreateSP(this, &FMixerSourceDashboardViewFactory::CreatePlotsTab))
+				.SetDisplayName(LOCTEXT("MixerSourceTab_PlotsTab_Name", "Mixer Sources Plots"))
+				.SetGroup(MixerSourcesWorkspace.ToSharedRef())
+				.SetMenuType(ETabSpawnerMenuType::Hidden);
+		}
 	}
 
 	void FMixerSourceDashboardViewFactory::UnregisterTabSpawners()
@@ -1001,9 +1015,7 @@ namespace UE::Audio::Insights
 		if (MixerSourcesTabManager.IsValid())
 		{
 			MixerSourcesTabManager->UnregisterTabSpawner(MixerSourcesTableTabName);
-#if WITH_EDITOR
 			MixerSourcesTabManager->UnregisterTabSpawner(MixerSourcesPlotsTabName);
-#endif // WITH_EDITOR
 		}
 	}
 
@@ -1011,34 +1023,47 @@ namespace UE::Audio::Insights
 	{
 		using namespace MixerSourcePrivate;
 
-		return FTabManager::NewLayout("MixerSourceTabsLayout_v1")
-		->AddArea
-		(
-			FTabManager::NewPrimaryArea()
+		float SizeCoefficient = 0.7f;
+
+#if !WITH_EDITOR
+		const TSharedPtr<const FAudioInsightsComponent> AudioInsightsComponent = FAudioInsightsModule::GetChecked().GetAudioInsightsComponent();
+		if (AudioInsightsComponent.IsValid() && !AudioInsightsComponent->GetIsLiveSession())
+		{
+			SizeCoefficient = 1.0f;
+		}
+#endif // !WITH_EDITOR
+
+		const TSharedRef<FTabManager::FSplitter> TabSplitter = FTabManager::NewSplitter()
 			->SetOrientation(Orient_Vertical)
+			->SetSizeCoefficient(SizeCoefficient)
 			->Split
 			(
-				FTabManager::NewSplitter()
+				FTabManager::NewStack()
+				->SetSizeCoefficient(SizeCoefficient)
+				->SetHideTabWell(true)
+				->AddTab(MixerSourcesTableTabName, ETabState::OpenedTab)
+			);
+
+#if !WITH_EDITOR
+		if (AudioInsightsComponent.IsValid() && AudioInsightsComponent->GetIsLiveSession())
+#endif // !WITH_EDITOR
+		{
+			TabSplitter->Split(FTabManager::NewStack()
+				->SetSizeCoefficient(0.3f)
+				->AddTab(MixerSourcesPlotsTabName, ETabState::OpenedTab)
+				->SetHideTabWell(true));
+		}
+
+		return FTabManager::NewLayout("MixerSourceTabsLayout_v2")
+			->AddArea
+			(
+				FTabManager::NewPrimaryArea()
 				->SetOrientation(Orient_Vertical)
-				->SetSizeCoefficient(0.7f)
-				->Split
-				(
-					FTabManager::NewStack()
-					->SetSizeCoefficient(0.7f)
-					->SetHideTabWell(true)
-					->AddTab(MixerSourcesTableTabName, ETabState::OpenedTab)
-				)
-#if WITH_EDITOR
 				->Split
 				(
-					FTabManager::NewStack()
-					->SetSizeCoefficient(0.3f)
-					->AddTab(MixerSourcesPlotsTabName, ETabState::OpenedTab)
-					->SetHideTabWell(true)
+					TabSplitter
 				)
-#endif // WITH_EDITOR
-			)
-		);
+			);
 	}
 
 	TSharedRef<SWidget> FMixerSourceDashboardViewFactory::MakeWidget()
diff --git a/Engine/Plugins/AudioInsights/Source/AudioInsights/Public/AudioInsightsStyle.h b/Engine/Plugins/AudioInsights/Source/AudioInsights/Public/AudioInsightsStyle.h
--- a/Engine/Plugins/AudioInsights/Source/AudioInsights/Public/AudioInsightsStyle.h
+++ b/Engine/Plugins/AudioInsights/Source/AudioInsights/Public/AudioInsightsStyle.h
@@ -124,7 +124,7 @@ namespace UE::Audio::Insights
 
 			const FVector2D Icon16(16.0f, 16.0f);
 			const FVector2D Icon20(20.0f, 20.0f);
-			const FVector2D Icon40(40.0f, 40.0f);
+			const FVector2D Icon24(24.0f, 24.0f);
 			const FVector2D Icon64(64.0f, 64.0f);
 
 			Set("AudioInsights.Icon", new IMAGE_BRUSH_SVG(TEXT("Icons/audio_insights_icon"), Icon16));
@@ -132,6 +132,7 @@ namespace UE::Audio::Insights
 			Set("AudioInsights.Icon.Event", new IMAGE_BRUSH_SVG(TEXT("Icons/audio_event"), Icon16));
 			Set("AudioInsights.Icon.Log", new IMAGE_BRUSH_SVG(TEXT("Icons/audio_log"), Icon16));
 			Set("AudioInsights.Icon.Sources", new IMAGE_BRUSH_SVG(TEXT("Icons/audio_sources"), Icon16));
+			Set("AudioInsights.Icon.Sources.Plots", new IMAGE_BRUSH(TEXT("Icons/audio_sources_plots"), Icon24));
 			Set("AudioInsights.Icon.Submix", new IMAGE_BRUSH_SVG(TEXT("Icons/audio_submix"), Icon16));
 			Set("AudioInsights.Icon.VirtualLoop", new IMAGE_BRUSH_SVG(TEXT("Icons/audio_virtualloop"), Icon16));
 			Set("AudioInsights.Icon.Viewport", new IMAGE_BRUSH_SVG(TEXT("Icons/viewport"), Icon16));
diff --git a/Engine/Plugins/AudioInsights/Source/AudioInsights/Public/Views/MixerSourceDashboardViewFactory.h b/Engine/Plugins/AudioInsights/Source/AudioInsights/Public/Views/MixerSourceDashboardViewFactory.h
--- a/Engine/Plugins/AudioInsights/Source/AudioInsights/Public/Views/MixerSourceDashboardViewFactory.h
+++ b/Engine/Plugins/AudioInsights/Source/AudioInsights/Public/Views/MixerSourceDashboardViewFactory.h
@@ -70,16 +70,16 @@ namespace UE::Audio::Insights
 
 #if WITH_EDITOR
 		TSharedRef<SWidget> MakeMuteSoloWidget();
-		TSharedRef<SWidget> MakePlotsButtonWidget();
 
 		void ToggleMuteForAllItems(ECheckBoxState NewState);
 		void ToggleSoloForAllItems(ECheckBoxState NewState);
 
 		TSharedRef<FTabManager::FLayout> LoadLayoutFromConfig();
 		void SaveLayoutToConfig();
+#endif // WITH EDITOR
 
+		TSharedRef<SWidget> MakePlotsButtonWidget();
 		void TogglePlotsTabVisibility(ECheckBoxState InCheckboxState);
-#endif // WITH EDITOR
 
 		TSharedRef<SDockTab> CreateMixerSourcesTab(const FSpawnTabArgs& Args);
 		TSharedRef<SDockTab> CreatePlotsTab(const FSpawnTabArgs& Args);
