[Audio Insights] Fix for mixer sources plot Y axis labels drawing outside background rectangle for large values

diff --git a/Engine/Plugins/AudioInsights/Source/AudioInsights/Private/Views/SAudioCurveView.cpp b/Engine/Plugins/AudioInsights/Source/AudioInsights/Private/Views/SAudioCurveView.cpp
--- a/Engine/Plugins/AudioInsights/Source/AudioInsights/Private/Views/SAudioCurveView.cpp
+++ b/Engine/Plugins/AudioInsights/Source/AudioInsights/Private/Views/SAudioCurveView.cpp
@@ -288,19 +288,9 @@ int32 SAudioCurveView::PaintGridLines(const FGeometry& AllottedGeometry, const F
 
 int32 SAudioCurveView::PaintYAxisLabels(const FGeometry& AllottedGeometry, FSlateWindowElementList& OutDrawElements, int32 LayerId) const
 {
-	// Draw Background rectangle
-	const FVector2D RectangleSize(32.0f, AllottedGeometry.GetLocalSize().Y);
-	const FVector2D RectanglePosition(0.0f, 0.0f);
-	
-	FSlateDrawElement::MakeBox(
-		OutDrawElements,
-		LayerId++,
-		AllottedGeometry.ToPaintGeometry(RectangleSize, FSlateLayoutTransform(RectanglePosition)),
-		FAppStyle::GetBrush("BlackBrush"),
-		ESlateDrawEffect::None,
-		FLinearColor(0.0f, 0.0f, 0.0f, 0.6f)
-	);
-	
+	const int32 BackgroundRectangleLayerId = LayerId++;
+	float MaxTextWidth = 0.0f;
+
 	// Draw Y axis labels
 	const FVector2f Size = AllottedGeometry.GetLocalSize();
 
@@ -325,6 +315,8 @@ int32 SAudioCurveView::PaintYAxisLabels(const FGeometry& AllottedGeometry, FSlat
 			const FVector2f TextSize = FontMeasureService->Measure(LabelString, LabelFont);
 			const FVector2f TextOffset(5.0f, WidgetY - TextSize.Y * 0.85f);
 
+			MaxTextWidth = FMath::Max(MaxTextWidth, TextSize.X);
+
 			FSlateDrawElement::MakeText(
 				OutDrawElements,
 				LayerId++,
@@ -336,6 +328,29 @@ int32 SAudioCurveView::PaintYAxisLabels(const FGeometry& AllottedGeometry, FSlat
 			);
 		}
 	}
+
+	// Draw Background rectangle (with gradient)
+	const float RectanglePadding = MaxTextWidth * 0.6f;
+	const FVector2D RectangleSize(MaxTextWidth + RectanglePadding, AllottedGeometry.GetLocalSize().Y);
+	const FVector2D RectanglePosition(0.0f, 0.0f);
+
+	const TArray<FSlateGradientStop> GradientStops
+	{
+		{FVector2D::ZeroVector,                    FLinearColor(0.0f, 0.0f, 0.0f, 0.8f)},
+		{FVector2D(RectangleSize.X * 0.50f, 0.0f), FLinearColor(0.0f, 0.0f, 0.0f, 0.65f)},
+		{FVector2D(RectangleSize.X * 0.75f, 0.0f), FLinearColor(0.0f, 0.0f, 0.0f, 0.5f)},
+		{FVector2D(RectangleSize.X, 0.0f),         FLinearColor(0.0f, 0.0f, 0.0f, 0.0f)}
+	};
+
+	FSlateDrawElement::MakeGradient(
+		OutDrawElements,
+		BackgroundRectangleLayerId,
+		AllottedGeometry.ToPaintGeometry(RectangleSize, FSlateLayoutTransform(RectanglePosition)),
+		GradientStops,
+		Orient_Vertical,
+		ESlateDrawEffect::None
+	);
+
 	return LayerId;
 }
 
