[Audio Insights] Fix for automatic trace start when a new level is opened in the editor.

diff --git a/Engine/Plugins/AudioInsights/Source/AudioInsightsEditor/Private/AudioInsightsEditorDashboardFactory.cpp b/Engine/Plugins/AudioInsights/Source/AudioInsightsEditor/Private/AudioInsightsEditorDashboardFactory.cpp
--- a/Engine/Plugins/AudioInsights/Source/AudioInsightsEditor/Private/AudioInsightsEditorDashboardFactory.cpp
+++ b/Engine/Plugins/AudioInsights/Source/AudioInsightsEditor/Private/AudioInsightsEditorDashboardFactory.cpp
@@ -62,7 +62,7 @@ namespace UE::Audio::Insights
 	{
 		if (InDeviceId != INDEX_NONE)
 		{
-			if (bStartWithPIE)
+			if (bStartWithPIE && InWorld && InWorld->WorldType == EWorldType::PIE)
 			{
 				const IAudioInsightsTraceModule& TraceModule = FAudioInsightsEditorModule::GetChecked().GetTraceModule();
 				TraceModule.StartTraceAnalysis();
