[Audio Insights] track "cpu" trace channel to get more useful reference info when recording traces.

diff --git a/Engine/Plugins/AudioInsights/Source/AudioInsights/Private/AudioInsightsTraceModule.cpp b/Engine/Plugins/AudioInsights/Source/AudioInsights/Private/AudioInsightsTraceModule.cpp
--- a/Engine/Plugins/AudioInsights/Source/AudioInsights/Private/AudioInsightsTraceModule.cpp
+++ b/Engine/Plugins/AudioInsights/Source/AudioInsights/Private/AudioInsightsTraceModule.cpp
@@ -63,6 +63,7 @@ namespace UE::Audio::Insights
 
 	void FTraceModule::EnableAudioInsightsTraceChannels()
 	{
+		UE::Trace::ToggleChannel(TEXT("Cpu"), true);
 		UE::Trace::ToggleChannel(TEXT("Audio"), true);
 		UE::Trace::ToggleChannel(TEXT("AudioMixer"), true);
 	}
