[Audio Insights] Fix for: Output Meter and Oscilloscope do not display values during PIE until the PIE session is reselected in the World Filter in Binary Builds

diff --git a/Engine/Plugins/AudioInsights/Source/AudioInsightsEditor/Private/AudioInsightsEditorDashboardFactory.cpp b/Engine/Plugins/AudioInsights/Source/AudioInsightsEditor/Private/AudioInsightsEditorDashboardFactory.cpp
--- a/Engine/Plugins/AudioInsights/Source/AudioInsightsEditor/Private/AudioInsightsEditorDashboardFactory.cpp
+++ b/Engine/Plugins/AudioInsights/Source/AudioInsightsEditor/Private/AudioInsightsEditorDashboardFactory.cpp
@@ -82,6 +82,11 @@ namespace UE::Audio::Insights
 		RefreshDeviceSelector();
 	}
 
+	void FEditorDashboardFactory::OnWorldUnregisteredFromAudioDevice(const UWorld* InWorld, ::Audio::FDeviceId InDeviceId)
+	{
+		RefreshDeviceSelector();
+	}
+
 	void FEditorDashboardFactory::OnPIEStarted(bool bSimulating)
 	{
 		if (bStartWithPIE)
@@ -93,7 +98,8 @@ namespace UE::Audio::Insights
 
 	void FEditorDashboardFactory::OnPostPIEStarted(bool bSimulating)
 	{
-		OnActiveAudioDeviceChanged.Broadcast();
+		int32 x = 0;
+		++x;
 	}
 
 	void FEditorDashboardFactory::OnPIEStopped(bool bSimulating)
@@ -107,9 +113,9 @@ namespace UE::Audio::Insights
 		RefreshDeviceSelector();
 	}
 
-	void FEditorDashboardFactory::OnWorldUnregisteredFromAudioDevice(const UWorld* InWorld, ::Audio::FDeviceId InDeviceId)
+	void FEditorDashboardFactory::OnDeviceCreated(::Audio::FDeviceId InDeviceId)
 	{
-		RefreshDeviceSelector();
+		OnActiveAudioDeviceChanged.Broadcast();
 	}
 
 	void FEditorDashboardFactory::OnDeviceDestroyed(::Audio::FDeviceId InDeviceId)
@@ -174,6 +180,12 @@ namespace UE::Audio::Insights
 			OnWorldUnregisteredFromAudioDeviceHandle.Reset();
 		}
 
+		if (OnDeviceCreatedHandle.IsValid())
+		{
+			FAudioDeviceManagerDelegates::OnAudioDeviceCreated.Remove(OnDeviceCreatedHandle);
+			OnDeviceCreatedHandle.Reset();
+		}
+
 		if (OnDeviceDestroyedHandle.IsValid())
 		{
 			FAudioDeviceManagerDelegates::OnAudioDeviceDestroyed.Remove(OnDeviceDestroyedHandle);
@@ -511,6 +523,11 @@ namespace UE::Audio::Insights
 			OnWorldUnregisteredFromAudioDeviceHandle = FAudioDeviceWorldDelegates::OnWorldUnregisteredWithAudioDevice.AddSP(this, &FEditorDashboardFactory::OnWorldUnregisteredFromAudioDevice);
 		}
 
+		if (!OnDeviceCreatedHandle.IsValid())
+		{
+			OnDeviceCreatedHandle = FAudioDeviceManagerDelegates::OnAudioDeviceCreated.AddSP(this, &FEditorDashboardFactory::OnDeviceCreated);
+		}
+
 		if (!OnDeviceDestroyedHandle.IsValid())
 		{
 			OnDeviceDestroyedHandle = FAudioDeviceManagerDelegates::OnAudioDeviceDestroyed.AddSP(this, &FEditorDashboardFactory::OnDeviceDestroyed);
diff --git a/Engine/Plugins/AudioInsights/Source/AudioInsightsEditor/Private/AudioInsightsEditorDashboardFactory.h b/Engine/Plugins/AudioInsights/Source/AudioInsightsEditor/Private/AudioInsightsEditorDashboardFactory.h
--- a/Engine/Plugins/AudioInsights/Source/AudioInsightsEditor/Private/AudioInsightsEditorDashboardFactory.h
+++ b/Engine/Plugins/AudioInsights/Source/AudioInsightsEditor/Private/AudioInsightsEditorDashboardFactory.h
@@ -31,6 +31,8 @@ namespace UE::Audio::Insights
 	private:
 		void OnWorldRegisteredToAudioDevice(const UWorld* InWorld, ::Audio::FDeviceId InDeviceId);
 		void OnWorldUnregisteredFromAudioDevice(const UWorld* InWorld, ::Audio::FDeviceId InDeviceId);
+
+		void OnDeviceCreated(::Audio::FDeviceId InDeviceId);
 		void OnDeviceDestroyed(::Audio::FDeviceId InDeviceId);
 
 		void OnPIEStarted(bool bSimulating);
@@ -52,10 +54,12 @@ namespace UE::Audio::Insights
 		TSharedRef<FTabManager::FLayout> LoadLayoutFromConfig();
 		void SaveLayoutToConfig();
 
-		FDelegateHandle OnDeviceDestroyedHandle;
 		FDelegateHandle OnWorldRegisteredToAudioDeviceHandle;
 		FDelegateHandle OnWorldUnregisteredFromAudioDeviceHandle;
 
+		FDelegateHandle OnDeviceCreatedHandle;
+		FDelegateHandle OnDeviceDestroyedHandle;
+
 		FDelegateHandle OnPIEStartedHandle;
 		FDelegateHandle OnPostPIEStartedHandle;
 		FDelegateHandle OnPIEStoppedHandle;
