[Audio Insights] Removed leftovers from: "Fix for: Output Meter and Oscilloscope do not display values during PIE until the PIE session is reselected in the World Filter in Binary Builds"

diff --git a/Engine/Plugins/AudioInsights/Source/AudioInsightsEditor/Private/AudioInsightsEditorDashboardFactory.cpp b/Engine/Plugins/AudioInsights/Source/AudioInsightsEditor/Private/AudioInsightsEditorDashboardFactory.cpp
--- a/Engine/Plugins/AudioInsights/Source/AudioInsightsEditor/Private/AudioInsightsEditorDashboardFactory.cpp
+++ b/Engine/Plugins/AudioInsights/Source/AudioInsightsEditor/Private/AudioInsightsEditorDashboardFactory.cpp
@@ -96,12 +96,6 @@ namespace UE::Audio::Insights
 		}
 	}
 
-	void FEditorDashboardFactory::OnPostPIEStarted(bool bSimulating)
-	{
-		int32 x = 0;
-		++x;
-	}
-
 	void FEditorDashboardFactory::OnPIEStopped(bool bSimulating)
 	{
 		if (bStopWithPIE)
@@ -198,12 +192,6 @@ namespace UE::Audio::Insights
 			OnPIEStartedHandle.Reset();
 		}
 
-		if (OnPostPIEStartedHandle.IsValid())
-		{
-			FEditorDelegates::PostPIEStarted.Remove(OnPostPIEStartedHandle);
-			OnPostPIEStartedHandle.Reset();
-		}
-
 		if (OnPIEStoppedHandle.IsValid())
 		{
 			FEditorDelegates::EndPIE.Remove(OnPIEStoppedHandle);
@@ -538,11 +526,6 @@ namespace UE::Audio::Insights
 			OnPIEStartedHandle = FEditorDelegates::PreBeginPIE.AddSP(this, &FEditorDashboardFactory::OnPIEStarted);
 		}
 
-		if (!OnPostPIEStartedHandle.IsValid())
-		{
-			OnPostPIEStartedHandle = FEditorDelegates::PostPIEStarted.AddSP(this, &FEditorDashboardFactory::OnPostPIEStarted);
-		}
-
 		if (!OnPIEStoppedHandle.IsValid())
 		{
 			OnPIEStoppedHandle = FEditorDelegates::EndPIE.AddSP(this, &FEditorDashboardFactory::OnPIEStopped);
diff --git a/Engine/Plugins/AudioInsights/Source/AudioInsightsEditor/Private/AudioInsightsEditorDashboardFactory.h b/Engine/Plugins/AudioInsights/Source/AudioInsightsEditor/Private/AudioInsightsEditorDashboardFactory.h
--- a/Engine/Plugins/AudioInsights/Source/AudioInsightsEditor/Private/AudioInsightsEditorDashboardFactory.h
+++ b/Engine/Plugins/AudioInsights/Source/AudioInsightsEditor/Private/AudioInsightsEditorDashboardFactory.h
@@ -36,7 +36,6 @@ namespace UE::Audio::Insights
 		void OnDeviceDestroyed(::Audio::FDeviceId InDeviceId);
 
 		void OnPIEStarted(bool bSimulating);
-		void OnPostPIEStarted(bool bSimulating);
 		void OnPIEStopped(bool bSimulating);
 
 		void RefreshDeviceSelector();
@@ -61,7 +60,6 @@ namespace UE::Audio::Insights
 		FDelegateHandle OnDeviceDestroyedHandle;
 
 		FDelegateHandle OnPIEStartedHandle;
-		FDelegateHandle OnPostPIEStartedHandle;
 		FDelegateHandle OnPIEStoppedHandle;
 
 		bool bStartWithPIE = true;
