[Audio Insights] Enable Plots for Offline Traces in Standalone

diff --git a/Engine/Plugins/AudioInsights/Source/AudioInsights/Private/AudioInsightsComponent.cpp b/Engine/Plugins/AudioInsights/Source/AudioInsights/Private/AudioInsightsComponent.cpp
--- a/Engine/Plugins/AudioInsights/Source/AudioInsights/Private/AudioInsightsComponent.cpp
+++ b/Engine/Plugins/AudioInsights/Source/AudioInsights/Private/AudioInsightsComponent.cpp
@@ -3,6 +3,7 @@
 
 #include "AudioInsightsModule.h"
 #include "AudioInsightsStyle.h"
+#include "Insights/IInsightsManager.h"
 #include "Modules/ModuleManager.h"
 #include "Trace/StoreClient.h"
 #include "TraceServices/Model/Diagnostics.h"
@@ -48,6 +49,14 @@ namespace UE::Audio::Insights
 
 			constexpr float TickDelay = 0.5f; // 500 ms. delay between ticks
 			OnTickHandle = FTSTicker::GetCoreTicker().AddTicker(OnTick, TickDelay);
+
+			IUnrealInsightsModule& UnrealInsightsModule = FModuleManager::LoadModuleChecked<IUnrealInsightsModule>("TraceInsights");
+
+			TSharedPtr<::Insights::IInsightsManager> InsightsManager = UnrealInsightsModule.GetInsightsManager();
+			if (InsightsManager.IsValid())
+			{
+				InsightsManager->GetSessionAnalysisCompletedEvent().AddSP(this, &FAudioInsightsComponent::OnSessionAnalysisCompletedEvent);
+			}
 		}
 	}
 
@@ -113,6 +122,11 @@ namespace UE::Audio::Insights
 		return DockTab;
 	}
 
+	void FAudioInsightsComponent::OnSessionAnalysisCompletedEvent()
+	{
+		OnSessionAnalysisCompleted.Broadcast();
+	}
+
 	bool FAudioInsightsComponent::GetIsLiveSession() const
 	{
 		IUnrealInsightsModule& UnrealInsightsModule = FModuleManager::LoadModuleChecked<IUnrealInsightsModule>("TraceInsights");
@@ -131,6 +145,20 @@ namespace UE::Audio::Insights
 		return false;
 	}
 
+	bool FAudioInsightsComponent::IsSessionAnalysisComplete() const
+	{
+		const IUnrealInsightsModule& UnrealInsightsModule = FModuleManager::LoadModuleChecked<IUnrealInsightsModule>("TraceInsights");
+
+		TSharedPtr<const TraceServices::IAnalysisSession> Session = UnrealInsightsModule.GetAnalysisSession();
+		if (Session.IsValid())
+		{
+			TraceServices::FAnalysisSessionReadScope SessionReadScope(*Session.Get());
+			return Session->IsAnalysisComplete();
+		}
+
+		return false;
+	}
+
 	bool FAudioInsightsComponent::Tick(float DeltaTime)
 	{
 		if (!bCanSpawnTab)
diff --git a/Engine/Plugins/AudioInsights/Source/AudioInsights/Private/AudioInsightsComponent.h b/Engine/Plugins/AudioInsights/Source/AudioInsights/Private/AudioInsightsComponent.h
--- a/Engine/Plugins/AudioInsights/Source/AudioInsights/Private/AudioInsightsComponent.h
+++ b/Engine/Plugins/AudioInsights/Source/AudioInsights/Private/AudioInsightsComponent.h
@@ -25,15 +25,21 @@ namespace UE::Audio::Insights
 		virtual void UnregisterMajorTabs() override;
 
 		bool GetIsLiveSession() const;
+		bool IsSessionAnalysisComplete() const;
 		bool GetIsEditorTrace() const { return bIsEditorTrace; }
 
 		DECLARE_MULTICAST_DELEGATE(FOnTabSpawn);
-		inline static FOnTabSpawn OnTabSpawn;
+		FOnTabSpawn OnTabSpawn;
+
+		DECLARE_MULTICAST_DELEGATE(FOnSessionAnalysisCompleted);
+		FOnSessionAnalysisCompleted OnSessionAnalysisCompleted;
 
 	private:
 		bool CanSpawnTab(const FSpawnTabArgs& Args) const;
 		TSharedRef<SDockTab> SpawnTab(const FSpawnTabArgs& Args);
 
+		void OnSessionAnalysisCompletedEvent();
+
 		bool Tick(float DeltaTime);
 
 		bool bIsInitialized = false;
diff --git a/Engine/Plugins/AudioInsights/Source/AudioInsights/Private/Views/MixerSourceDashboardViewFactory.cpp b/Engine/Plugins/AudioInsights/Source/AudioInsights/Private/Views/MixerSourceDashboardViewFactory.cpp
--- a/Engine/Plugins/AudioInsights/Source/AudioInsights/Private/Views/MixerSourceDashboardViewFactory.cpp
+++ b/Engine/Plugins/AudioInsights/Source/AudioInsights/Private/Views/MixerSourceDashboardViewFactory.cpp
@@ -19,6 +19,7 @@
 #include "Editor.h"
 #else
 #include "AudioInsightsComponent.h"
+#include "AudioInsightsTimingViewExtender.h"
 #endif // WITH_EDITOR
 
 #define LOCTEXT_NAMESPACE "AudioInsights"
@@ -56,8 +57,13 @@ namespace UE::Audio::Insights
 
 	} // namespace MixerSourcePrivate
 
-	const double FMixerSourceDashboardViewFactory::MaxPlotHistorySeconds = 5.0;
-	const int32 FMixerSourceDashboardViewFactory::MaxPlotSources = 16;
+	constexpr double FMixerSourceDashboardViewFactory::MaxPlotHistorySeconds = 5.0;
+
+#if WITH_EDITOR
+	constexpr int32 FMixerSourceDashboardViewFactory::MaxPlotSources = 16;
+#else
+	constexpr int32 FMixerSourceDashboardViewFactory::MaxPlotSources = 64;
+#endif // WITH_EDITOR
 
 	FMixerSourceDashboardViewFactory::FMixerSourceDashboardViewFactory()
 	{
@@ -432,6 +438,21 @@ namespace UE::Audio::Insights
 			GameState = AudioInsightsComponent->GetIsLiveSession() ? EGameState::Running : EGameState::Stopped;
 		}
 	}
+
+	void FMixerSourceDashboardViewFactory::OnSessionAnalysisCompleted()
+	{
+		GameState = EGameState::Stopped;
+	}
+
+	void FMixerSourceDashboardViewFactory::OnTimingViewTimeMarkerChanged(double TimeMarker)
+	{
+		CurrentTimestamp = TimeMarker;
+
+		for (TSharedPtr<SAudioCurveView> PlotWidget : PlotWidgets)
+		{
+			PlotWidget->UpdateYDataRangeFromTimestampRange(CurrentTimestamp - MaxPlotHistorySeconds, CurrentTimestamp);
+		}
+	}
 #endif // WITH_EDITOR
 
 #if WITH_EDITOR
@@ -504,7 +525,11 @@ namespace UE::Audio::Insights
 						BeginTimestamp   = FMath::Min(BeginTimestamp, Timestamp);
 						CurrentTimestamp = FMath::Max(CurrentTimestamp, Timestamp);
 
+#if WITH_EDITOR
 						const double DataPointTime = Timestamp - BeginTimestamp;
+#else
+						const double DataPointTime = Timestamp;
+#endif // WITH_EDITOR
 
 						DataPoints->Emplace(DataPointTime, Value);
 					}
@@ -515,7 +540,6 @@ namespace UE::Audio::Insights
 			if (bCanAddNewSources && !PlotWidgetMetadataPerCurve->Contains(SourceId))
 			{
 				FPlotCurveMetadata& NewMetadata = PlotWidgetMetadataPerCurve->Add(SourceId);
-				NewMetadata.CurveId = SourceId;
 				NewMetadata.CurveColor = FLinearColor(FColor::MakeRandomColor()); 
 				NewMetadata.DisplayName = FText::FromString(*FSoftObjectPath(SourceDataPoint.Name).GetAssetName());
 				bHasNewMetadata = true;
@@ -531,8 +555,10 @@ namespace UE::Audio::Insights
 			}
 		}
 
-		// Remove old points and set curve data for each widget 
+		// Remove old points and set curve data for each widget
+#if WITH_EDITOR
 		const double PlotDrawLimitTimestamp = CurrentTimestamp - BeginTimestamp - (MaxPlotHistorySeconds + 0.2 /* extra grace time to avoid curve cuts being displayed */);
+#endif // WITH_EDITOR
 
 		for (int32 WidgetIndex = 0; WidgetIndex < NumPlotWidgets; ++WidgetIndex)
 		{
@@ -541,6 +567,7 @@ namespace UE::Audio::Insights
 			const TSharedPtr<FPointDataPerCurveMap> CurveDataMapPtr = *PlotWidgetCurveIdToPointDataMapPerColumn.Find(SelectedPlotColumn);
 			if (CurveDataMapPtr.IsValid())
 			{
+#if WITH_EDITOR
 				// Remove points that are older than max history limit from the most recent timestamp
 				for (auto& [CurveId, CurvePoints] : *CurveDataMapPtr)
 				{
@@ -554,12 +581,36 @@ namespace UE::Audio::Insights
 						CurvePoints.RemoveAt(0, FoundIndex, EAllowShrinking::No);
 					}
 				}
+#endif // WITH_EDITOR
 
 				PlotWidgets[WidgetIndex]->SetCurvesPointData(CurveDataMapPtr);
 			}
 		}
 	}
 
+#if !WITH_EDITOR
+	void FMixerSourceDashboardViewFactory::FilterOfflinePlots(const FString& InFilterString, TMap<int32, FPlotCurveMetadata>& InPlotWidgetMetadataPerCurve)
+	{
+		// Reset plots visibility
+		for (auto& [CurveId, CurveMetadata] : InPlotWidgetMetadataPerCurve)
+		{
+			CurveMetadata.CurveColor.A = 1.0f;
+		}
+
+		// Plots that don't match the filter will become transparent
+		if (!InFilterString.IsEmpty())
+		{
+			for (auto& [CurveId, CurveMetadata] : InPlotWidgetMetadataPerCurve)
+			{
+				if (!CurveMetadata.DisplayName.ToString().Contains(InFilterString))
+				{
+					CurveMetadata.CurveColor.A = 0.0f;
+				}
+			}
+		}
+	}
+#endif // !WITH_EDITOR
+
 	const TMap<FName, FMixerSourceDashboardViewFactory::FPlotColumnInfo>& FMixerSourceDashboardViewFactory::GetPlotColumnInfo()
 	{
 		auto CreatePlotColumnInfo = []()
@@ -662,7 +713,11 @@ namespace UE::Audio::Insights
 		// Create plot widgets
 		auto GetViewRange = [this]()
 		{
+#if WITH_EDITOR
 			if (GameState == EGameState::Stopped || BeginTimestamp == TNumericLimits<double>::Max())
+#else
+			if (BeginTimestamp == TNumericLimits<double>::Max())
+#endif // WITH_EDITOR
 			{
 				return TRange<double>(0, MaxPlotHistorySeconds);
 			}
@@ -679,6 +734,14 @@ namespace UE::Audio::Insights
 				CurrentTimestamp += TimestampsDiff;
 			}
 
+#if !WITH_EDITOR
+			const TSharedPtr<const FAudioInsightsComponent> AudioInsightsComponent = FAudioInsightsModule::GetChecked().GetAudioInsightsComponent();
+			if (AudioInsightsComponent.IsValid() && !AudioInsightsComponent->GetIsLiveSession())
+			{
+				return TRange<double>(CurrentTimestamp - MaxPlotHistorySeconds, CurrentTimestamp);
+			}
+#endif // !WITH_EDITOR
+
 			return TRange<double>(CurrentTimestamp - MaxPlotHistorySeconds - BeginTimestamp, CurrentTimestamp - BeginTimestamp);
 		};
 
@@ -910,10 +973,6 @@ namespace UE::Audio::Insights
 
 	TSharedRef<SDockTab> FMixerSourceDashboardViewFactory::CreateMixerSourcesTab(const FSpawnTabArgs& Args)
 	{
-#if !WITH_EDITOR
-		const TSharedPtr<const FAudioInsightsComponent> AudioInsightsComponent = FAudioInsightsModule::GetChecked().GetAudioInsightsComponent();
-#endif // !WITH_EDITOR
-
 		return SNew(SDockTab)
 			.Clipping(EWidgetClipping::ClipToBounds)
 			[
@@ -940,11 +999,7 @@ namespace UE::Audio::Insights
 					.AutoWidth()
 					.HAlign(HAlign_Right)
 					[
-#if WITH_EDITOR
 						MakePlotsButtonWidget()
-#else
-						AudioInsightsComponent.IsValid() && AudioInsightsComponent->GetIsLiveSession() ? MakePlotsButtonWidget() : SNullWidget::NullWidget
-#endif // WITH_EDITOR
 					]
 				]
 				+ SVerticalBox::Slot()
@@ -996,16 +1051,11 @@ namespace UE::Audio::Insights
 			.SetDisplayName(LOCTEXT("MixerSourceTab_MixerSourcesTable_Name", "Mixer Sources"))
 			.SetGroup(MixerSourcesWorkspace.ToSharedRef())
 			.SetMenuType(ETabSpawnerMenuType::Hidden);
-#if !WITH_EDITOR
-		const TSharedPtr<const FAudioInsightsComponent> AudioInsightsComponent = FAudioInsightsModule::GetChecked().GetAudioInsightsComponent();
-		if (AudioInsightsComponent.IsValid() && AudioInsightsComponent->GetIsLiveSession())
-#endif // !WITH_EDITOR
-		{
-			MixerSourcesTabManager->RegisterTabSpawner(MixerSourcesPlotsTabName, FOnSpawnTab::CreateSP(this, &FMixerSourceDashboardViewFactory::CreatePlotsTab))
-				.SetDisplayName(LOCTEXT("MixerSourceTab_PlotsTab_Name", "Mixer Sources Plots"))
-				.SetGroup(MixerSourcesWorkspace.ToSharedRef())
-				.SetMenuType(ETabSpawnerMenuType::Hidden);
-		}
+
+		MixerSourcesTabManager->RegisterTabSpawner(MixerSourcesPlotsTabName, FOnSpawnTab::CreateSP(this, &FMixerSourceDashboardViewFactory::CreatePlotsTab))
+			.SetDisplayName(LOCTEXT("MixerSourceTab_PlotsTab_Name", "Mixer Sources Plots"))
+			.SetGroup(MixerSourcesWorkspace.ToSharedRef())
+			.SetMenuType(ETabSpawnerMenuType::Hidden);
 	}
 
 	void FMixerSourceDashboardViewFactory::UnregisterTabSpawners()
@@ -1023,37 +1073,6 @@ namespace UE::Audio::Insights
 	{
 		using namespace MixerSourcePrivate;
 
-		float SizeCoefficient = 0.7f;
-
-#if !WITH_EDITOR
-		const TSharedPtr<const FAudioInsightsComponent> AudioInsightsComponent = FAudioInsightsModule::GetChecked().GetAudioInsightsComponent();
-		if (AudioInsightsComponent.IsValid() && !AudioInsightsComponent->GetIsLiveSession())
-		{
-			SizeCoefficient = 1.0f;
-		}
-#endif // !WITH_EDITOR
-
-		const TSharedRef<FTabManager::FSplitter> TabSplitter = FTabManager::NewSplitter()
-			->SetOrientation(Orient_Vertical)
-			->SetSizeCoefficient(SizeCoefficient)
-			->Split
-			(
-				FTabManager::NewStack()
-				->SetSizeCoefficient(SizeCoefficient)
-				->SetHideTabWell(true)
-				->AddTab(MixerSourcesTableTabName, ETabState::OpenedTab)
-			);
-
-#if !WITH_EDITOR
-		if (AudioInsightsComponent.IsValid() && AudioInsightsComponent->GetIsLiveSession())
-#endif // !WITH_EDITOR
-		{
-			TabSplitter->Split(FTabManager::NewStack()
-				->SetSizeCoefficient(0.3f)
-				->AddTab(MixerSourcesPlotsTabName, ETabState::OpenedTab)
-				->SetHideTabWell(true));
-		}
-
 		return FTabManager::NewLayout("MixerSourceTabsLayout_v2")
 			->AddArea
 			(
@@ -1061,7 +1080,23 @@ namespace UE::Audio::Insights
 				->SetOrientation(Orient_Vertical)
 				->Split
 				(
-					TabSplitter
+					FTabManager::NewSplitter()
+					->SetOrientation(Orient_Vertical)
+					->SetSizeCoefficient(0.7f)
+					->Split
+					(
+						FTabManager::NewStack()
+						->SetSizeCoefficient(0.7f)
+						->SetHideTabWell(true)
+						->AddTab(MixerSourcesTableTabName, ETabState::OpenedTab)
+					)
+					->Split
+					(
+						FTabManager::NewStack()
+						->SetSizeCoefficient(0.3f)
+						->AddTab(MixerSourcesPlotsTabName, ETabState::OpenedTab)
+						->SetHideTabWell(true)
+					)
 				)
 			);
 	}
@@ -1076,7 +1111,16 @@ namespace UE::Audio::Insights
 		FEditorDelegates::PausePIE.AddSP(this, &FMixerSourceDashboardViewFactory::OnPIEPaused);
 		FEditorDelegates::ResumePIE.AddSP(this, &FMixerSourceDashboardViewFactory::OnPIEResumed);
 #else
-		FAudioInsightsComponent::OnTabSpawn.AddSP(this, &FMixerSourceDashboardViewFactory::OnAudioInsightsComponentTabSpawn);
+		FAudioInsightsModule& AudioInsightsModule = FAudioInsightsModule::GetChecked();
+
+		AudioInsightsModule.GetTimingViewExtender().OnTimingViewTimeMarkerChanged.AddSP(this, &FMixerSourceDashboardViewFactory::OnTimingViewTimeMarkerChanged);
+
+		const TSharedPtr<FAudioInsightsComponent> AudioInsightsComponent = AudioInsightsModule.GetAudioInsightsComponent();
+		if (AudioInsightsComponent.IsValid())
+		{
+			AudioInsightsComponent->OnTabSpawn.AddSP(this, &FMixerSourceDashboardViewFactory::OnAudioInsightsComponentTabSpawn);
+			AudioInsightsComponent->OnSessionAnalysisCompleted.AddSP(this, &FMixerSourceDashboardViewFactory::OnSessionAnalysisCompleted);
+		}
 #endif // WITH_EDITOR
 
 		PlotsWidget = MakePlotsWidget();
@@ -1139,7 +1183,25 @@ namespace UE::Audio::Insights
 			return true;
 		});
 
+#if WITH_EDITOR
 		UpdatePlotsWidgetsData();
+#else
+		const TSharedPtr<const FAudioInsightsComponent> AudioInsightsComponent = FAudioInsightsModule::GetChecked().GetAudioInsightsComponent();
+		if (AudioInsightsComponent.IsValid())
+		{
+			if (!AudioInsightsComponent->IsSessionAnalysisComplete())
+			{
+				UpdatePlotsWidgetsData();
+			}
+			else
+			{
+				if (InReason == EProcessReason::FilterUpdated && PlotWidgetMetadataPerCurve.IsValid())
+				{
+					FilterOfflinePlots(FilterString, *PlotWidgetMetadataPerCurve);
+				}
+			}
+		}
+#endif // WITH_EDITOR
 
 #if WITH_EDITOR
 		// Update the mute and solo states if the filter string changes
diff --git a/Engine/Plugins/AudioInsights/Source/AudioInsights/Private/Views/SAudioCurveView.cpp b/Engine/Plugins/AudioInsights/Source/AudioInsights/Private/Views/SAudioCurveView.cpp
--- a/Engine/Plugins/AudioInsights/Source/AudioInsights/Private/Views/SAudioCurveView.cpp
+++ b/Engine/Plugins/AudioInsights/Source/AudioInsights/Private/Views/SAudioCurveView.cpp
@@ -100,10 +100,13 @@ void SAudioCurveView::UpdateCurveToolTip(const FGeometry& InMyGeometry, const FP
 		FCurvePoint ClosestPoint;
 		float ClosestDistance = TNumericLimits<float>::Max();
 
-		for (auto Iter = MetadataPerCurve->CreateConstIterator(); Iter; ++Iter)
+		for (const auto& [CurveId, CurveMetadata] : *MetadataPerCurve)
 		{
-			const FCurveMetadata& CurveMetadata = Iter->Value;
-			const int32 CurveId = CurveMetadata.CurveId;
+			if (CurveMetadata.CurveColor.A == 0.0f)
+			{
+				continue;
+			}
+
 			const TArray<FCurvePoint>* CurvePoints = PointDataPerCurve->Find(CurveId);
 			const int32 NumPoints = CurvePoints->Num();
 			if (NumPoints > 0)
@@ -180,6 +183,64 @@ void SAudioCurveView::UpdateYDataRange()
 	YDataRange.Y = MaxValue;
 }
 
+#if !WITH_EDITOR
+void SAudioCurveView::UpdateYDataRangeFromTimestampRange(const double InLowerBoundTimestamp, const double InUpperBoundTimestamp)
+{
+	const static auto FCurvePointCompare = [](const FCurvePoint& A, const FCurvePoint& B)
+		{
+			return A.Value < B.Value;
+		};
+
+	float MinValue = TNumericLimits<float>::Max();
+	float MaxValue = TNumericLimits<float>::Min();
+
+	if (PointDataPerCurve.IsValid())
+	{
+		for (const auto& [CurveId, CurvePoints] : *PointDataPerCurve)
+		{
+			if (!CurvePoints.IsEmpty())
+			{
+				const int32 FoundLowerIndex = CurvePoints.IndexOfByPredicate([&InLowerBoundTimestamp](const FCurvePoint& InDataPoint)
+				{
+					return InDataPoint.Key >= InLowerBoundTimestamp;
+				});
+
+				const int32 FoundUpperIndex = CurvePoints.IndexOfByPredicate([&InUpperBoundTimestamp](const FCurvePoint& InDataPoint)
+				{
+					return InDataPoint.Key >= InUpperBoundTimestamp;
+				});
+
+				if (FoundLowerIndex != INDEX_NONE && FoundUpperIndex != INDEX_NONE && FoundLowerIndex != FoundUpperIndex)
+				{
+					const TArrayView<const FCurvePoint> CurvePointsArrayView = MakeArrayView(CurvePoints.GetData() + FoundLowerIndex, FoundUpperIndex - FoundLowerIndex);
+
+					const FCurvePoint* MinPoint = Algo::MinElement(CurvePointsArrayView, FCurvePointCompare);
+					const FCurvePoint* MaxPoint = Algo::MaxElement(CurvePointsArrayView, FCurvePointCompare);
+					MinValue = MinPoint ? FMath::Min(MinValue, MinPoint->Value) : 0.0f;
+					MaxValue = MaxPoint ? FMath::Max(MaxValue, MaxPoint->Value) : 1.0f;
+				}
+			}
+		}
+	}
+
+	// Adjust Y range in case MinValue and MaxValue are not updated or if values are too close together
+	constexpr float Epsilon = 0.001f;
+
+	if (MinValue == TNumericLimits<float>::Max())
+	{
+		MinValue = 0.0f;
+	}
+
+	if (MaxValue == TNumericLimits<float>::Min() || FMath::Abs(MaxValue - MinValue) < Epsilon)
+	{
+		MaxValue = MinValue + Epsilon;
+	}
+
+	YDataRange.X = MinValue;
+	YDataRange.Y = MaxValue;
+}
+#endif // !WITH_EDITOR
+
 void SAudioCurveView::SetCurvesPointData(TSharedPtr<TMap<int32, TArray<FCurvePoint>>> InPointDataPerCurve)
 {
 	PointDataPerCurve = InPointDataPerCurve;
@@ -375,10 +436,9 @@ int32 SAudioCurveView::PaintCurves(const FGeometry& AllottedGeometry, const FSla
 	const FVector2f Size = AllottedGeometry.GetLocalSize();
 
 	// Create and draw points per curve
-	for (auto Iter = MetadataPerCurve->CreateConstIterator(); Iter; ++Iter)
+	for (const auto& [CurveId, CurveMetadata] : *MetadataPerCurve)
 	{
-		const FCurveMetadata& CurveMetadata = Iter->Value;
-		const TArray<FCurvePoint>* CurvePointsPtr = PointDataPerCurve->Find(CurveMetadata.CurveId);
+		const TArray<FCurvePoint>* CurvePointsPtr = PointDataPerCurve->Find(CurveId);
 
 		if (!CurvePointsPtr || CurvePointsPtr->IsEmpty())
 		{
diff --git a/Engine/Plugins/AudioInsights/Source/AudioInsights/Public/Views/MixerSourceDashboardViewFactory.h b/Engine/Plugins/AudioInsights/Source/AudioInsights/Public/Views/MixerSourceDashboardViewFactory.h
--- a/Engine/Plugins/AudioInsights/Source/AudioInsights/Public/Views/MixerSourceDashboardViewFactory.h
+++ b/Engine/Plugins/AudioInsights/Source/AudioInsights/Public/Views/MixerSourceDashboardViewFactory.h
@@ -45,7 +45,9 @@ namespace UE::Audio::Insights
 		TSharedRef<SWidget> MakePlotsWidget();
 
 	private:
-		void ResetPlots();
+		using FPlotCurvePoint = SAudioCurveView::FCurvePoint;
+		using FPointDataPerCurveMap = TMap<int32, TArray<FPlotCurvePoint>>; // Map of source id to data point array 
+		using FPlotCurveMetadata = SAudioCurveView::FCurveMetadata;
 
 #if WITH_EDITOR
 		void OnPIEStarted(bool bSimulating);
@@ -54,10 +56,17 @@ namespace UE::Audio::Insights
 		void OnPIEResumed(bool bSimulating);
 #else
 		void OnAudioInsightsComponentTabSpawn();
+		void OnSessionAnalysisCompleted();
+		void OnTimingViewTimeMarkerChanged(double InTimeMarker);
 #endif // WITH_EDITOR
 
+		void ResetPlots();
 		void UpdatePlotsWidgetsData();
 
+#if !WITH_EDITOR
+		void FilterOfflinePlots(const FString& InFilterString, TMap<int32, FPlotCurveMetadata>& InPlotWidgetMetadataPerCurve);
+#endif // !WITH_EDITOR
+
 #if WITH_EDITOR
 		void UpdateMuteSoloState();
 #endif // WITH_EDITOR
@@ -89,11 +98,6 @@ namespace UE::Audio::Insights
 
 		TSharedRef<FTabManager::FLayout> GetDefaultTabLayout();
 
-		using FPlotCurvePoint = SAudioCurveView::FCurvePoint;
-		// Map of source id to data point array 
-		using FPointDataPerCurveMap = TMap<int32, TArray<FPlotCurvePoint>>;
-		using FPlotCurveMetadata = SAudioCurveView::FCurveMetadata;
-
 		FCheckBoxStyle MuteToggleButtonStyle;
 		FCheckBoxStyle SoloToggleButtonStyle;
 
diff --git a/Engine/Plugins/AudioInsights/Source/AudioInsights/Public/Views/SAudioCurveView.h b/Engine/Plugins/AudioInsights/Source/AudioInsights/Public/Views/SAudioCurveView.h
--- a/Engine/Plugins/AudioInsights/Source/AudioInsights/Public/Views/SAudioCurveView.h
+++ b/Engine/Plugins/AudioInsights/Source/AudioInsights/Public/Views/SAudioCurveView.h
@@ -20,7 +20,6 @@ public:
 
 	struct FCurveMetadata
 	{
-		int32 CurveId; 
 		FLinearColor CurveColor;
 		FText DisplayName;
 	};
@@ -70,6 +69,10 @@ public:
 	float ValueToLocalY(const FVector2f AllottedLocalSize, const float Value) const;
 	float LocalYToValue(const FVector2f AllottedLocalSize, const float LocalY) const;
 	
+#if !WITH_EDITOR
+	void UpdateYDataRangeFromTimestampRange(const double InLowerBoundTimestamp, const double InUpperBoundTimestamp);
+#endif // !WITH_EDITOR
+
 protected:
 	// SWidget interface
 	virtual FVector2D ComputeDesiredSize(float) const override;
