[Audio Insights] MixerSourceDashboardViewFactory: Adjusted in GetLastEntryArrayValue function for readability

diff --git a/Engine/Plugins/AudioInsights/Source/AudioInsights/Private/Views/MixerSourceDashboardViewFactory.cpp b/Engine/Plugins/AudioInsights/Source/AudioInsights/Private/Views/MixerSourceDashboardViewFactory.cpp
--- a/Engine/Plugins/AudioInsights/Source/AudioInsights/Private/Views/MixerSourceDashboardViewFactory.cpp
+++ b/Engine/Plugins/AudioInsights/Source/AudioInsights/Private/Views/MixerSourceDashboardViewFactory.cpp
@@ -34,17 +34,15 @@ namespace UE::Audio::Insights
 			return static_cast<const FMixerSourceDashboardEntry&>(InData);
 		};
 
-		float GetLastEntryArrayValue(const ::Audio::TCircularAudioBuffer<FDataPoint>& EnvelopeDataPoints)
+		float GetLastEntryArrayValue(const ::Audio::TCircularAudioBuffer<FDataPoint>& InDataPoints)
 		{
-			float LastValue = 0.0f;
-
-			if (EnvelopeDataPoints.Num() > 0)
+			if (InDataPoints.Num() > 0)
 			{
-				const ::Audio::DisjointedArrayView<const FDataPoint> EnvelopeDataPointsDisjointedArrayView = EnvelopeDataPoints.PeekInPlace(EnvelopeDataPoints.Num());
-				LastValue = EnvelopeDataPointsDisjointedArrayView.FirstBuffer.Last().Value;
+				const ::Audio::DisjointedArrayView<const FDataPoint> DataPointsDisjointedArrayView = InDataPoints.PeekInPlace(InDataPoints.Num());
+				return DataPointsDisjointedArrayView.FirstBuffer.Last().Value;
 			}
 
-			return LastValue;
+			return 0.0f;
 		};
 
 		const FText PlotColumnSelectDescription = LOCTEXT("AudioDashboard_MixerSources_SelectPlotColumnDescription", "Select a column from the table to plot.");
