[AudioInsights] Now using the AudioWidgets AudioAnalyzerRack in AudioInsights. Created AudioAnalyzerRackDashboardViewFactory to replace the separate OutputMeterDashboardViewFactory and OutputOscilloscopeDashboardViewFactory.

diff --git a/Engine/Plugins/AudioInsights/Source/AudioInsights/Private/AudioInsightsDashboardFactory.cpp b/Engine/Plugins/AudioInsights/Source/AudioInsights/Private/AudioInsightsDashboardFactory.cpp
--- a/Engine/Plugins/AudioInsights/Source/AudioInsights/Private/AudioInsightsDashboardFactory.cpp
+++ b/Engine/Plugins/AudioInsights/Source/AudioInsights/Private/AudioInsightsDashboardFactory.cpp
@@ -118,8 +118,7 @@ namespace UE::Audio::Insights
 										InvokedOutputMeterTab.IsValid() && DashboardViewFactories[FactoryName].IsValid())
 									{
 										if (const EDefaultDashboardTabStack DefaultTabStack = DashboardViewFactories[FactoryName]->GetDefaultTabStack();
-											DefaultTabStack == EDefaultDashboardTabStack::AudioMeter ||
-											DefaultTabStack == EDefaultDashboardTabStack::Oscilloscope)
+											DefaultTabStack == EDefaultDashboardTabStack::AudioAnalyzerRack)
 										{
 											InvokedOutputMeterTab->SetParentDockTabStackTabWellHidden(true);
 										}
diff --git a/Engine/Plugins/AudioInsights/Source/AudioInsights/Public/Views/DashboardViewFactory.h b/Engine/Plugins/AudioInsights/Source/AudioInsights/Public/Views/DashboardViewFactory.h
--- a/Engine/Plugins/AudioInsights/Source/AudioInsights/Public/Views/DashboardViewFactory.h
+++ b/Engine/Plugins/AudioInsights/Source/AudioInsights/Public/Views/DashboardViewFactory.h
@@ -18,8 +18,7 @@ namespace UE::Audio::Insights
 		Log,
 		Analysis,
 		AudioMeters,
-		AudioMeter,
-		Oscilloscope
+		AudioAnalyzerRack,
 	};
 
 	class IDashboardViewFactory
diff --git a/Engine/Plugins/AudioInsights/Source/AudioInsightsEditor/Private/Analyzers/AudioOscilloscopeAnalyzer.cpp b/Engine/Plugins/AudioInsights/Source/AudioInsightsEditor/Private/Analyzers/AudioOscilloscopeAnalyzer.cpp
deleted file mode 100644
diff --git a/Engine/Plugins/AudioInsights/Source/AudioInsightsEditor/Private/Analyzers/AudioOscilloscopeAnalyzer.h b/Engine/Plugins/AudioInsights/Source/AudioInsightsEditor/Private/Analyzers/AudioOscilloscopeAnalyzer.h
deleted file mode 100644
diff --git a/Engine/Plugins/AudioInsights/Source/AudioInsightsEditor/Private/Analyzers/SubmixAudioAnalyzerRack.cpp b/Engine/Plugins/AudioInsights/Source/AudioInsightsEditor/Private/Analyzers/SubmixAudioAnalyzerRack.cpp
new file mode 100644
--- /dev/null
+++ b/Engine/Plugins/AudioInsights/Source/AudioInsightsEditor/Private/Analyzers/SubmixAudioAnalyzerRack.cpp
@@ -0,0 +1,223 @@
+// Copyright Epic Games, Inc. All Rights Reserved.
+#include "SubmixAudioAnalyzerRack.h"
+
+#include "AudioAnalyzerRack.h"
+#include "AudioDeviceManager.h"
+#include "AudioInsightsEditorModule.h"
+#include "AudioInsightsEditorSettings.h"
+#include "AudioInsightsStyle.h"
+#include "AudioMaterialSlate/AudioMaterialSlateTypes.h"
+#include "AudioMeter.h"
+#include "AudioMixerDevice.h"
+#include "AudioMixerSubmix.h"
+#include "AudioOscilloscopePanelStyle.h"
+#include "AudioVectorscopePanelStyle.h"
+#include "AudioWidgetsStyle.h"
+
+namespace UE::Audio::Insights
+{
+	namespace FSubmixAudioAnalyzerRackPrivate
+	{
+		/**
+		 * This style set is given to the AudioWidgets::FAudioAnalyzerRack to override the parent AudioWidgetsStyle.
+		 */
+		class FAnalyzerRackStyleSet final : public FSlateStyleSet
+		{
+		public:
+			static FAnalyzerRackStyleSet& Get()
+			{
+				static FAnalyzerRackStyleSet Instance;
+				return Instance;
+			}
+
+			FAnalyzerRackStyleSet()
+				: FSlateStyleSet("AudioInsightsAnalyzerRackStyleSet")
+			{
+				SetParentStyleName(FAudioWidgetsStyle::Get().GetStyleSetName());
+
+				const FLinearColor AnalyzerForegroundColor(0.025719f, 0.208333f, 0.069907f, 1.0f); // "Audio" Green
+
+				// Override colors for these widget styles:
+
+				FAudioMeterDefaultColorStyle MeterStyle;
+				MeterStyle.MeterValueColor = AnalyzerForegroundColor;
+				Set("AudioMeter.DefaultColorStyle", MeterStyle);
+
+				Set("AudioOscilloscope.PanelStyle", FAudioOscilloscopePanelStyle()
+					.SetWaveViewerStyle(FSampledSequenceViewerStyle()
+						.SetSequenceColor(AnalyzerForegroundColor)));
+
+				Set("AudioSpectrumPlot.Style", FAudioSpectrumPlotStyle()
+					.SetCrosshairColor(FSlateColor(AnalyzerForegroundColor).UseSubduedForeground())
+					.SetSpectrumColor(AnalyzerForegroundColor));
+
+				Set("AudioVectorscope.PanelStyle", FAudioVectorscopePanelStyle()
+					.SetVectorViewerStyle(FSampledSequenceVectorViewerStyle()
+						.SetLineColor(AnalyzerForegroundColor)));
+			}
+
+		protected:
+			virtual const FSlateWidgetStyle* GetWidgetStyleInternal(const FName DesiredTypeName, const FName StyleName, const FSlateWidgetStyle* DefaultStyle, bool bWarnIfNotFound) const override
+			{
+				if (DesiredTypeName == FAudioMaterialMeterStyle::TypeName)
+				{
+					// Return null for this type to disable use of the audio material meter.
+					ensure(!bWarnIfNotFound);
+					return nullptr;
+				}
+
+				return FSlateStyleSet::GetWidgetStyleInternal(DesiredTypeName, StyleName, DefaultStyle, bWarnIfNotFound);
+			}
+		};
+
+		TSharedRef<AudioWidgets::FAudioAnalyzerRack> CreateAudioAnalyzerRack()
+		{
+			using namespace AudioWidgets;
+
+			// Set params so that rack layout is stored specific to Audio Insights and the custom style set is used for analyzer widgets:
+			const FAudioAnalyzerRack::FRackConstructParams Params
+			{
+				.TabManagerLayoutName = TEXT("AudioInsights_FAudioAnalyzerRack_v0"),
+				.StyleSet = &FAnalyzerRackStyleSet::Get(),
+				.EditorSettingsClass = UAudioInsightsEditorSettings::StaticClass(),
+			};
+
+			return MakeShared<FAudioAnalyzerRack>(Params);
+		}
+	} // namespace FSubmixAudioAnalyzerRackPrivate
+
+	FSubmixAudioAnalyzerRack::FSubmixAudioAnalyzerRack(TWeakObjectPtr<USoundSubmix> InSoundSubmix)
+		: AudioAnalyzerRack(FSubmixAudioAnalyzerRackPrivate::CreateAudioAnalyzerRack())
+	{
+		RebuildAudioAnalyzerRack(InSoundSubmix);
+	}
+
+	FSubmixAudioAnalyzerRack::~FSubmixAudioAnalyzerRack()
+	{
+		CleanupAudioAnalyzerRack();
+	}
+
+	TSharedRef<SWidget> FSubmixAudioAnalyzerRack::MakeWidget()
+	{
+		return AudioAnalyzerRack->CreateWidget(SNew(SDockTab));
+	}
+
+	void FSubmixAudioAnalyzerRack::RebuildAudioAnalyzerRack(TWeakObjectPtr<USoundSubmix> InSoundSubmix)
+	{
+		using namespace ::Audio;
+
+		if (SoundSubmix.IsValid())
+		{
+			CleanupAudioAnalyzerRack();
+		}
+
+		SoundSubmix = InSoundSubmix;
+
+		const FAudioDeviceManager* AudioDeviceManager = FAudioDeviceManager::Get();
+		if (!AudioDeviceManager)
+		{
+			return;
+		}
+
+		const FAudioInsightsEditorModule& AudioInsightsEditorModule = FAudioInsightsEditorModule::GetChecked();
+		const FDeviceId AudioDeviceId = AudioInsightsEditorModule.GetDeviceId();
+
+		const FMixerDevice* MixerDevice = static_cast<const FMixerDevice*>(AudioDeviceManager->GetAudioDeviceRaw(AudioDeviceId));
+		if (!MixerDevice)
+		{
+			return;
+		}
+
+		if (!SoundSubmix.IsValid())
+		{
+			return;
+		}
+
+		FMixerSubmixWeakPtr MixerSubmixWeakPtr = MixerDevice->GetSubmixInstance(SoundSubmix.Get());
+		if (!MixerSubmixWeakPtr.IsValid())
+		{
+			return;
+		}
+
+		AudioAnalyzerRack->Init(MixerDevice->GetNumDeviceChannels(), AudioDeviceId);
+
+		// Start processing
+		AudioAnalyzerRack->StartProcessing();
+
+		// Register audio bus in submix
+		const TObjectPtr<UAudioBus> AudioBus = AudioAnalyzerRack->GetAudioBus();
+		if (!AudioBus)
+		{
+			return;
+		}
+
+		const FAudioBusKey AudioBusKey(AudioBus->GetUniqueID());
+		const int32 AudioBusNumChannels = AudioBus->GetNumChannels();
+
+		FAudioThread::RunCommandOnAudioThread([MixerDevice, MixerSubmixWeakPtr, AudioBusKey, AudioBusNumChannels]()
+		{
+			TObjectPtr<UAudioBusSubsystem> AudioBusSubsystem = MixerDevice->GetSubsystem<UAudioBusSubsystem>();
+			check(AudioBusSubsystem);
+
+			if (FMixerSubmixPtr MixerSubmix = MixerSubmixWeakPtr.Pin();
+				MixerSubmix.IsValid())
+			{
+				MixerSubmix->RegisterAudioBus(AudioBusKey, AudioBusSubsystem->AddPatchInputForAudioBus(AudioBusKey, MixerDevice->GetNumOutputFrames(), AudioBusNumChannels));
+			}
+		});
+	}
+
+	void FSubmixAudioAnalyzerRack::CleanupAudioAnalyzerRack()
+	{
+		using namespace ::Audio;
+
+		const FAudioDeviceManager* AudioDeviceManager = FAudioDeviceManager::Get();
+		if (!AudioDeviceManager)
+		{
+			return;
+		}
+
+		const FAudioInsightsEditorModule& AudioInsightsEditorModule = FAudioInsightsEditorModule::GetChecked();
+		const FDeviceId AudioDeviceId = AudioInsightsEditorModule.GetDeviceId();
+
+		const FMixerDevice* MixerDevice = static_cast<const FMixerDevice*>(AudioDeviceManager->GetAudioDeviceRaw(AudioDeviceId));
+		if (!MixerDevice)
+		{
+			return;
+		}
+
+		if (!SoundSubmix.IsValid())
+		{
+			return;
+		}
+
+		FMixerSubmixWeakPtr MixerSubmixWeakPtr = MixerDevice->GetSubmixInstance(SoundSubmix.Get());
+		if (!MixerSubmixWeakPtr.IsValid())
+		{
+			return;
+		}
+
+		// Unregister audio bus from submix
+		const TObjectPtr<UAudioBus> AudioBus = AudioAnalyzerRack->GetAudioBus();
+		if (!AudioBus)
+		{
+			return;
+		}
+
+		const FAudioBusKey AudioBusKey(AudioBus->GetUniqueID());
+
+		FAudioThread::RunCommandOnAudioThread([MixerSubmixWeakPtr, AudioBusKey]()
+		{
+			if (FMixerSubmixPtr MixerSubmix = MixerSubmixWeakPtr.Pin();
+				MixerSubmix.IsValid())
+			{
+				MixerSubmix->UnregisterAudioBus(AudioBusKey);
+			}
+		});
+
+		// Stop processing
+		AudioAnalyzerRack->StopProcessing();
+
+		SoundSubmix.Reset();
+	}
+} // namespace UE::Audio::Insights
diff --git a/Engine/Plugins/AudioInsights/Source/AudioInsightsEditor/Private/Analyzers/SubmixAudioAnalyzerRack.h b/Engine/Plugins/AudioInsights/Source/AudioInsightsEditor/Private/Analyzers/SubmixAudioAnalyzerRack.h
new file mode 100644
--- /dev/null
+++ b/Engine/Plugins/AudioInsights/Source/AudioInsightsEditor/Private/Analyzers/SubmixAudioAnalyzerRack.h
@@ -0,0 +1,30 @@
+// Copyright Epic Games, Inc. All Rights Reserved.
+#pragma once
+
+#include "Templates/SharedPointer.h"
+#include "UObject/WeakObjectPtr.h"
+
+namespace AudioWidgets { class FAudioAnalyzerRack; }
+
+class SWidget;
+class USoundSubmix;
+
+namespace UE::Audio::Insights
+{
+	class FSubmixAudioAnalyzerRack : public TSharedFromThis<FSubmixAudioAnalyzerRack>
+	{
+	public:
+		FSubmixAudioAnalyzerRack(TWeakObjectPtr<USoundSubmix> InSoundSubmix);
+		virtual ~FSubmixAudioAnalyzerRack();
+
+		TSharedRef<SWidget> MakeWidget();
+
+		void RebuildAudioAnalyzerRack(TWeakObjectPtr<USoundSubmix> InSoundSubmix);
+
+	private:
+		void CleanupAudioAnalyzerRack();
+
+		TSharedRef<AudioWidgets::FAudioAnalyzerRack> AudioAnalyzerRack;
+		TWeakObjectPtr<USoundSubmix> SoundSubmix;
+	};
+} // namespace UE::Audio::Insights
diff --git a/Engine/Plugins/AudioInsights/Source/AudioInsightsEditor/Private/AudioInsightsEditorDashboardFactory.cpp b/Engine/Plugins/AudioInsights/Source/AudioInsightsEditor/Private/AudioInsightsEditorDashboardFactory.cpp
--- a/Engine/Plugins/AudioInsights/Source/AudioInsightsEditor/Private/AudioInsightsEditorDashboardFactory.cpp
+++ b/Engine/Plugins/AudioInsights/Source/AudioInsightsEditor/Private/AudioInsightsEditorDashboardFactory.cpp
@@ -320,8 +320,7 @@ namespace UE::Audio::Insights
 										InvokedOutputMeterTab.IsValid() && DashboardViewFactories[FactoryName].IsValid())
 									{
 										if (const EDefaultDashboardTabStack DefaultTabStack = DashboardViewFactories[FactoryName]->GetDefaultTabStack();
-											DefaultTabStack == EDefaultDashboardTabStack::AudioMeter ||
-											DefaultTabStack == EDefaultDashboardTabStack::Oscilloscope)
+											DefaultTabStack == EDefaultDashboardTabStack::AudioAnalyzerRack)
 										{
 											InvokedOutputMeterTab->SetParentDockTabStackTabWellHidden(true);
 										}
@@ -540,8 +539,7 @@ namespace UE::Audio::Insights
 		TSharedRef<FTabManager::FStack> LogTabStack = FTabManager::NewStack();
 		TSharedRef<FTabManager::FStack> AnalysisTabStack = FTabManager::NewStack();
 		TSharedRef<FTabManager::FStack> AudioMetersTabStack = FTabManager::NewStack();
-		TSharedRef<FTabManager::FStack> AudioMeterTabStack = FTabManager::NewStack();
-		TSharedRef<FTabManager::FStack> OscilloscopeTabStack = FTabManager::NewStack();
+		TSharedRef<FTabManager::FStack> AudioAnalyzerRackTabStack = FTabManager::NewStack()->SetHideTabWell(true)->SetSizeCoefficient(0.15f);
 
 		for (const auto& KVP : DashboardViewFactories)
 		{
@@ -575,15 +573,9 @@ namespace UE::Audio::Insights
 				}
 				break;
 
-				case EDefaultDashboardTabStack::AudioMeter:
+				case EDefaultDashboardTabStack::AudioAnalyzerRack:
 				{
-					AudioMeterTabStack->AddTab(FactoryName, ETabState::OpenedTab);
-				}
-				break;
-
-				case EDefaultDashboardTabStack::Oscilloscope:
-				{
-					OscilloscopeTabStack->AddTab(FactoryName, ETabState::OpenedTab);
+					AudioAnalyzerRackTabStack->AddTab(FactoryName, ETabState::OpenedTab);
 				}
 				break;
 
@@ -594,7 +586,7 @@ namespace UE::Audio::Insights
 
 		AnalysisTabStack->SetForegroundTab(FName("MixerSources"));
 
-		return FTabManager::NewLayout("AudioDashboard_Editor_Layout_v1")
+		return FTabManager::NewLayout("AudioDashboard_Editor_Layout_v2")
 		->AddArea
 		(
 			FTabManager::NewPrimaryArea()
@@ -650,23 +642,7 @@ namespace UE::Audio::Insights
 				// Right column
 				->Split
 				(
-					// Top
-					FTabManager::NewSplitter()
-					->SetOrientation(Orient_Vertical)
-					->SetSizeCoefficient(0.15f) // Column width
-					->Split
-					(
-						AudioMeterTabStack
-						->SetSizeCoefficient(0.7f)
-						->SetHideTabWell(true)
-					)
-					// Bottom
-					->Split
-					(
-						OscilloscopeTabStack
-						->SetSizeCoefficient(0.3f)
-						->SetHideTabWell(true)
-					)
+					AudioAnalyzerRackTabStack
 				)
 			)
 		);
diff --git a/Engine/Plugins/AudioInsights/Source/AudioInsightsEditor/Private/AudioInsightsEditorModule.cpp b/Engine/Plugins/AudioInsights/Source/AudioInsightsEditor/Private/AudioInsightsEditorModule.cpp
--- a/Engine/Plugins/AudioInsights/Source/AudioInsightsEditor/Private/AudioInsightsEditorModule.cpp
+++ b/Engine/Plugins/AudioInsights/Source/AudioInsightsEditor/Private/AudioInsightsEditorModule.cpp
@@ -8,12 +8,11 @@
 #include "Modules/ModuleManager.h"
 #include "Templates/SharedPointer.h"
 #include "UObject/NameTypes.h"
+#include "Views/AudioAnalyzerRackDashboardViewFactory.h"
 #include "Views/AudioBusesDashboardViewFactory.h"
 #include "Views/AudioMetersDashboardViewFactory.h"
 #include "Views/LogDashboardViewFactory.h"
 #include "Views/MixerSourceDashboardViewFactory.h"
-#include "Views/OutputMeterDashboardViewFactory.h"
-#include "Views/OutputOscilloscopeDashboardViewFactory.h"
 #include "Views/SubmixesDashboardViewFactory.h"
 #include "Views/ViewportDashboardViewFactory.h"
 #include "Views/VirtualLoopDashboardViewFactory.h"
@@ -42,8 +41,7 @@ namespace UE::Audio::Insights
 			DashboardFactory->RegisterViewFactory(MakeShared<FSubmixesDashboardViewFactory>());
 			DashboardFactory->RegisterViewFactory(MakeShared<FAudioBusesDashboardViewFactory>());
 			DashboardFactory->RegisterViewFactory(MakeShared<FAudioMetersDashboardViewFactory>());
-			DashboardFactory->RegisterViewFactory(MakeShared<FOutputMeterDashboardViewFactory>());
-			DashboardFactory->RegisterViewFactory(MakeShared<FOutputOscilloscopeDashboardViewFactory>());
+			DashboardFactory->RegisterViewFactory(MakeShared<FAudioAnalyzerRackDashboardViewFactory>());
 		}
 	}
 
diff --git a/Engine/Plugins/AudioInsights/Source/AudioInsightsEditor/Private/Views/AudioAnalyzerRackDashboardViewFactory.cpp b/Engine/Plugins/AudioInsights/Source/AudioInsightsEditor/Private/Views/AudioAnalyzerRackDashboardViewFactory.cpp
new file mode 100644
--- /dev/null
+++ b/Engine/Plugins/AudioInsights/Source/AudioInsightsEditor/Private/Views/AudioAnalyzerRackDashboardViewFactory.cpp
@@ -0,0 +1,81 @@
+// Copyright Epic Games, Inc. All Rights Reserved.
+
+#include "Views/AudioAnalyzerRackDashboardViewFactory.h"
+
+#include "Analyzers/SubmixAudioAnalyzerRack.h"
+#include "AudioInsightsEditorDashboardFactory.h"
+#include "AudioInsightsStyle.h"
+#include "Sound/AudioSettings.h"
+#include "Sound/SoundSubmix.h"
+#include "SubmixesDashboardViewFactory.h"
+#include "UObject/WeakObjectPtr.h"
+#include "Widgets/Colors/SColorBlock.h"
+
+#define LOCTEXT_NAMESPACE "AudioInsights"
+
+namespace UE::Audio::Insights
+{
+	FName FAudioAnalyzerRackDashboardViewFactory::GetName() const
+	{
+		return "AudioAnalyzerRack";
+	}
+
+	FText FAudioAnalyzerRackDashboardViewFactory::GetDisplayName() const
+	{
+		return LOCTEXT("AudioDashboard_DashboardsAnalyzerRackTab_DisplayName", "Analyzers");
+	}
+
+	EDefaultDashboardTabStack FAudioAnalyzerRackDashboardViewFactory::GetDefaultTabStack() const
+	{
+		return EDefaultDashboardTabStack::AudioAnalyzerRack;
+	}
+
+	FSlateIcon FAudioAnalyzerRackDashboardViewFactory::GetIcon() const
+	{
+		return FSlateStyle::Get().CreateIcon("AudioInsights.Icon");
+	}
+
+	TSharedRef<SWidget> FAudioAnalyzerRackDashboardViewFactory::MakeWidget()
+	{
+		// Create SubmixAudioAnalyzerRack (use MainSubmix as default)
+		if (TWeakObjectPtr<UAudioSettings> AudioSettings = GetMutableDefault<UAudioSettings>();
+			AudioSettings.IsValid())
+		{
+			MainSubmix = Cast<USoundSubmix>(AudioSettings->MasterSubmix.ResolveObject());
+
+			if (MainSubmix)
+			{
+				SubmixAudioAnalyzerRack = MakeShared<FSubmixAudioAnalyzerRack>(MainSubmix);
+
+				FEditorDashboardFactory::OnActiveAudioDeviceChanged.AddSP(this, &FAudioAnalyzerRackDashboardViewFactory::HandleOnActiveAudioDeviceChanged);
+				FSubmixesDashboardViewFactory::OnSubmixSelectionChanged.AddSP(this, &FAudioAnalyzerRackDashboardViewFactory::HandleOnSubmixSelectionChanged);
+			}
+		}
+		
+		if (!SubmixAudioAnalyzerRack.IsValid())
+		{
+			return SNew(SColorBlock)
+				.Color(FSlateStyle::Get().GetColor("AudioInsights.Analyzers.BackgroundColor"));
+		}
+
+		return SubmixAudioAnalyzerRack->MakeWidget();
+	}
+
+	void FAudioAnalyzerRackDashboardViewFactory::HandleOnActiveAudioDeviceChanged()
+	{
+		if (SubmixAudioAnalyzerRack.IsValid())
+		{
+			SubmixAudioAnalyzerRack->RebuildAudioAnalyzerRack(MainSubmix);
+		}
+	}
+
+	void FAudioAnalyzerRackDashboardViewFactory::HandleOnSubmixSelectionChanged(const TWeakObjectPtr<USoundSubmix> InSoundSubmix)
+	{
+		if (SubmixAudioAnalyzerRack.IsValid())
+		{
+			SubmixAudioAnalyzerRack->RebuildAudioAnalyzerRack(InSoundSubmix);
+		}
+	}
+} // namespace UE::Audio::Insights
+
+#undef LOCTEXT_NAMESPACE
diff --git a/Engine/Plugins/AudioInsights/Source/AudioInsightsEditor/Private/Views/AudioAnalyzerRackDashboardViewFactory.h b/Engine/Plugins/AudioInsights/Source/AudioInsightsEditor/Private/Views/AudioAnalyzerRackDashboardViewFactory.h
new file mode 100644
--- /dev/null
+++ b/Engine/Plugins/AudioInsights/Source/AudioInsightsEditor/Private/Views/AudioAnalyzerRackDashboardViewFactory.h
@@ -0,0 +1,29 @@
+// Copyright Epic Games, Inc. All Rights Reserved.
+#pragma once
+
+#include "Templates/SharedPointer.h"
+#include "Views/DashboardViewFactory.h"
+
+class USoundSubmix;
+
+namespace UE::Audio::Insights
+{
+	class FSubmixAudioAnalyzerRack;
+
+	class FAudioAnalyzerRackDashboardViewFactory : public IDashboardViewFactory, public TSharedFromThis<FAudioAnalyzerRackDashboardViewFactory>
+	{
+	public:
+		virtual FName GetName() const override;
+		virtual FText GetDisplayName() const override;
+		virtual EDefaultDashboardTabStack GetDefaultTabStack() const override;
+		virtual FSlateIcon GetIcon() const override;
+		virtual TSharedRef<SWidget> MakeWidget() override;
+
+	private:
+		void HandleOnActiveAudioDeviceChanged();
+		void HandleOnSubmixSelectionChanged(const TWeakObjectPtr<USoundSubmix> InSoundSubmix);
+
+		TSharedPtr<FSubmixAudioAnalyzerRack> SubmixAudioAnalyzerRack;
+		TObjectPtr<USoundSubmix> MainSubmix;
+	};
+} // namespace UE::Audio::Insights
diff --git a/Engine/Plugins/AudioInsights/Source/AudioInsightsEditor/Private/Views/OutputMeterDashboardViewFactory.cpp b/Engine/Plugins/AudioInsights/Source/AudioInsightsEditor/Private/Views/OutputMeterDashboardViewFactory.cpp
deleted file mode 100644
diff --git a/Engine/Plugins/AudioInsights/Source/AudioInsightsEditor/Private/Views/OutputMeterDashboardViewFactory.h b/Engine/Plugins/AudioInsights/Source/AudioInsightsEditor/Private/Views/OutputMeterDashboardViewFactory.h
deleted file mode 100644
diff --git a/Engine/Plugins/AudioInsights/Source/AudioInsightsEditor/Private/Views/OutputOscilloscopeDashboardViewFactory.cpp b/Engine/Plugins/AudioInsights/Source/AudioInsightsEditor/Private/Views/OutputOscilloscopeDashboardViewFactory.cpp
deleted file mode 100644
diff --git a/Engine/Plugins/AudioInsights/Source/AudioInsightsEditor/Private/Views/OutputOscilloscopeDashboardViewFactory.h b/Engine/Plugins/AudioInsights/Source/AudioInsightsEditor/Private/Views/OutputOscilloscopeDashboardViewFactory.h
deleted file mode 100644
diff --git a/Engine/Plugins/AudioInsights/Source/AudioInsightsEditor/Public/AudioInsightsEditorSettings.h b/Engine/Plugins/AudioInsights/Source/AudioInsightsEditor/Public/AudioInsightsEditorSettings.h
--- a/Engine/Plugins/AudioInsights/Source/AudioInsightsEditor/Public/AudioInsightsEditorSettings.h
+++ b/Engine/Plugins/AudioInsights/Source/AudioInsightsEditor/Public/AudioInsightsEditorSettings.h
@@ -2,6 +2,8 @@
 
 #pragma once
 
+#include "AudioSpectrogram.h"
+#include "AudioSpectrumAnalyzer.h"
 #include "Engine/DeveloperSettings.h"
 
 #include "AudioInsightsEditorSettings.generated.h"
@@ -19,4 +21,12 @@ public:
 	/** Whether to automatically set the first PIE client in Audio Insights World Filter. */
 	UPROPERTY(Config, EditAnywhere, Category="World Filter")
 	bool bWorldFilterDefaultsToFirstClient = false;
+
+	/** Settings for analyzer rack spectrogram widget */
+	UPROPERTY(EditAnywhere, config, Category = Spectrogram, meta = (ShowOnlyInnerProperties))
+	FSpectrogramRackUnitSettings SpectrogramSettings;
+
+	/** Settings for analyzer rack spectrum analyzer widget */
+	UPROPERTY(EditAnywhere, config, Category = SpectrumAnalyzer, meta = (ShowOnlyInnerProperties))
+	FSpectrumAnalyzerRackUnitSettings SpectrumAnalyzerSettings;
 };
diff --git a/Engine/Plugins/Runtime/AudioWidgets/Source/AudioWidgets/Private/AudioSpectrogramViewport.cpp b/Engine/Plugins/Runtime/AudioWidgets/Source/AudioWidgets/Private/AudioSpectrogramViewport.cpp
--- a/Engine/Plugins/Runtime/AudioWidgets/Source/AudioWidgets/Private/AudioSpectrogramViewport.cpp
+++ b/Engine/Plugins/Runtime/AudioWidgets/Source/AudioWidgets/Private/AudioSpectrogramViewport.cpp
@@ -122,9 +122,12 @@ FAudioSpectrogramViewport::~FAudioSpectrogramViewport()
 	// Release texture if we have one:
 	if (UpdatableTexture != nullptr)
 	{
-		if (FSlateRenderer* Renderer = FSlateApplication::Get().GetRenderer())
+		if (FSlateApplication::IsInitialized())
 		{
-			Renderer->ReleaseUpdatableTexture(UpdatableTexture);
+			if (FSlateRenderer* Renderer = FSlateApplication::Get().GetRenderer())
+			{
+				Renderer->ReleaseUpdatableTexture(UpdatableTexture);
+			}
 		}
 
 		UpdatableTexture = nullptr;
