[AudioInsights] Fix for crash if Reset Layout pressed after a nested docktab has been undocked - Now passing owner tab through to NewTabManager calls in FMixerSourceDashboardViewFactory and FAudioAnalyzerRack.

diff --git a/Engine/Plugins/AudioInsights/Source/AudioInsights/Private/AudioInsightsDashboardFactory.cpp b/Engine/Plugins/AudioInsights/Source/AudioInsights/Private/AudioInsightsDashboardFactory.cpp
--- a/Engine/Plugins/AudioInsights/Source/AudioInsights/Private/AudioInsightsDashboardFactory.cpp
+++ b/Engine/Plugins/AudioInsights/Source/AudioInsights/Private/AudioInsightsDashboardFactory.cpp
@@ -220,13 +220,14 @@ namespace UE::Audio::Insights
 
 			DashboardTabManager->RegisterTabSpawner(FactoryName, FOnSpawnTab::CreateLambda([this, Factory](const FSpawnTabArgs& Args)
 			{
-				TSharedPtr<SWidget> DashboardView = Factory->MakeWidget();
-				return SNew(SDockTab)
+				TSharedRef<SDockTab> DockTab = SNew(SDockTab)
 					.Clipping(EWidgetClipping::ClipToBounds)
-					.Label(Factory->GetDisplayName())
-					[
-						DashboardView->AsShared()
-					];
+					.Label(Factory->GetDisplayName());
+
+				TSharedRef<SWidget> DashboardView = Factory->MakeWidget(DockTab);
+				DockTab->SetContent(DashboardView);
+
+				return DockTab;
 			}))
 			.SetDisplayName(Factory->GetDisplayName())
 			.SetGroup(DashboardWorkspace->AsShared())
diff --git a/Engine/Plugins/AudioInsights/Source/AudioInsights/Private/Views/LogDashboardViewFactory.cpp b/Engine/Plugins/AudioInsights/Source/AudioInsights/Private/Views/LogDashboardViewFactory.cpp
--- a/Engine/Plugins/AudioInsights/Source/AudioInsights/Private/Views/LogDashboardViewFactory.cpp
+++ b/Engine/Plugins/AudioInsights/Source/AudioInsights/Private/Views/LogDashboardViewFactory.cpp
@@ -54,7 +54,7 @@ namespace UE::Audio::Insights
 		return FSlateStyle::Get().CreateIcon("AudioInsights.Icon.Log");
 	}
 
-	TSharedRef<SWidget> FLogDashboardViewFactory::MakeWidget()
+	TSharedRef<SWidget> FLogDashboardViewFactory::MakeWidget(TSharedRef<SDockTab> OwnerTab)
 	{
 		using namespace LogDashboardViewFactoryPrivate;
 
diff --git a/Engine/Plugins/AudioInsights/Source/AudioInsights/Private/Views/MixerSourceDashboardViewFactory.cpp b/Engine/Plugins/AudioInsights/Source/AudioInsights/Private/Views/MixerSourceDashboardViewFactory.cpp
--- a/Engine/Plugins/AudioInsights/Source/AudioInsights/Private/Views/MixerSourceDashboardViewFactory.cpp
+++ b/Engine/Plugins/AudioInsights/Source/AudioInsights/Private/Views/MixerSourceDashboardViewFactory.cpp
@@ -971,9 +971,10 @@ namespace UE::Audio::Insights
 
 	TSharedRef<SDockTab> FMixerSourceDashboardViewFactory::CreateMixerSourcesTab(const FSpawnTabArgs& Args)
 	{
-		return SNew(SDockTab)
-			.Clipping(EWidgetClipping::ClipToBounds)
-			[
+		TSharedRef<SDockTab> DockTab = SNew(SDockTab)
+			.Clipping(EWidgetClipping::ClipToBounds);
+
+		DockTab->SetContent(
 				SNew(SVerticalBox)
 				+SVerticalBox::Slot()
 				.AutoHeight()
@@ -1003,9 +1004,11 @@ namespace UE::Audio::Insights
 				+ SVerticalBox::Slot()
 				.HAlign(HAlign_Fill)
 				[
-					FTraceTableDashboardViewFactory::MakeWidget()
+					FTraceTableDashboardViewFactory::MakeWidget(DockTab)
 				]
-			];
+			);
+
+		return DockTab;
 	}
 
 	TSharedRef<SDockTab> FMixerSourceDashboardViewFactory::CreatePlotsTab(const FSpawnTabArgs& Args)
@@ -1099,7 +1102,7 @@ namespace UE::Audio::Insights
 			);
 	}
 
-	TSharedRef<SWidget> FMixerSourceDashboardViewFactory::MakeWidget()
+	TSharedRef<SWidget> FMixerSourceDashboardViewFactory::MakeWidget(TSharedRef<SDockTab> OwnerTab)
 	{
 		using namespace MixerSourcePrivate;
 
@@ -1123,8 +1126,7 @@ namespace UE::Audio::Insights
 
 		PlotsWidget = MakePlotsWidget();
 
-		const TSharedRef<SDockTab> DockTab = SNew(SDockTab);
-		MixerSourcesTabManager = FGlobalTabmanager::Get()->NewTabManager(DockTab);
+		MixerSourcesTabManager = FGlobalTabmanager::Get()->NewTabManager(OwnerTab);
 
 #if WITH_EDITOR
 		MixerSourcesTabManager->SetOnPersistLayout(FTabManager::FOnPersistLayout::CreateStatic([](const TSharedRef<FTabManager::FLayout>& InLayout)
@@ -1136,7 +1138,7 @@ namespace UE::Audio::Insights
 		}));
 #endif // WITH_EDITOR
 
-		DockTab->SetOnTabClosed(SDockTab::FOnTabClosedCallback::CreateLambda([this](TSharedRef<SDockTab> TabClosed)
+		OwnerTab->SetOnTabClosed(SDockTab::FOnTabClosedCallback::CreateLambda([this](TSharedRef<SDockTab> TabClosed)
 		{
 			UnregisterTabSpawners();
 
diff --git a/Engine/Plugins/AudioInsights/Source/AudioInsights/Private/Views/TableDashboardViewFactory.cpp b/Engine/Plugins/AudioInsights/Source/AudioInsights/Private/Views/TableDashboardViewFactory.cpp
--- a/Engine/Plugins/AudioInsights/Source/AudioInsights/Private/Views/TableDashboardViewFactory.cpp
+++ b/Engine/Plugins/AudioInsights/Source/AudioInsights/Private/Views/TableDashboardViewFactory.cpp
@@ -126,7 +126,7 @@ namespace UE::Audio::Insights
 		return HeaderRowWidget;
 	}
 
-	TSharedRef<SWidget> FTraceTableDashboardViewFactory::MakeWidget()
+	TSharedRef<SWidget> FTraceTableDashboardViewFactory::MakeWidget(TSharedRef<SDockTab> OwnerTab)
 	{
 		if (!DashboardWidget.IsValid())
 		{
@@ -354,7 +354,7 @@ namespace UE::Audio::Insights
 		return FReply::Unhandled();
 	}
 
-	TSharedRef<SWidget> FTraceObjectTableDashboardViewFactory::MakeWidget()
+	TSharedRef<SWidget> FTraceObjectTableDashboardViewFactory::MakeWidget(TSharedRef<SDockTab> OwnerTab)
 	{
 		if (!DashboardWidget.IsValid())
 		{
@@ -376,7 +376,7 @@ namespace UE::Audio::Insights
 			.AutoHeight()
 			.HAlign(HAlign_Fill)
 			[
-				FTraceTableDashboardViewFactory::MakeWidget()
+				FTraceTableDashboardViewFactory::MakeWidget(OwnerTab)
 			];
 
 			if (FilteredEntriesListView.IsValid())
diff --git a/Engine/Plugins/AudioInsights/Source/AudioInsights/Public/Views/DashboardViewFactory.h b/Engine/Plugins/AudioInsights/Source/AudioInsights/Public/Views/DashboardViewFactory.h
--- a/Engine/Plugins/AudioInsights/Source/AudioInsights/Public/Views/DashboardViewFactory.h
+++ b/Engine/Plugins/AudioInsights/Source/AudioInsights/Public/Views/DashboardViewFactory.h
@@ -9,6 +9,7 @@
 #include "UObject/NameTypes.h"
 #include "Widgets/SWidget.h"
 
+class SDockTab;
 
 namespace UE::Audio::Insights
 {
@@ -30,7 +31,7 @@ namespace UE::Audio::Insights
 		virtual FText GetDisplayName() const = 0;
 		virtual FName GetName() const = 0;
 		virtual FSlateIcon GetIcon() const = 0;
-		virtual TSharedRef<SWidget> MakeWidget() = 0;
+		virtual TSharedRef<SWidget> MakeWidget(TSharedRef<SDockTab> OwnerTab) = 0;
 	};
 
 	class FTraceDashboardViewFactoryBase : public IDashboardViewFactory
diff --git a/Engine/Plugins/AudioInsights/Source/AudioInsights/Public/Views/LogDashboardViewFactory.h b/Engine/Plugins/AudioInsights/Source/AudioInsights/Public/Views/LogDashboardViewFactory.h
--- a/Engine/Plugins/AudioInsights/Source/AudioInsights/Public/Views/LogDashboardViewFactory.h
+++ b/Engine/Plugins/AudioInsights/Source/AudioInsights/Public/Views/LogDashboardViewFactory.h
@@ -24,7 +24,7 @@ namespace UE::Audio::Insights
 		virtual FText GetDisplayName() const override;
 		virtual EDefaultDashboardTabStack GetDefaultTabStack() const override;
 		virtual FSlateIcon GetIcon() const override;
-		virtual TSharedRef<SWidget> MakeWidget() override;
+		virtual TSharedRef<SWidget> MakeWidget(TSharedRef<SDockTab> OwnerTab) override;
 
 	private:
 		struct FLogCategoryCollector : public FOutputDevice
diff --git a/Engine/Plugins/AudioInsights/Source/AudioInsights/Public/Views/MixerSourceDashboardViewFactory.h b/Engine/Plugins/AudioInsights/Source/AudioInsights/Public/Views/MixerSourceDashboardViewFactory.h
--- a/Engine/Plugins/AudioInsights/Source/AudioInsights/Public/Views/MixerSourceDashboardViewFactory.h
+++ b/Engine/Plugins/AudioInsights/Source/AudioInsights/Public/Views/MixerSourceDashboardViewFactory.h
@@ -26,7 +26,7 @@ namespace UE::Audio::Insights
 		virtual FSlateIcon GetIcon() const override;
 		virtual EDefaultDashboardTabStack GetDefaultTabStack() const override;
 		virtual void ProcessEntries(FTraceTableDashboardViewFactory::EProcessReason Reason) override;
-		virtual TSharedRef<SWidget> MakeWidget() override;
+		virtual TSharedRef<SWidget> MakeWidget(TSharedRef<SDockTab> OwnerTab) override;
 
 		// Maximum amount of data history kept for plots (in seconds)
 		static const double MaxPlotHistorySeconds;
diff --git a/Engine/Plugins/AudioInsights/Source/AudioInsights/Public/Views/TableDashboardViewFactory.h b/Engine/Plugins/AudioInsights/Source/AudioInsights/Public/Views/TableDashboardViewFactory.h
--- a/Engine/Plugins/AudioInsights/Source/AudioInsights/Public/Views/TableDashboardViewFactory.h
+++ b/Engine/Plugins/AudioInsights/Source/AudioInsights/Public/Views/TableDashboardViewFactory.h
@@ -60,7 +60,7 @@ namespace UE::Audio::Insights
 	public:
 		virtual TSharedRef<SWidget> GenerateWidgetForColumn(TSharedRef<IDashboardDataViewEntry> InRowData, const FName& Column);
 		virtual FReply OnDataRowKeyInput(const FGeometry& Geometry, const FKeyEvent& KeyEvent) const;
-		virtual TSharedRef<SWidget> MakeWidget() override;
+		virtual TSharedRef<SWidget> MakeWidget(TSharedRef<SDockTab> OwnerTab) override;
 		virtual void RebuildFilteredEntriesListView();
 
 		enum class EProcessReason : uint8
@@ -185,7 +185,7 @@ namespace UE::Audio::Insights
 
 		virtual TSharedRef<SWidget> GenerateWidgetForColumn(TSharedRef<IDashboardDataViewEntry> InRowData, const FName& Column) override;
 		virtual FReply OnDataRowKeyInput(const FGeometry& Geometry, const FKeyEvent& KeyEvent) const override;
-		virtual TSharedRef<SWidget> MakeWidget() override;
+		virtual TSharedRef<SWidget> MakeWidget(TSharedRef<SDockTab> OwnerTab) override;
 
 	protected:
 #if WITH_EDITOR
diff --git a/Engine/Plugins/AudioInsights/Source/AudioInsightsEditor/Private/Analyzers/SubmixAudioAnalyzerRack.cpp b/Engine/Plugins/AudioInsights/Source/AudioInsightsEditor/Private/Analyzers/SubmixAudioAnalyzerRack.cpp
--- a/Engine/Plugins/AudioInsights/Source/AudioInsightsEditor/Private/Analyzers/SubmixAudioAnalyzerRack.cpp
+++ b/Engine/Plugins/AudioInsights/Source/AudioInsightsEditor/Private/Analyzers/SubmixAudioAnalyzerRack.cpp
@@ -97,9 +97,9 @@ namespace UE::Audio::Insights
 		CleanupAudioAnalyzerRack();
 	}
 
-	TSharedRef<SWidget> FSubmixAudioAnalyzerRack::MakeWidget()
+	TSharedRef<SWidget> FSubmixAudioAnalyzerRack::MakeWidget(TSharedRef<SDockTab> InOwnerTab)
 	{
-		return AudioAnalyzerRack->CreateWidget(SNew(SDockTab));
+		return AudioAnalyzerRack->CreateWidget(InOwnerTab);
 	}
 
 	void FSubmixAudioAnalyzerRack::RebuildAudioAnalyzerRack(TWeakObjectPtr<USoundSubmix> InSoundSubmix)
diff --git a/Engine/Plugins/AudioInsights/Source/AudioInsightsEditor/Private/Analyzers/SubmixAudioAnalyzerRack.h b/Engine/Plugins/AudioInsights/Source/AudioInsightsEditor/Private/Analyzers/SubmixAudioAnalyzerRack.h
--- a/Engine/Plugins/AudioInsights/Source/AudioInsightsEditor/Private/Analyzers/SubmixAudioAnalyzerRack.h
+++ b/Engine/Plugins/AudioInsights/Source/AudioInsightsEditor/Private/Analyzers/SubmixAudioAnalyzerRack.h
@@ -6,6 +6,7 @@
 
 namespace AudioWidgets { class FAudioAnalyzerRack; }
 
+class SDockTab;
 class SWidget;
 class USoundSubmix;
 
@@ -17,7 +18,7 @@ namespace UE::Audio::Insights
 		FSubmixAudioAnalyzerRack(TWeakObjectPtr<USoundSubmix> InSoundSubmix);
 		virtual ~FSubmixAudioAnalyzerRack();
 
-		TSharedRef<SWidget> MakeWidget();
+		TSharedRef<SWidget> MakeWidget(TSharedRef<SDockTab> InOwnerTab);
 
 		void RebuildAudioAnalyzerRack(TWeakObjectPtr<USoundSubmix> InSoundSubmix);
 
diff --git a/Engine/Plugins/AudioInsights/Source/AudioInsightsEditor/Private/AudioInsightsEditorDashboardFactory.cpp b/Engine/Plugins/AudioInsights/Source/AudioInsightsEditor/Private/AudioInsightsEditorDashboardFactory.cpp
--- a/Engine/Plugins/AudioInsights/Source/AudioInsightsEditor/Private/AudioInsightsEditorDashboardFactory.cpp
+++ b/Engine/Plugins/AudioInsights/Source/AudioInsightsEditor/Private/AudioInsightsEditorDashboardFactory.cpp
@@ -661,13 +661,14 @@ namespace UE::Audio::Insights
 
 			DashboardTabManager->RegisterTabSpawner(FactoryName, FOnSpawnTab::CreateLambda([this, Factory](const FSpawnTabArgs& Args)
 			{
-				TSharedPtr<SWidget> DashboardView = Factory->MakeWidget();
-				return SNew(SDockTab)
+				TSharedRef<SDockTab> DockTab = SNew(SDockTab)
 					.Clipping(EWidgetClipping::ClipToBounds)
-					.Label(Factory->GetDisplayName())
-					[
-						DashboardView->AsShared()
-					];
+					.Label(Factory->GetDisplayName());
+
+				TSharedRef<SWidget> DashboardView = Factory->MakeWidget(DockTab);
+				DockTab->SetContent(DashboardView);
+
+				return DockTab;
 			}))
 			.SetDisplayName(Factory->GetDisplayName())
 			.SetGroup(DashboardWorkspace->AsShared())
diff --git a/Engine/Plugins/AudioInsights/Source/AudioInsightsEditor/Private/Views/AudioAnalyzerRackDashboardViewFactory.cpp b/Engine/Plugins/AudioInsights/Source/AudioInsightsEditor/Private/Views/AudioAnalyzerRackDashboardViewFactory.cpp
--- a/Engine/Plugins/AudioInsights/Source/AudioInsightsEditor/Private/Views/AudioAnalyzerRackDashboardViewFactory.cpp
+++ b/Engine/Plugins/AudioInsights/Source/AudioInsightsEditor/Private/Views/AudioAnalyzerRackDashboardViewFactory.cpp
@@ -35,7 +35,7 @@ namespace UE::Audio::Insights
 		return FSlateStyle::Get().CreateIcon("AudioInsights.Icon");
 	}
 
-	TSharedRef<SWidget> FAudioAnalyzerRackDashboardViewFactory::MakeWidget()
+	TSharedRef<SWidget> FAudioAnalyzerRackDashboardViewFactory::MakeWidget(TSharedRef<SDockTab> OwnerTab)
 	{
 		// Create SubmixAudioAnalyzerRack (use MainSubmix as default)
 		if (TWeakObjectPtr<UAudioSettings> AudioSettings = GetMutableDefault<UAudioSettings>();
@@ -58,7 +58,7 @@ namespace UE::Audio::Insights
 				.Color(FSlateStyle::Get().GetColor("AudioInsights.Analyzers.BackgroundColor"));
 		}
 
-		return SubmixAudioAnalyzerRack->MakeWidget();
+		return SubmixAudioAnalyzerRack->MakeWidget(OwnerTab);
 	}
 
 	void FAudioAnalyzerRackDashboardViewFactory::HandleOnActiveAudioDeviceChanged()
diff --git a/Engine/Plugins/AudioInsights/Source/AudioInsightsEditor/Private/Views/AudioAnalyzerRackDashboardViewFactory.h b/Engine/Plugins/AudioInsights/Source/AudioInsightsEditor/Private/Views/AudioAnalyzerRackDashboardViewFactory.h
--- a/Engine/Plugins/AudioInsights/Source/AudioInsightsEditor/Private/Views/AudioAnalyzerRackDashboardViewFactory.h
+++ b/Engine/Plugins/AudioInsights/Source/AudioInsightsEditor/Private/Views/AudioAnalyzerRackDashboardViewFactory.h
@@ -17,7 +17,7 @@ namespace UE::Audio::Insights
 		virtual FText GetDisplayName() const override;
 		virtual EDefaultDashboardTabStack GetDefaultTabStack() const override;
 		virtual FSlateIcon GetIcon() const override;
-		virtual TSharedRef<SWidget> MakeWidget() override;
+		virtual TSharedRef<SWidget> MakeWidget(TSharedRef<SDockTab> OwnerTab) override;
 
 	private:
 		void HandleOnActiveAudioDeviceChanged();
diff --git a/Engine/Plugins/AudioInsights/Source/AudioInsightsEditor/Private/Views/AudioBusesDashboardViewFactory.cpp b/Engine/Plugins/AudioInsights/Source/AudioInsightsEditor/Private/Views/AudioBusesDashboardViewFactory.cpp
--- a/Engine/Plugins/AudioInsights/Source/AudioInsightsEditor/Private/Views/AudioBusesDashboardViewFactory.cpp
+++ b/Engine/Plugins/AudioInsights/Source/AudioInsightsEditor/Private/Views/AudioBusesDashboardViewFactory.cpp
@@ -135,11 +135,11 @@ namespace UE::Audio::Insights
 		return EDefaultDashboardTabStack::Analysis;
 	}
 
-	TSharedRef<SWidget> FAudioBusesDashboardViewFactory::MakeWidget()
+	TSharedRef<SWidget> FAudioBusesDashboardViewFactory::MakeWidget(TSharedRef<SDockTab> OwnerTab)
 	{
 		if (!DashboardWidget.IsValid())
 		{
-			DashboardWidget = FTraceTableDashboardViewFactory::MakeWidget();
+			DashboardWidget = FTraceTableDashboardViewFactory::MakeWidget(OwnerTab);
 
 			if (FilteredEntriesListView.IsValid())
 			{
diff --git a/Engine/Plugins/AudioInsights/Source/AudioInsightsEditor/Private/Views/AudioBusesDashboardViewFactory.h b/Engine/Plugins/AudioInsights/Source/AudioInsightsEditor/Private/Views/AudioBusesDashboardViewFactory.h
--- a/Engine/Plugins/AudioInsights/Source/AudioInsightsEditor/Private/Views/AudioBusesDashboardViewFactory.h
+++ b/Engine/Plugins/AudioInsights/Source/AudioInsightsEditor/Private/Views/AudioBusesDashboardViewFactory.h
@@ -22,7 +22,7 @@ namespace UE::Audio::Insights
 		virtual FText GetDisplayName() const override;
 		virtual FSlateIcon GetIcon() const override;
 		virtual EDefaultDashboardTabStack GetDefaultTabStack() const override;
-		virtual TSharedRef<SWidget> MakeWidget() override;
+		virtual TSharedRef<SWidget> MakeWidget(TSharedRef<SDockTab> OwnerTab) override;
 
 		DECLARE_MULTICAST_DELEGATE_TwoParams(FOnAudioBusAssetChecked, const bool /*bIsChecked*/, const TWeakObjectPtr<UAudioBus> /*AudioBus*/);
 		inline static FOnAudioBusAssetChecked OnAudioBusAssetChecked;
diff --git a/Engine/Plugins/AudioInsights/Source/AudioInsightsEditor/Private/Views/AudioMetersDashboardViewFactory.cpp b/Engine/Plugins/AudioInsights/Source/AudioInsightsEditor/Private/Views/AudioMetersDashboardViewFactory.cpp
--- a/Engine/Plugins/AudioInsights/Source/AudioInsightsEditor/Private/Views/AudioMetersDashboardViewFactory.cpp
+++ b/Engine/Plugins/AudioInsights/Source/AudioInsightsEditor/Private/Views/AudioMetersDashboardViewFactory.cpp
@@ -34,7 +34,7 @@ namespace UE::Audio::Insights
 		return FSlateStyle::Get().CreateIcon("AudioInsights.Icon.Submix");
 	}
 
-	TSharedRef<SWidget> FAudioMetersDashboardViewFactory::MakeWidget()
+	TSharedRef<SWidget> FAudioMetersDashboardViewFactory::MakeWidget(TSharedRef<SDockTab> OwnerTab)
 	{
 		if (!OnSubmixAssetCheckedHandle.IsValid())
 		{
diff --git a/Engine/Plugins/AudioInsights/Source/AudioInsightsEditor/Private/Views/AudioMetersDashboardViewFactory.h b/Engine/Plugins/AudioInsights/Source/AudioInsightsEditor/Private/Views/AudioMetersDashboardViewFactory.h
--- a/Engine/Plugins/AudioInsights/Source/AudioInsightsEditor/Private/Views/AudioMetersDashboardViewFactory.h
+++ b/Engine/Plugins/AudioInsights/Source/AudioInsightsEditor/Private/Views/AudioMetersDashboardViewFactory.h
@@ -23,7 +23,7 @@ namespace UE::Audio::Insights
 		virtual FText GetDisplayName() const override;
 		virtual EDefaultDashboardTabStack GetDefaultTabStack() const override;
 		virtual FSlateIcon GetIcon() const override;
-		virtual TSharedRef<SWidget> MakeWidget() override;
+		virtual TSharedRef<SWidget> MakeWidget(TSharedRef<SDockTab> OwnerTab) override;
 
 	private:
 		void HandleOnSubmixAssetChecked(const bool bInIsChecked, const TWeakObjectPtr<USoundSubmix> InSoundSubmix);
diff --git a/Engine/Plugins/AudioInsights/Source/AudioInsightsEditor/Private/Views/SubmixesDashboardViewFactory.cpp b/Engine/Plugins/AudioInsights/Source/AudioInsightsEditor/Private/Views/SubmixesDashboardViewFactory.cpp
--- a/Engine/Plugins/AudioInsights/Source/AudioInsightsEditor/Private/Views/SubmixesDashboardViewFactory.cpp
+++ b/Engine/Plugins/AudioInsights/Source/AudioInsightsEditor/Private/Views/SubmixesDashboardViewFactory.cpp
@@ -135,11 +135,11 @@ namespace UE::Audio::Insights
 		return EDefaultDashboardTabStack::Analysis;
 	}
 
-	TSharedRef<SWidget> FSubmixesDashboardViewFactory::MakeWidget()
+	TSharedRef<SWidget> FSubmixesDashboardViewFactory::MakeWidget(TSharedRef<SDockTab> OwnerTab)
 	{
 		if (!DashboardWidget.IsValid())
 		{
-			DashboardWidget = FTraceTableDashboardViewFactory::MakeWidget();
+			DashboardWidget = FTraceTableDashboardViewFactory::MakeWidget(OwnerTab);
 
 			if (FilteredEntriesListView.IsValid())
 			{
diff --git a/Engine/Plugins/AudioInsights/Source/AudioInsightsEditor/Private/Views/SubmixesDashboardViewFactory.h b/Engine/Plugins/AudioInsights/Source/AudioInsightsEditor/Private/Views/SubmixesDashboardViewFactory.h
--- a/Engine/Plugins/AudioInsights/Source/AudioInsightsEditor/Private/Views/SubmixesDashboardViewFactory.h
+++ b/Engine/Plugins/AudioInsights/Source/AudioInsightsEditor/Private/Views/SubmixesDashboardViewFactory.h
@@ -21,7 +21,7 @@ namespace UE::Audio::Insights
 		virtual FText GetDisplayName() const override;
 		virtual FSlateIcon GetIcon() const override;
 		virtual EDefaultDashboardTabStack GetDefaultTabStack() const override;
-		virtual TSharedRef<SWidget> MakeWidget() override;
+		virtual TSharedRef<SWidget> MakeWidget(TSharedRef<SDockTab> OwnerTab) override;
 
 		DECLARE_MULTICAST_DELEGATE_TwoParams(FOnSubmixAssetChecked, const bool /*bIsChecked*/, TWeakObjectPtr<USoundSubmix> /*SoundSubmix*/);
 		inline static FOnSubmixAssetChecked OnSubmixAssetChecked;
diff --git a/Engine/Plugins/AudioInsights/Source/AudioInsightsEditor/Private/Views/ViewportDashboardViewFactory.h b/Engine/Plugins/AudioInsights/Source/AudioInsightsEditor/Private/Views/ViewportDashboardViewFactory.h
--- a/Engine/Plugins/AudioInsights/Source/AudioInsightsEditor/Private/Views/ViewportDashboardViewFactory.h
+++ b/Engine/Plugins/AudioInsights/Source/AudioInsightsEditor/Private/Views/ViewportDashboardViewFactory.h
@@ -35,7 +35,7 @@ namespace UE::Audio::Insights
 			return FSlateStyle::Get().CreateIcon("AudioInsights.Icon.Viewport");
 		}
 
-		virtual TSharedRef<SWidget> MakeWidget() override
+		virtual TSharedRef<SWidget> MakeWidget(TSharedRef<SDockTab> OwnerTab) override
 		{
 			// TODO: This test currently hijacks the main level viewport which is not ideal.
 			// Find way to make our own instance and bind to tab here.
diff --git a/Engine/Plugins/Runtime/AudioModulation/Source/AudioModulationEditor/Private/Insights/Views/ControlBusDashboardViewFactory.cpp b/Engine/Plugins/Runtime/AudioModulation/Source/AudioModulationEditor/Private/Insights/Views/ControlBusDashboardViewFactory.cpp
--- a/Engine/Plugins/Runtime/AudioModulation/Source/AudioModulationEditor/Private/Insights/Views/ControlBusDashboardViewFactory.cpp
+++ b/Engine/Plugins/Runtime/AudioModulation/Source/AudioModulationEditor/Private/Insights/Views/ControlBusDashboardViewFactory.cpp
@@ -83,9 +83,9 @@ namespace AudioModulationEditor
 		return UE::Audio::Insights::EDefaultDashboardTabStack::Analysis;
 	}
 
-	TSharedRef<SWidget> FControlBusDashboardViewFactory::MakeWidget()
+	TSharedRef<SWidget> FControlBusDashboardViewFactory::MakeWidget(TSharedRef<SDockTab> OwnerTab)
 	{
-		TSharedRef<SWidget> TableDashboardWidget = UE::Audio::Insights::FTraceTableDashboardViewFactory::MakeWidget();
+		TSharedRef<SWidget> TableDashboardWidget = UE::Audio::Insights::FTraceTableDashboardViewFactory::MakeWidget(OwnerTab);
 		TSharedRef<SWidget> BusWatchWidget = MakeControlBusWatchWidget();
 
 		return SNew(SHorizontalBox)
diff --git a/Engine/Plugins/Runtime/AudioModulation/Source/AudioModulationEditor/Private/Insights/Views/ControlBusDashboardViewFactory.h b/Engine/Plugins/Runtime/AudioModulation/Source/AudioModulationEditor/Private/Insights/Views/ControlBusDashboardViewFactory.h
--- a/Engine/Plugins/Runtime/AudioModulation/Source/AudioModulationEditor/Private/Insights/Views/ControlBusDashboardViewFactory.h
+++ b/Engine/Plugins/Runtime/AudioModulation/Source/AudioModulationEditor/Private/Insights/Views/ControlBusDashboardViewFactory.h
@@ -19,7 +19,7 @@ namespace AudioModulationEditor
 		virtual FSlateIcon GetIcon() const override;
 		virtual UE::Audio::Insights::EDefaultDashboardTabStack GetDefaultTabStack() const override;
 		virtual void ProcessEntries(UE::Audio::Insights::FTraceTableDashboardViewFactory::EProcessReason Reason) override;
-		virtual TSharedRef<SWidget> MakeWidget() override;
+		virtual TSharedRef<SWidget> MakeWidget(TSharedRef<SDockTab> OwnerTab) override;
 
 	protected:
 		virtual const TMap<FName, UE::Audio::Insights::FTraceTableDashboardViewFactory::FColumnData>& GetColumns() const override;
