[Audio Insights] MixerSourceTraceProvider: Added missing base class call.

diff --git a/Engine/Plugins/AudioInsights/Source/AudioInsights/Private/Providers/MixerSourceTraceProvider.cpp b/Engine/Plugins/AudioInsights/Source/AudioInsights/Private/Providers/MixerSourceTraceProvider.cpp
--- a/Engine/Plugins/AudioInsights/Source/AudioInsights/Private/Providers/MixerSourceTraceProvider.cpp
+++ b/Engine/Plugins/AudioInsights/Source/AudioInsights/Private/Providers/MixerSourceTraceProvider.cpp
@@ -452,6 +452,8 @@ namespace UE::Audio::Insights
 
 			virtual void OnAnalysisBegin(const FOnAnalysisContext& Context) override
 			{
+				FTraceAnalyzerBase::OnAnalysisBegin(Context);
+
 				UE::Trace::IAnalyzer::FInterfaceBuilder& Builder = Context.InterfaceBuilder;
 				Builder.RouteEvent(RouteId_DistanceAttenuation, "Audio", "MixerSourceDistanceAttenuation");
 				Builder.RouteEvent(RouteId_Envelope, "Audio", "MixerSourceEnvelope");
