[Audio Insights] FTraceTableDashboardViewFactory: Added horizontal scrollbox to table based dashboards (specially useful for the incoming modulation matrix dashboard that could contain a lot of "horizontal" data)

diff --git a/Engine/Plugins/AudioInsights/Source/AudioInsights/Private/Views/TableDashboardViewFactory.cpp b/Engine/Plugins/AudioInsights/Source/AudioInsights/Private/Views/TableDashboardViewFactory.cpp
--- a/Engine/Plugins/AudioInsights/Source/AudioInsights/Private/Views/TableDashboardViewFactory.cpp
+++ b/Engine/Plugins/AudioInsights/Source/AudioInsights/Private/Views/TableDashboardViewFactory.cpp
@@ -6,6 +6,7 @@
 #include "Containers/Array.h"
 #include "Framework/Commands/UICommandList.h"
 #include "Templates/SharedPointer.h"
+#include "Widgets/Layout/SScrollBox.h"
 #include "Widgets/SBoxPanel.h"
 
 #if WITH_EDITOR
@@ -169,7 +170,15 @@ namespace UE::Audio::Insights
 			.FillHeight(1.0f)
 			.Padding(0, 2)
 			[
-				FilteredEntriesListView->AsShared()
+				SNew(SScrollBox)
+				.Orientation(EOrientation::Orient_Horizontal)
+				+ SScrollBox::Slot()
+				.Padding(0)
+				.FillSize(1.0f)
+				.HAlign(HAlign_Fill)
+				[
+					FilteredEntriesListView->AsShared()
+				]
 			];
 		}
 
