[Audio Insights] FTraceTableDashboardViewFactory: Make HeaderRowWidget accessible to derived classes so it can be dynamically modified (specially useful for the incoming modulation matrix dashboard where bus columns can be added/removed dynamically)

diff --git a/Engine/Plugins/AudioInsights/Source/AudioInsights/Private/Views/TableDashboardViewFactory.cpp b/Engine/Plugins/AudioInsights/Source/AudioInsights/Private/Views/TableDashboardViewFactory.cpp
--- a/Engine/Plugins/AudioInsights/Source/AudioInsights/Private/Views/TableDashboardViewFactory.cpp
+++ b/Engine/Plugins/AudioInsights/Source/AudioInsights/Private/Views/TableDashboardViewFactory.cpp
@@ -99,7 +99,7 @@ namespace UE::Audio::Insights
 			[](const TPair<FName, FColumnData>& ColumnInfo) { return ColumnInfo.Value.bDefaultHidden; },
 			[](const TPair<FName, FColumnData>& ColumnInfo) { return ColumnInfo.Key; });
 
-		TSharedRef<SHeaderRow> HeaderRowWidget = SNew(SHeaderRow)
+		SAssignNew(HeaderRowWidget, SHeaderRow)
 		.CanSelectGeneratedColumn(true); // Allows for showing/hiding columns
 
 		// This only works if header row columns are added with slots and not programmatically
@@ -124,7 +124,7 @@ namespace UE::Audio::Insights
 			HeaderRowWidget->AddColumn(*NewColumn);
 		}
 
-		return HeaderRowWidget;
+		return HeaderRowWidget.ToSharedRef();
 	}
 
 	TSharedRef<SWidget> FTraceTableDashboardViewFactory::MakeWidget(TSharedRef<SDockTab> OwnerTab)
diff --git a/Engine/Plugins/AudioInsights/Source/AudioInsights/Public/Views/TableDashboardViewFactory.h b/Engine/Plugins/AudioInsights/Source/AudioInsights/Public/Views/TableDashboardViewFactory.h
--- a/Engine/Plugins/AudioInsights/Source/AudioInsights/Public/Views/TableDashboardViewFactory.h
+++ b/Engine/Plugins/AudioInsights/Source/AudioInsights/Public/Views/TableDashboardViewFactory.h
@@ -161,6 +161,7 @@ namespace UE::Audio::Insights
 		TMap<FName, uint64> UpdateIds;
 
 		TSharedPtr<SWidget> DashboardWidget;
+		TSharedPtr<SHeaderRow> HeaderRowWidget;
 		TSharedPtr<SListView<TSharedPtr<IDashboardDataViewEntry>>> FilteredEntriesListView;
 
 		FName SortByColumn;
