[Audio Insights] TDeviceDataMapTraceProvider: Don't process message queues if trace is not active + do proper cleanup of DeviceDataMap (when trace is stopped or when an audio device is destroyed).

diff --git a/Engine/Plugins/AudioInsights/Source/AudioInsights/Public/AudioInsightsTraceProviderBase.h b/Engine/Plugins/AudioInsights/Source/AudioInsights/Public/AudioInsightsTraceProviderBase.h
--- a/Engine/Plugins/AudioInsights/Source/AudioInsights/Public/AudioInsightsTraceProviderBase.h
+++ b/Engine/Plugins/AudioInsights/Source/AudioInsights/Public/AudioInsightsTraceProviderBase.h
@@ -13,6 +13,10 @@
 #include "TraceServices/Model/AnalysisSession.h"
 #include "UObject/NameTypes.h"
 
+#if WITH_EDITOR
+#include "AudioDeviceManager.h"
+#include "ProfilingDebugging/TraceAuxiliary.h"
+#endif // WITH_EDITOR
 
 namespace UE::Audio::Insights
 {
@@ -113,10 +117,18 @@ namespace UE::Audio::Insights
 				LastUpdateId = GetLastMessageId();
 				return true;
 			});
+
+#if WITH_EDITOR
+			RegisterDelegates();
+#endif // WITH_EDITOR
 		}
 
 		virtual ~TDeviceDataMapTraceProvider()
 		{
+#if WITH_EDITOR
+			UnregisterDelegates();
+#endif // WITH_EDITOR
+
 			FTSTicker::GetCoreTicker().RemoveTicker(TickerHandle);
 		}
 
@@ -153,15 +165,19 @@ namespace UE::Audio::Insights
 		void ProcessMessageQueue(
 			TAnalyzerMessageQueue<TMsgType>& InQueue,
 			TFunctionRef<EntryType*(const TMsgType&)> GetEntry,
-			TFunctionRef<void(const TMsgType&, EntryType*)> ProcessEntry
-		)
+			TFunctionRef<void(const TMsgType&, EntryType*)> ProcessEntry)
 		{
-			TArray<TMsgType> Messages = InQueue.DequeueAll();
-			for (const TMsgType& Msg : Messages)
+#if WITH_EDITOR
+			if (bIsTraceActive)
+#endif // WITH_EDITOR
 			{
-				EntryType* Entry = GetEntry(Msg);
-				ProcessEntry(Msg, Entry);
-			};
+				TArray<TMsgType> Messages = InQueue.DequeueAll();
+				for (const TMsgType& Msg : Messages)
+				{
+					EntryType* Entry = GetEntry(Msg);
+					ProcessEntry(Msg, Entry);
+				};
+			}
 		}
 
 		EntryType* FindDeviceEntry(::Audio::FDeviceId InDeviceId, const EntryKey& InKey)
@@ -223,6 +239,40 @@ namespace UE::Audio::Insights
 		}
 
 	private:
+#if WITH_EDITOR
+		void RegisterDelegates()
+		{
+			FTraceAuxiliary::OnTraceStarted.AddRaw(this, &TDeviceDataMapTraceProvider::OnTraceStarted);			
+			FTraceAuxiliary::OnTraceStopped.AddRaw(this, &TDeviceDataMapTraceProvider::OnTraceStopped);
+			FAudioDeviceManagerDelegates::OnAudioDeviceDestroyed.AddRaw(this, &TDeviceDataMapTraceProvider::OnAudioDeviceDestroyed);
+		}
+
+		void UnregisterDelegates()
+		{
+			FAudioDeviceManagerDelegates::OnAudioDeviceDestroyed.RemoveAll(this);
+			FTraceAuxiliary::OnTraceStopped.RemoveAll(this);
+			FTraceAuxiliary::OnTraceStarted.RemoveAll(this);
+		}
+
+		void OnTraceStarted(FTraceAuxiliary::EConnectionType InTraceType, const FString& InTraceDestination)
+		{
+			bIsTraceActive = true;
+		}
+
+		void OnTraceStopped(FTraceAuxiliary::EConnectionType InTraceType, const FString& InTraceDestination)
+		{
+			bIsTraceActive = false;
+			Reset();
+		}
+
+		void OnAudioDeviceDestroyed(::Audio::FDeviceId InDeviceId)
+		{
+			DeviceDataMap.Remove(InDeviceId);
+		}
+
+		bool bIsTraceActive = false;
+#endif // WITH_EDITOR
+
 		FTSTicker::FDelegateHandle TickerHandle;
 	};
 } // namespace UE::Audio::Insights
