[Audio Insights] IAudioInsightsEditorModule: Added IsModuleLoaded method

diff --git a/Engine/Plugins/AudioInsights/Source/AudioInsightsEditor/Private/AudioInsightsEditorModule.cpp b/Engine/Plugins/AudioInsights/Source/AudioInsightsEditor/Private/AudioInsightsEditorModule.cpp
--- a/Engine/Plugins/AudioInsights/Source/AudioInsightsEditor/Private/AudioInsightsEditorModule.cpp
+++ b/Engine/Plugins/AudioInsights/Source/AudioInsightsEditor/Private/AudioInsightsEditorModule.cpp
@@ -68,16 +68,20 @@ namespace UE::Audio::Insights
 		return DashboardFactory->GetDeviceId();
 	}
 
+	bool FAudioInsightsEditorModule::IsModuleLoaded()
+	{
+		return FModuleManager::Get().IsModuleLoaded(AudioInsightsEditorModuleName);
+	}
+
 	FAudioInsightsEditorModule& FAudioInsightsEditorModule::GetChecked()
 	{
-		return static_cast<FAudioInsightsEditorModule&>(FModuleManager::GetModuleChecked<IAudioInsightsEditorModule>("AudioInsightsEditor"));
+		return static_cast<FAudioInsightsEditorModule&>(FModuleManager::GetModuleChecked<IAudioInsightsEditorModule>(AudioInsightsEditorModuleName));
 	}
 
 	IAudioInsightsTraceModule& FAudioInsightsEditorModule::GetTraceModule()
 	{
 		IAudioInsightsModule& InsightsModule = IAudioInsightsModule::GetChecked();
 		return InsightsModule.GetTraceModule();
-
 	}
 
 	TSharedRef<FEditorDashboardFactory> FAudioInsightsEditorModule::GetDashboardFactory()
diff --git a/Engine/Plugins/AudioInsights/Source/AudioInsightsEditor/Private/AudioInsightsEditorModule.h b/Engine/Plugins/AudioInsights/Source/AudioInsightsEditor/Private/AudioInsightsEditorModule.h
--- a/Engine/Plugins/AudioInsights/Source/AudioInsightsEditor/Private/AudioInsightsEditorModule.h
+++ b/Engine/Plugins/AudioInsights/Source/AudioInsightsEditor/Private/AudioInsightsEditorModule.h
@@ -29,6 +29,7 @@ namespace UE::Audio::Insights
 		TSharedRef<FEditorDashboardFactory> GetDashboardFactory();
 		const TSharedRef<FEditorDashboardFactory> GetDashboardFactory() const;
 
+		static bool IsModuleLoaded();
 		static FAudioInsightsEditorModule& GetChecked();
 		virtual IAudioInsightsTraceModule& GetTraceModule() override;
 
@@ -36,6 +37,8 @@ namespace UE::Audio::Insights
 		void RegisterMenus();
 		TSharedRef<SDockTab> CreateDashboardTabWidget(const FSpawnTabArgs& Args);
 		
+		inline static const FName AudioInsightsEditorModuleName = "AudioInsightsEditor";
+
 		TSharedPtr<FEditorDashboardFactory> DashboardFactory;
 		
 		FVirtualLoopsDebugDraw VirtualLoopsDebugDraw;
diff --git a/Engine/Plugins/AudioInsights/Source/AudioInsightsEditor/Private/IAudioInsightsEditorModule.cpp b/Engine/Plugins/AudioInsights/Source/AudioInsightsEditor/Private/IAudioInsightsEditorModule.cpp
--- a/Engine/Plugins/AudioInsights/Source/AudioInsightsEditor/Private/IAudioInsightsEditorModule.cpp
+++ b/Engine/Plugins/AudioInsights/Source/AudioInsightsEditor/Private/IAudioInsightsEditorModule.cpp
@@ -2,6 +2,7 @@
 #include "IAudioInsightsEditorModule.h"
 
 #include "AudioInsightsEditorModule.h"
+#include "Modules/ModuleManager.h"
 
 IAudioInsightsTraceModule& IAudioInsightsEditorModule::GetTraceModule()
 {
@@ -9,6 +10,11 @@ IAudioInsightsTraceModule& IAudioInsightsEditorModule::GetTraceModule()
 	return AudioInsightsModule.GetTraceModule();
 }
 
+bool IAudioInsightsEditorModule::IsModuleLoaded()
+{
+	return UE::Audio::Insights::FAudioInsightsEditorModule::IsModuleLoaded();
+}
+
 IAudioInsightsEditorModule& IAudioInsightsEditorModule::GetChecked()
 {
 	return static_cast<IAudioInsightsEditorModule&>(UE::Audio::Insights::FAudioInsightsEditorModule::GetChecked());
diff --git a/Engine/Plugins/AudioInsights/Source/AudioInsightsEditor/Public/IAudioInsightsEditorModule.h b/Engine/Plugins/AudioInsights/Source/AudioInsightsEditor/Public/IAudioInsightsEditorModule.h
--- a/Engine/Plugins/AudioInsights/Source/AudioInsightsEditor/Public/IAudioInsightsEditorModule.h
+++ b/Engine/Plugins/AudioInsights/Source/AudioInsightsEditor/Public/IAudioInsightsEditorModule.h
@@ -30,5 +30,7 @@ public:
 
 	AUDIOINSIGHTSEDITOR_API virtual TSharedRef<SDockTab> CreateDashboardTabWidget(const FSpawnTabArgs& Args) = 0;
 
+	AUDIOINSIGHTSEDITOR_API static bool IsModuleLoaded();
+
 	AUDIOINSIGHTSEDITOR_API static IAudioInsightsEditorModule& GetChecked();
 };
diff --git a/Engine/Plugins/Runtime/AudioModulation/Source/AudioModulationEditor/Private/AudioModulationEditor.cpp b/Engine/Plugins/Runtime/AudioModulation/Source/AudioModulationEditor/Private/AudioModulationEditor.cpp
--- a/Engine/Plugins/Runtime/AudioModulation/Source/AudioModulationEditor/Private/AudioModulationEditor.cpp
+++ b/Engine/Plugins/Runtime/AudioModulation/Source/AudioModulationEditor/Private/AudioModulationEditor.cpp
@@ -72,9 +72,13 @@ namespace
 	{
 		if (bIsModulationRegisteredInAudioInsights)
 		{
-			IAudioInsightsEditorModule& InsightsModule = IAudioInsightsEditorModule::GetChecked();
-			InsightsModule.UnregisterDashboardViewFactory("ControlBuses");
-			bIsModulationRegisteredInAudioInsights = false;
+			if (IAudioInsightsEditorModule::IsModuleLoaded())
+			{
+				IAudioInsightsEditorModule& InsightsEditorModule = IAudioInsightsEditorModule::GetChecked();
+				InsightsEditorModule.UnregisterDashboardViewFactory("ControlBuses");
+
+				bIsModulationRegisteredInAudioInsights = false;
+			}
 		}
 	}
 
