[AudioInsights] Reset Layout will apply to SubTabManagers also.

diff --git a/Engine/Plugins/AudioInsights/Source/AudioInsightsEditor/Private/AudioInsightsEditorDashboardFactory.cpp b/Engine/Plugins/AudioInsights/Source/AudioInsightsEditor/Private/AudioInsightsEditorDashboardFactory.cpp
--- a/Engine/Plugins/AudioInsights/Source/AudioInsightsEditor/Private/AudioInsightsEditorDashboardFactory.cpp
+++ b/Engine/Plugins/AudioInsights/Source/AudioInsightsEditor/Private/AudioInsightsEditorDashboardFactory.cpp
@@ -357,6 +357,26 @@ namespace UE::Audio::Insights
 					{
 						if (DashboardTabManager)
 						{
+							for (const auto& KVP : DashboardViewFactories)
+							{
+								// Try and get the dashboard tab:
+								TSharedPtr<SDockTab> DashboardTab = DashboardTabManager->FindExistingLiveTab(KVP.Key);
+								if (!DashboardTab.IsValid())
+								{
+									DashboardTab = DashboardTabManager->TryInvokeTab(KVP.Key);
+								}
+
+								if (DashboardTab.IsValid())
+								{
+									if (TSharedPtr<FTabManager> SubTabManager = FGlobalTabmanager::Get()->GetTabManagerForMajorTab(DashboardTab))
+									{
+										// There is a sub tab manager for this dashbaord tab; clear its persisted areas:
+										SubTabManager->CloseAllAreas();
+										SubTabManager->SavePersistentLayout();
+									}
+								}
+							}
+
 							if (TSharedPtr<SDockTab> OwnerTab = DashboardTabManager->GetOwnerTab())
 							{
 								// Wipe all the persisted areas and close tab
