[AudioInsights] Explicitly close dashboard tabs when audio insights is closed. This will allow each dashboard to close any undocked sub-managed tabs of its own.

diff --git a/Engine/Plugins/AudioInsights/Source/AudioInsightsEditor/Private/AudioInsightsEditorDashboardFactory.cpp b/Engine/Plugins/AudioInsights/Source/AudioInsightsEditor/Private/AudioInsightsEditorDashboardFactory.cpp
--- a/Engine/Plugins/AudioInsights/Source/AudioInsightsEditor/Private/AudioInsightsEditorDashboardFactory.cpp
+++ b/Engine/Plugins/AudioInsights/Source/AudioInsightsEditor/Private/AudioInsightsEditorDashboardFactory.cpp
@@ -258,6 +258,15 @@ namespace UE::Audio::Insights
 			UnregisterTabSpawners();
 			SaveLayoutToConfig();
 
+			for (const auto& KVP : DashboardViewFactories)
+			{
+				if (TSharedPtr<SDockTab> DashboardTab = DashboardTabManager->FindExistingLiveTab(KVP.Key))
+				{
+					// Explicitly close each dashboard tab. This will give a chance for it to close any undocked sub-managed tabs of its own:
+					DashboardTab->RequestCloseTab();
+				}
+			}
+
 			DashboardTabManager->CloseAllAreas();
 
 			DashboardTabManager.Reset();
