[Audio Modulation Insights] Load AudioInsights modules if needed. (this is to fix a bug where the AudioinsightsEditor module has not been loaded yet when AudioModulationEditor requests it)

diff --git a/Engine/Plugins/AudioInsights/Source/AudioInsights/Private/AudioInsightsModule.cpp b/Engine/Plugins/AudioInsights/Source/AudioInsights/Private/AudioInsightsModule.cpp
--- a/Engine/Plugins/AudioInsights/Source/AudioInsights/Private/AudioInsightsModule.cpp
+++ b/Engine/Plugins/AudioInsights/Source/AudioInsights/Private/AudioInsightsModule.cpp
@@ -100,7 +100,7 @@ namespace UE::Audio::Insights
 
 	FAudioInsightsModule& FAudioInsightsModule::GetChecked()
 	{
-		return static_cast<FAudioInsightsModule&>(FModuleManager::GetModuleChecked<IAudioInsightsModule>("AudioInsights"));
+		return static_cast<FAudioInsightsModule&>(FModuleManager::LoadModuleChecked<IAudioInsightsModule>("AudioInsights"));
 	}
 
 	TSharedRef<FDashboardFactory> FAudioInsightsModule::GetDashboardFactory()
diff --git a/Engine/Plugins/AudioInsights/Source/AudioInsightsEditor/Private/AudioInsightsEditorModule.cpp b/Engine/Plugins/AudioInsights/Source/AudioInsightsEditor/Private/AudioInsightsEditorModule.cpp
--- a/Engine/Plugins/AudioInsights/Source/AudioInsightsEditor/Private/AudioInsightsEditorModule.cpp
+++ b/Engine/Plugins/AudioInsights/Source/AudioInsightsEditor/Private/AudioInsightsEditorModule.cpp
@@ -75,7 +75,7 @@ namespace UE::Audio::Insights
 
 	FAudioInsightsEditorModule& FAudioInsightsEditorModule::GetChecked()
 	{
-		return static_cast<FAudioInsightsEditorModule&>(FModuleManager::GetModuleChecked<IAudioInsightsEditorModule>(AudioInsightsEditorModuleName));
+		return static_cast<FAudioInsightsEditorModule&>(FModuleManager::LoadModuleChecked<IAudioInsightsEditorModule>(AudioInsightsEditorModuleName));
 	}
 
 	IAudioInsightsTraceModule& FAudioInsightsEditorModule::GetTraceModule()
