[AudioInsights] Undocked tabs maintain always on top state after reopening editor. Passing owner window through to TabManager RestoreFrom layout calls.

diff --git a/Engine/Plugins/AudioInsights/Source/AudioInsights/Private/AudioInsightsDashboardFactory.cpp b/Engine/Plugins/AudioInsights/Source/AudioInsights/Private/AudioInsightsDashboardFactory.cpp
--- a/Engine/Plugins/AudioInsights/Source/AudioInsights/Private/AudioInsightsDashboardFactory.cpp
+++ b/Engine/Plugins/AudioInsights/Source/AudioInsights/Private/AudioInsightsDashboardFactory.cpp
@@ -52,7 +52,7 @@ namespace UE::Audio::Insights
 			]
 			+ SVerticalBox::Slot()
 			[
-				DashboardTabManager->RestoreFrom(TabLayout->AsShared(), TSharedPtr<SWindow>()).ToSharedRef()
+				DashboardTabManager->RestoreFrom(TabLayout->AsShared(), Args.GetOwnerWindow()).ToSharedRef()
 			];
 
 		DockTab->SetContent(TabContent);
@@ -224,7 +224,7 @@ namespace UE::Audio::Insights
 					.Clipping(EWidgetClipping::ClipToBounds)
 					.Label(Factory->GetDisplayName());
 
-				TSharedRef<SWidget> DashboardView = Factory->MakeWidget(DockTab);
+				TSharedRef<SWidget> DashboardView = Factory->MakeWidget(DockTab, Args);
 				DockTab->SetContent(DashboardView);
 
 				return DockTab;
diff --git a/Engine/Plugins/AudioInsights/Source/AudioInsights/Private/Views/LogDashboardViewFactory.cpp b/Engine/Plugins/AudioInsights/Source/AudioInsights/Private/Views/LogDashboardViewFactory.cpp
--- a/Engine/Plugins/AudioInsights/Source/AudioInsights/Private/Views/LogDashboardViewFactory.cpp
+++ b/Engine/Plugins/AudioInsights/Source/AudioInsights/Private/Views/LogDashboardViewFactory.cpp
@@ -54,7 +54,7 @@ namespace UE::Audio::Insights
 		return FSlateStyle::Get().CreateIcon("AudioInsights.Icon.Log");
 	}
 
-	TSharedRef<SWidget> FLogDashboardViewFactory::MakeWidget(TSharedRef<SDockTab> OwnerTab)
+	TSharedRef<SWidget> FLogDashboardViewFactory::MakeWidget(TSharedRef<SDockTab> OwnerTab, const FSpawnTabArgs& SpawnTabArgs)
 	{
 		using namespace LogDashboardViewFactoryPrivate;
 
diff --git a/Engine/Plugins/AudioInsights/Source/AudioInsights/Private/Views/MixerSourceDashboardViewFactory.cpp b/Engine/Plugins/AudioInsights/Source/AudioInsights/Private/Views/MixerSourceDashboardViewFactory.cpp
--- a/Engine/Plugins/AudioInsights/Source/AudioInsights/Private/Views/MixerSourceDashboardViewFactory.cpp
+++ b/Engine/Plugins/AudioInsights/Source/AudioInsights/Private/Views/MixerSourceDashboardViewFactory.cpp
@@ -1004,7 +1004,7 @@ namespace UE::Audio::Insights
 				+ SVerticalBox::Slot()
 				.HAlign(HAlign_Fill)
 				[
-					FTraceTableDashboardViewFactory::MakeWidget(DockTab)
+					FTraceTableDashboardViewFactory::MakeWidget(DockTab, Args)
 				]
 			);
 
@@ -1102,7 +1102,7 @@ namespace UE::Audio::Insights
 			);
 	}
 
-	TSharedRef<SWidget> FMixerSourceDashboardViewFactory::MakeWidget(TSharedRef<SDockTab> OwnerTab)
+	TSharedRef<SWidget> FMixerSourceDashboardViewFactory::MakeWidget(TSharedRef<SDockTab> OwnerTab, const FSpawnTabArgs& SpawnTabArgs)
 	{
 		using namespace MixerSourcePrivate;
 
@@ -1166,7 +1166,7 @@ namespace UE::Audio::Insights
 		const TSharedRef<FTabManager::FLayout> TabLayout = GetDefaultTabLayout();
 #endif // WITH_EDITOR
 
-		return MixerSourcesTabManager->RestoreFrom(TabLayout, TSharedPtr<SWindow>()).ToSharedRef();
+		return MixerSourcesTabManager->RestoreFrom(TabLayout, SpawnTabArgs.GetOwnerWindow()).ToSharedRef();
 	}	
 
 	void FMixerSourceDashboardViewFactory::ProcessEntries(FTraceTableDashboardViewFactory::EProcessReason InReason)
diff --git a/Engine/Plugins/AudioInsights/Source/AudioInsights/Private/Views/TableDashboardViewFactory.cpp b/Engine/Plugins/AudioInsights/Source/AudioInsights/Private/Views/TableDashboardViewFactory.cpp
--- a/Engine/Plugins/AudioInsights/Source/AudioInsights/Private/Views/TableDashboardViewFactory.cpp
+++ b/Engine/Plugins/AudioInsights/Source/AudioInsights/Private/Views/TableDashboardViewFactory.cpp
@@ -127,7 +127,7 @@ namespace UE::Audio::Insights
 		return HeaderRowWidget.ToSharedRef();
 	}
 
-	TSharedRef<SWidget> FTraceTableDashboardViewFactory::MakeWidget(TSharedRef<SDockTab> OwnerTab)
+	TSharedRef<SWidget> FTraceTableDashboardViewFactory::MakeWidget(TSharedRef<SDockTab> OwnerTab, const FSpawnTabArgs& SpawnTabArgs)
 	{
 		if (!DashboardWidget.IsValid())
 		{
@@ -363,7 +363,7 @@ namespace UE::Audio::Insights
 		return FReply::Unhandled();
 	}
 
-	TSharedRef<SWidget> FTraceObjectTableDashboardViewFactory::MakeWidget(TSharedRef<SDockTab> OwnerTab)
+	TSharedRef<SWidget> FTraceObjectTableDashboardViewFactory::MakeWidget(TSharedRef<SDockTab> OwnerTab, const FSpawnTabArgs& SpawnTabArgs)
 	{
 		if (!DashboardWidget.IsValid())
 		{
@@ -385,7 +385,7 @@ namespace UE::Audio::Insights
 			.AutoHeight()
 			.HAlign(HAlign_Fill)
 			[
-				FTraceTableDashboardViewFactory::MakeWidget(OwnerTab)
+				FTraceTableDashboardViewFactory::MakeWidget(OwnerTab, SpawnTabArgs)
 			];
 
 			if (FilteredEntriesListView.IsValid())
diff --git a/Engine/Plugins/AudioInsights/Source/AudioInsights/Public/Views/DashboardViewFactory.h b/Engine/Plugins/AudioInsights/Source/AudioInsights/Public/Views/DashboardViewFactory.h
--- a/Engine/Plugins/AudioInsights/Source/AudioInsights/Public/Views/DashboardViewFactory.h
+++ b/Engine/Plugins/AudioInsights/Source/AudioInsights/Public/Views/DashboardViewFactory.h
@@ -9,6 +9,7 @@
 #include "UObject/NameTypes.h"
 #include "Widgets/SWidget.h"
 
+class FSpawnTabArgs;
 class SDockTab;
 
 namespace UE::Audio::Insights
@@ -31,7 +32,7 @@ namespace UE::Audio::Insights
 		virtual FText GetDisplayName() const = 0;
 		virtual FName GetName() const = 0;
 		virtual FSlateIcon GetIcon() const = 0;
-		virtual TSharedRef<SWidget> MakeWidget(TSharedRef<SDockTab> OwnerTab) = 0;
+		virtual TSharedRef<SWidget> MakeWidget(TSharedRef<SDockTab> OwnerTab, const FSpawnTabArgs& SpawnTabArgs) = 0;
 	};
 
 	class FTraceDashboardViewFactoryBase : public IDashboardViewFactory
diff --git a/Engine/Plugins/AudioInsights/Source/AudioInsights/Public/Views/LogDashboardViewFactory.h b/Engine/Plugins/AudioInsights/Source/AudioInsights/Public/Views/LogDashboardViewFactory.h
--- a/Engine/Plugins/AudioInsights/Source/AudioInsights/Public/Views/LogDashboardViewFactory.h
+++ b/Engine/Plugins/AudioInsights/Source/AudioInsights/Public/Views/LogDashboardViewFactory.h
@@ -24,7 +24,7 @@ namespace UE::Audio::Insights
 		virtual FText GetDisplayName() const override;
 		virtual EDefaultDashboardTabStack GetDefaultTabStack() const override;
 		virtual FSlateIcon GetIcon() const override;
-		virtual TSharedRef<SWidget> MakeWidget(TSharedRef<SDockTab> OwnerTab) override;
+		virtual TSharedRef<SWidget> MakeWidget(TSharedRef<SDockTab> OwnerTab, const FSpawnTabArgs& SpawnTabArgs) override;
 
 	private:
 		struct FLogCategoryCollector : public FOutputDevice
diff --git a/Engine/Plugins/AudioInsights/Source/AudioInsights/Public/Views/MixerSourceDashboardViewFactory.h b/Engine/Plugins/AudioInsights/Source/AudioInsights/Public/Views/MixerSourceDashboardViewFactory.h
--- a/Engine/Plugins/AudioInsights/Source/AudioInsights/Public/Views/MixerSourceDashboardViewFactory.h
+++ b/Engine/Plugins/AudioInsights/Source/AudioInsights/Public/Views/MixerSourceDashboardViewFactory.h
@@ -26,7 +26,7 @@ namespace UE::Audio::Insights
 		virtual FSlateIcon GetIcon() const override;
 		virtual EDefaultDashboardTabStack GetDefaultTabStack() const override;
 		virtual void ProcessEntries(FTraceTableDashboardViewFactory::EProcessReason Reason) override;
-		virtual TSharedRef<SWidget> MakeWidget(TSharedRef<SDockTab> OwnerTab) override;
+		virtual TSharedRef<SWidget> MakeWidget(TSharedRef<SDockTab> OwnerTab, const FSpawnTabArgs& SpawnTabArgs) override;
 
 		// Maximum amount of data history kept for plots (in seconds)
 		static const double MaxPlotHistorySeconds;
diff --git a/Engine/Plugins/AudioInsights/Source/AudioInsights/Public/Views/TableDashboardViewFactory.h b/Engine/Plugins/AudioInsights/Source/AudioInsights/Public/Views/TableDashboardViewFactory.h
--- a/Engine/Plugins/AudioInsights/Source/AudioInsights/Public/Views/TableDashboardViewFactory.h
+++ b/Engine/Plugins/AudioInsights/Source/AudioInsights/Public/Views/TableDashboardViewFactory.h
@@ -60,7 +60,7 @@ namespace UE::Audio::Insights
 	public:
 		virtual TSharedRef<SWidget> GenerateWidgetForColumn(TSharedRef<IDashboardDataViewEntry> InRowData, const FName& Column);
 		virtual FReply OnDataRowKeyInput(const FGeometry& Geometry, const FKeyEvent& KeyEvent) const;
-		virtual TSharedRef<SWidget> MakeWidget(TSharedRef<SDockTab> OwnerTab) override;
+		virtual TSharedRef<SWidget> MakeWidget(TSharedRef<SDockTab> OwnerTab, const FSpawnTabArgs& SpawnTabArgs) override;
 		virtual void RebuildFilteredEntriesListView();
 
 		enum class EProcessReason : uint8
@@ -186,7 +186,7 @@ namespace UE::Audio::Insights
 
 		virtual TSharedRef<SWidget> GenerateWidgetForColumn(TSharedRef<IDashboardDataViewEntry> InRowData, const FName& Column) override;
 		virtual FReply OnDataRowKeyInput(const FGeometry& Geometry, const FKeyEvent& KeyEvent) const override;
-		virtual TSharedRef<SWidget> MakeWidget(TSharedRef<SDockTab> OwnerTab) override;
+		virtual TSharedRef<SWidget> MakeWidget(TSharedRef<SDockTab> OwnerTab, const FSpawnTabArgs& SpawnTabArgs) override;
 
 	protected:
 #if WITH_EDITOR
diff --git a/Engine/Plugins/AudioInsights/Source/AudioInsightsEditor/Private/Analyzers/SubmixAudioAnalyzerRack.cpp b/Engine/Plugins/AudioInsights/Source/AudioInsightsEditor/Private/Analyzers/SubmixAudioAnalyzerRack.cpp
--- a/Engine/Plugins/AudioInsights/Source/AudioInsightsEditor/Private/Analyzers/SubmixAudioAnalyzerRack.cpp
+++ b/Engine/Plugins/AudioInsights/Source/AudioInsightsEditor/Private/Analyzers/SubmixAudioAnalyzerRack.cpp
@@ -97,9 +97,9 @@ namespace UE::Audio::Insights
 		CleanupAudioAnalyzerRack();
 	}
 
-	TSharedRef<SWidget> FSubmixAudioAnalyzerRack::MakeWidget(TSharedRef<SDockTab> InOwnerTab)
+	TSharedRef<SWidget> FSubmixAudioAnalyzerRack::MakeWidget(TSharedRef<SDockTab> InOwnerTab, const FSpawnTabArgs& InSpawnTabArgs)
 	{
-		return AudioAnalyzerRack->CreateWidget(InOwnerTab);
+		return AudioAnalyzerRack->CreateWidget(InOwnerTab, InSpawnTabArgs);
 	}
 
 	void FSubmixAudioAnalyzerRack::RebuildAudioAnalyzerRack(TWeakObjectPtr<USoundSubmix> InSoundSubmix)
diff --git a/Engine/Plugins/AudioInsights/Source/AudioInsightsEditor/Private/Analyzers/SubmixAudioAnalyzerRack.h b/Engine/Plugins/AudioInsights/Source/AudioInsightsEditor/Private/Analyzers/SubmixAudioAnalyzerRack.h
--- a/Engine/Plugins/AudioInsights/Source/AudioInsightsEditor/Private/Analyzers/SubmixAudioAnalyzerRack.h
+++ b/Engine/Plugins/AudioInsights/Source/AudioInsightsEditor/Private/Analyzers/SubmixAudioAnalyzerRack.h
@@ -6,6 +6,7 @@
 
 namespace AudioWidgets { class FAudioAnalyzerRack; }
 
+class FSpawnTabArgs;
 class SDockTab;
 class SWidget;
 class USoundSubmix;
@@ -18,7 +19,7 @@ namespace UE::Audio::Insights
 		FSubmixAudioAnalyzerRack(TWeakObjectPtr<USoundSubmix> InSoundSubmix);
 		virtual ~FSubmixAudioAnalyzerRack();
 
-		TSharedRef<SWidget> MakeWidget(TSharedRef<SDockTab> InOwnerTab);
+		TSharedRef<SWidget> MakeWidget(TSharedRef<SDockTab> InOwnerTab, const FSpawnTabArgs& InSpawnTabArgs);
 
 		void RebuildAudioAnalyzerRack(TWeakObjectPtr<USoundSubmix> InSoundSubmix);
 
diff --git a/Engine/Plugins/AudioInsights/Source/AudioInsightsEditor/Private/AudioInsightsEditorDashboardFactory.cpp b/Engine/Plugins/AudioInsights/Source/AudioInsightsEditor/Private/AudioInsightsEditorDashboardFactory.cpp
--- a/Engine/Plugins/AudioInsights/Source/AudioInsightsEditor/Private/AudioInsightsEditorDashboardFactory.cpp
+++ b/Engine/Plugins/AudioInsights/Source/AudioInsightsEditor/Private/AudioInsightsEditorDashboardFactory.cpp
@@ -247,7 +247,7 @@ namespace UE::Audio::Insights
 			]
 			+ SVerticalBox::Slot()
 			[
-				DashboardTabManager->RestoreFrom(TabLayout, TSharedPtr<SWindow>()).ToSharedRef()
+				DashboardTabManager->RestoreFrom(TabLayout, Args.GetOwnerWindow()).ToSharedRef()
 			];
 
 		DockTab->SetContent(TabContent);
@@ -694,7 +694,7 @@ namespace UE::Audio::Insights
 					.Clipping(EWidgetClipping::ClipToBounds)
 					.Label(Factory->GetDisplayName());
 
-				TSharedRef<SWidget> DashboardView = Factory->MakeWidget(DockTab);
+				TSharedRef<SWidget> DashboardView = Factory->MakeWidget(DockTab, Args);
 				DockTab->SetContent(DashboardView);
 
 				return DockTab;
diff --git a/Engine/Plugins/AudioInsights/Source/AudioInsightsEditor/Private/Views/AudioAnalyzerRackDashboardViewFactory.cpp b/Engine/Plugins/AudioInsights/Source/AudioInsightsEditor/Private/Views/AudioAnalyzerRackDashboardViewFactory.cpp
--- a/Engine/Plugins/AudioInsights/Source/AudioInsightsEditor/Private/Views/AudioAnalyzerRackDashboardViewFactory.cpp
+++ b/Engine/Plugins/AudioInsights/Source/AudioInsightsEditor/Private/Views/AudioAnalyzerRackDashboardViewFactory.cpp
@@ -35,7 +35,7 @@ namespace UE::Audio::Insights
 		return FSlateStyle::Get().CreateIcon("AudioInsights.Icon");
 	}
 
-	TSharedRef<SWidget> FAudioAnalyzerRackDashboardViewFactory::MakeWidget(TSharedRef<SDockTab> OwnerTab)
+	TSharedRef<SWidget> FAudioAnalyzerRackDashboardViewFactory::MakeWidget(TSharedRef<SDockTab> OwnerTab, const FSpawnTabArgs& SpawnTabArgs)
 	{
 		// Create SubmixAudioAnalyzerRack (use MainSubmix as default)
 		if (TWeakObjectPtr<UAudioSettings> AudioSettings = GetMutableDefault<UAudioSettings>();
@@ -58,7 +58,7 @@ namespace UE::Audio::Insights
 				.Color(FSlateStyle::Get().GetColor("AudioInsights.Analyzers.BackgroundColor"));
 		}
 
-		return SubmixAudioAnalyzerRack->MakeWidget(OwnerTab);
+		return SubmixAudioAnalyzerRack->MakeWidget(OwnerTab, SpawnTabArgs);
 	}
 
 	void FAudioAnalyzerRackDashboardViewFactory::HandleOnActiveAudioDeviceChanged()
diff --git a/Engine/Plugins/AudioInsights/Source/AudioInsightsEditor/Private/Views/AudioAnalyzerRackDashboardViewFactory.h b/Engine/Plugins/AudioInsights/Source/AudioInsightsEditor/Private/Views/AudioAnalyzerRackDashboardViewFactory.h
--- a/Engine/Plugins/AudioInsights/Source/AudioInsightsEditor/Private/Views/AudioAnalyzerRackDashboardViewFactory.h
+++ b/Engine/Plugins/AudioInsights/Source/AudioInsightsEditor/Private/Views/AudioAnalyzerRackDashboardViewFactory.h
@@ -17,7 +17,7 @@ namespace UE::Audio::Insights
 		virtual FText GetDisplayName() const override;
 		virtual EDefaultDashboardTabStack GetDefaultTabStack() const override;
 		virtual FSlateIcon GetIcon() const override;
-		virtual TSharedRef<SWidget> MakeWidget(TSharedRef<SDockTab> OwnerTab) override;
+		virtual TSharedRef<SWidget> MakeWidget(TSharedRef<SDockTab> OwnerTab, const FSpawnTabArgs& SpawnTabArgs) override;
 
 	private:
 		void HandleOnActiveAudioDeviceChanged();
diff --git a/Engine/Plugins/AudioInsights/Source/AudioInsightsEditor/Private/Views/AudioBusesDashboardViewFactory.cpp b/Engine/Plugins/AudioInsights/Source/AudioInsightsEditor/Private/Views/AudioBusesDashboardViewFactory.cpp
--- a/Engine/Plugins/AudioInsights/Source/AudioInsightsEditor/Private/Views/AudioBusesDashboardViewFactory.cpp
+++ b/Engine/Plugins/AudioInsights/Source/AudioInsightsEditor/Private/Views/AudioBusesDashboardViewFactory.cpp
@@ -135,11 +135,11 @@ namespace UE::Audio::Insights
 		return EDefaultDashboardTabStack::Analysis;
 	}
 
-	TSharedRef<SWidget> FAudioBusesDashboardViewFactory::MakeWidget(TSharedRef<SDockTab> OwnerTab)
+	TSharedRef<SWidget> FAudioBusesDashboardViewFactory::MakeWidget(TSharedRef<SDockTab> OwnerTab, const FSpawnTabArgs& SpawnTabArgs)
 	{
 		if (!DashboardWidget.IsValid())
 		{
-			DashboardWidget = FTraceTableDashboardViewFactory::MakeWidget(OwnerTab);
+			DashboardWidget = FTraceTableDashboardViewFactory::MakeWidget(OwnerTab, SpawnTabArgs);
 
 			if (FilteredEntriesListView.IsValid())
 			{
diff --git a/Engine/Plugins/AudioInsights/Source/AudioInsightsEditor/Private/Views/AudioBusesDashboardViewFactory.h b/Engine/Plugins/AudioInsights/Source/AudioInsightsEditor/Private/Views/AudioBusesDashboardViewFactory.h
--- a/Engine/Plugins/AudioInsights/Source/AudioInsightsEditor/Private/Views/AudioBusesDashboardViewFactory.h
+++ b/Engine/Plugins/AudioInsights/Source/AudioInsightsEditor/Private/Views/AudioBusesDashboardViewFactory.h
@@ -22,7 +22,7 @@ namespace UE::Audio::Insights
 		virtual FText GetDisplayName() const override;
 		virtual FSlateIcon GetIcon() const override;
 		virtual EDefaultDashboardTabStack GetDefaultTabStack() const override;
-		virtual TSharedRef<SWidget> MakeWidget(TSharedRef<SDockTab> OwnerTab) override;
+		virtual TSharedRef<SWidget> MakeWidget(TSharedRef<SDockTab> OwnerTab, const FSpawnTabArgs& SpawnTabArgs) override;
 
 		DECLARE_MULTICAST_DELEGATE_TwoParams(FOnAudioBusAssetChecked, const bool /*bIsChecked*/, const TWeakObjectPtr<UAudioBus> /*AudioBus*/);
 		inline static FOnAudioBusAssetChecked OnAudioBusAssetChecked;
diff --git a/Engine/Plugins/AudioInsights/Source/AudioInsightsEditor/Private/Views/AudioMetersDashboardViewFactory.cpp b/Engine/Plugins/AudioInsights/Source/AudioInsightsEditor/Private/Views/AudioMetersDashboardViewFactory.cpp
--- a/Engine/Plugins/AudioInsights/Source/AudioInsightsEditor/Private/Views/AudioMetersDashboardViewFactory.cpp
+++ b/Engine/Plugins/AudioInsights/Source/AudioInsightsEditor/Private/Views/AudioMetersDashboardViewFactory.cpp
@@ -34,7 +34,7 @@ namespace UE::Audio::Insights
 		return FSlateStyle::Get().CreateIcon("AudioInsights.Icon.Submix");
 	}
 
-	TSharedRef<SWidget> FAudioMetersDashboardViewFactory::MakeWidget(TSharedRef<SDockTab> OwnerTab)
+	TSharedRef<SWidget> FAudioMetersDashboardViewFactory::MakeWidget(TSharedRef<SDockTab> OwnerTab, const FSpawnTabArgs& SpawnTabArgs)
 	{
 		if (!OnSubmixAssetCheckedHandle.IsValid())
 		{
diff --git a/Engine/Plugins/AudioInsights/Source/AudioInsightsEditor/Private/Views/AudioMetersDashboardViewFactory.h b/Engine/Plugins/AudioInsights/Source/AudioInsightsEditor/Private/Views/AudioMetersDashboardViewFactory.h
--- a/Engine/Plugins/AudioInsights/Source/AudioInsightsEditor/Private/Views/AudioMetersDashboardViewFactory.h
+++ b/Engine/Plugins/AudioInsights/Source/AudioInsightsEditor/Private/Views/AudioMetersDashboardViewFactory.h
@@ -23,7 +23,7 @@ namespace UE::Audio::Insights
 		virtual FText GetDisplayName() const override;
 		virtual EDefaultDashboardTabStack GetDefaultTabStack() const override;
 		virtual FSlateIcon GetIcon() const override;
-		virtual TSharedRef<SWidget> MakeWidget(TSharedRef<SDockTab> OwnerTab) override;
+		virtual TSharedRef<SWidget> MakeWidget(TSharedRef<SDockTab> OwnerTab, const FSpawnTabArgs& SpawnTabArgs) override;
 
 	private:
 		void HandleOnSubmixAssetChecked(const bool bInIsChecked, const TWeakObjectPtr<USoundSubmix> InSoundSubmix);
diff --git a/Engine/Plugins/AudioInsights/Source/AudioInsightsEditor/Private/Views/SubmixesDashboardViewFactory.cpp b/Engine/Plugins/AudioInsights/Source/AudioInsightsEditor/Private/Views/SubmixesDashboardViewFactory.cpp
--- a/Engine/Plugins/AudioInsights/Source/AudioInsightsEditor/Private/Views/SubmixesDashboardViewFactory.cpp
+++ b/Engine/Plugins/AudioInsights/Source/AudioInsightsEditor/Private/Views/SubmixesDashboardViewFactory.cpp
@@ -135,11 +135,11 @@ namespace UE::Audio::Insights
 		return EDefaultDashboardTabStack::Analysis;
 	}
 
-	TSharedRef<SWidget> FSubmixesDashboardViewFactory::MakeWidget(TSharedRef<SDockTab> OwnerTab)
+	TSharedRef<SWidget> FSubmixesDashboardViewFactory::MakeWidget(TSharedRef<SDockTab> OwnerTab, const FSpawnTabArgs& SpawnTabArgs)
 	{
 		if (!DashboardWidget.IsValid())
 		{
-			DashboardWidget = FTraceTableDashboardViewFactory::MakeWidget(OwnerTab);
+			DashboardWidget = FTraceTableDashboardViewFactory::MakeWidget(OwnerTab, SpawnTabArgs);
 
 			if (FilteredEntriesListView.IsValid())
 			{
diff --git a/Engine/Plugins/AudioInsights/Source/AudioInsightsEditor/Private/Views/SubmixesDashboardViewFactory.h b/Engine/Plugins/AudioInsights/Source/AudioInsightsEditor/Private/Views/SubmixesDashboardViewFactory.h
--- a/Engine/Plugins/AudioInsights/Source/AudioInsightsEditor/Private/Views/SubmixesDashboardViewFactory.h
+++ b/Engine/Plugins/AudioInsights/Source/AudioInsightsEditor/Private/Views/SubmixesDashboardViewFactory.h
@@ -21,7 +21,7 @@ namespace UE::Audio::Insights
 		virtual FText GetDisplayName() const override;
 		virtual FSlateIcon GetIcon() const override;
 		virtual EDefaultDashboardTabStack GetDefaultTabStack() const override;
-		virtual TSharedRef<SWidget> MakeWidget(TSharedRef<SDockTab> OwnerTab) override;
+		virtual TSharedRef<SWidget> MakeWidget(TSharedRef<SDockTab> OwnerTab, const FSpawnTabArgs& SpawnTabArgs) override;
 
 		DECLARE_MULTICAST_DELEGATE_TwoParams(FOnSubmixAssetChecked, const bool /*bIsChecked*/, TWeakObjectPtr<USoundSubmix> /*SoundSubmix*/);
 		inline static FOnSubmixAssetChecked OnSubmixAssetChecked;
diff --git a/Engine/Plugins/AudioInsights/Source/AudioInsightsEditor/Private/Views/ViewportDashboardViewFactory.h b/Engine/Plugins/AudioInsights/Source/AudioInsightsEditor/Private/Views/ViewportDashboardViewFactory.h
--- a/Engine/Plugins/AudioInsights/Source/AudioInsightsEditor/Private/Views/ViewportDashboardViewFactory.h
+++ b/Engine/Plugins/AudioInsights/Source/AudioInsightsEditor/Private/Views/ViewportDashboardViewFactory.h
@@ -35,7 +35,7 @@ namespace UE::Audio::Insights
 			return FSlateStyle::Get().CreateIcon("AudioInsights.Icon.Viewport");
 		}
 
-		virtual TSharedRef<SWidget> MakeWidget(TSharedRef<SDockTab> OwnerTab) override
+		virtual TSharedRef<SWidget> MakeWidget(TSharedRef<SDockTab> OwnerTab, const FSpawnTabArgs& SpawnTabArgs) override
 		{
 			// TODO: This test currently hijacks the main level viewport which is not ideal.
 			// Find way to make our own instance and bind to tab here.
diff --git a/Engine/Plugins/Runtime/AudioModulation/Source/AudioModulationEditor/Private/Insights/Views/ControlBusDashboardViewFactory.cpp b/Engine/Plugins/Runtime/AudioModulation/Source/AudioModulationEditor/Private/Insights/Views/ControlBusDashboardViewFactory.cpp
--- a/Engine/Plugins/Runtime/AudioModulation/Source/AudioModulationEditor/Private/Insights/Views/ControlBusDashboardViewFactory.cpp
+++ b/Engine/Plugins/Runtime/AudioModulation/Source/AudioModulationEditor/Private/Insights/Views/ControlBusDashboardViewFactory.cpp
@@ -83,9 +83,9 @@ namespace AudioModulationEditor
 		return UE::Audio::Insights::EDefaultDashboardTabStack::Analysis;
 	}
 
-	TSharedRef<SWidget> FControlBusDashboardViewFactory::MakeWidget(TSharedRef<SDockTab> OwnerTab)
+	TSharedRef<SWidget> FControlBusDashboardViewFactory::MakeWidget(TSharedRef<SDockTab> OwnerTab, const FSpawnTabArgs& SpawnTabArgs)
 	{
-		TSharedRef<SWidget> TableDashboardWidget = UE::Audio::Insights::FTraceTableDashboardViewFactory::MakeWidget(OwnerTab);
+		TSharedRef<SWidget> TableDashboardWidget = UE::Audio::Insights::FTraceTableDashboardViewFactory::MakeWidget(OwnerTab, SpawnTabArgs);
 		TSharedRef<SWidget> BusWatchWidget = MakeControlBusWatchWidget();
 
 		return SNew(SHorizontalBox)
diff --git a/Engine/Plugins/Runtime/AudioModulation/Source/AudioModulationEditor/Private/Insights/Views/ControlBusDashboardViewFactory.h b/Engine/Plugins/Runtime/AudioModulation/Source/AudioModulationEditor/Private/Insights/Views/ControlBusDashboardViewFactory.h
--- a/Engine/Plugins/Runtime/AudioModulation/Source/AudioModulationEditor/Private/Insights/Views/ControlBusDashboardViewFactory.h
+++ b/Engine/Plugins/Runtime/AudioModulation/Source/AudioModulationEditor/Private/Insights/Views/ControlBusDashboardViewFactory.h
@@ -19,7 +19,7 @@ namespace AudioModulationEditor
 		virtual FSlateIcon GetIcon() const override;
 		virtual UE::Audio::Insights::EDefaultDashboardTabStack GetDefaultTabStack() const override;
 		virtual void ProcessEntries(UE::Audio::Insights::FTraceTableDashboardViewFactory::EProcessReason Reason) override;
-		virtual TSharedRef<SWidget> MakeWidget(TSharedRef<SDockTab> OwnerTab) override;
+		virtual TSharedRef<SWidget> MakeWidget(TSharedRef<SDockTab> OwnerTab, const FSpawnTabArgs& SpawnTabArgs) override;
 
 	protected:
 		virtual const TMap<FName, UE::Audio::Insights::FTraceTableDashboardViewFactory::FColumnData>& GetColumns() const override;
diff --git a/Engine/Plugins/Runtime/AudioModulation/Source/AudioModulationEditor/Private/Insights/Views/ModulationMatrixDashboardViewFactory.cpp b/Engine/Plugins/Runtime/AudioModulation/Source/AudioModulationEditor/Private/Insights/Views/ModulationMatrixDashboardViewFactory.cpp
--- a/Engine/Plugins/Runtime/AudioModulation/Source/AudioModulationEditor/Private/Insights/Views/ModulationMatrixDashboardViewFactory.cpp
+++ b/Engine/Plugins/Runtime/AudioModulation/Source/AudioModulationEditor/Private/Insights/Views/ModulationMatrixDashboardViewFactory.cpp
@@ -220,7 +220,7 @@ namespace AudioModulationEditor
 			];
 	}
 
-	TSharedRef<SWidget> FModulationMatrixDashboardViewFactory::MakeWidget(TSharedRef<SDockTab> OwnerTab)
+	TSharedRef<SWidget> FModulationMatrixDashboardViewFactory::MakeWidget(TSharedRef<SDockTab> OwnerTab, const FSpawnTabArgs& SpawnTabArgs)
 	{
 		using namespace UE::Audio::Insights;
 
@@ -237,7 +237,7 @@ namespace AudioModulationEditor
 		+ SVerticalBox::Slot()
 		.HAlign(HAlign_Fill)
 		[
-			FTraceTableDashboardViewFactory::MakeWidget(OwnerTab)
+			FTraceTableDashboardViewFactory::MakeWidget(OwnerTab, SpawnTabArgs)
 		];
 	}
 
diff --git a/Engine/Plugins/Runtime/AudioModulation/Source/AudioModulationEditor/Private/Insights/Views/ModulationMatrixDashboardViewFactory.h b/Engine/Plugins/Runtime/AudioModulation/Source/AudioModulationEditor/Private/Insights/Views/ModulationMatrixDashboardViewFactory.h
--- a/Engine/Plugins/Runtime/AudioModulation/Source/AudioModulationEditor/Private/Insights/Views/ModulationMatrixDashboardViewFactory.h
+++ b/Engine/Plugins/Runtime/AudioModulation/Source/AudioModulationEditor/Private/Insights/Views/ModulationMatrixDashboardViewFactory.h
@@ -24,7 +24,7 @@ namespace AudioModulationEditor
 		virtual FSlateIcon GetIcon() const override;
 		virtual UE::Audio::Insights::EDefaultDashboardTabStack GetDefaultTabStack() const override;
 
-		virtual TSharedRef<SWidget> MakeWidget(TSharedRef<SDockTab> OwnerTab) override;
+		virtual TSharedRef<SWidget> MakeWidget(TSharedRef<SDockTab> OwnerTab, const FSpawnTabArgs& SpawnTabArgs) override;
 
 		virtual void ProcessEntries(UE::Audio::Insights::FTraceTableDashboardViewFactory::EProcessReason Reason) override;
 
diff --git a/Engine/Plugins/Runtime/AudioWidgets/Source/AudioWidgets/Private/AudioAnalyzerRack.cpp b/Engine/Plugins/Runtime/AudioWidgets/Source/AudioWidgets/Private/AudioAnalyzerRack.cpp
--- a/Engine/Plugins/Runtime/AudioWidgets/Source/AudioWidgets/Private/AudioAnalyzerRack.cpp
+++ b/Engine/Plugins/Runtime/AudioWidgets/Source/AudioWidgets/Private/AudioAnalyzerRack.cpp
@@ -81,7 +81,7 @@ namespace AudioWidgets
 		RackUnits.Reset();
 	}
 
-	TSharedRef<SWidget> FAudioAnalyzerRack::CreateWidget(TSharedRef<SDockTab> DockTab)
+	TSharedRef<SWidget> FAudioAnalyzerRack::CreateWidget(TSharedRef<SDockTab> DockTab, const FSpawnTabArgs& SpawnTabArgs)
 	{
 		if (TabManager.IsValid())
 		{
@@ -164,7 +164,7 @@ namespace AudioWidgets
 					]
 					+ SOverlay::Slot()
 					[
-						TabManager->RestoreFrom(TabLayout, DockTab->GetParentWindow()).ToSharedRef()
+						TabManager->RestoreFrom(TabLayout, SpawnTabArgs.GetOwnerWindow()).ToSharedRef()
 					]
 			];
 
diff --git a/Engine/Plugins/Runtime/AudioWidgets/Source/AudioWidgets/Public/AudioAnalyzerRack.h b/Engine/Plugins/Runtime/AudioWidgets/Source/AudioWidgets/Public/AudioAnalyzerRack.h
--- a/Engine/Plugins/Runtime/AudioWidgets/Source/AudioWidgets/Public/AudioAnalyzerRack.h
+++ b/Engine/Plugins/Runtime/AudioWidgets/Source/AudioWidgets/Public/AudioAnalyzerRack.h
@@ -87,7 +87,7 @@ namespace AudioWidgets
 		void Init(int32 InNumChannels, Audio::FDeviceId InAudioDeviceId);
 		void DestroyAnalyzers();
 
-		TSharedRef<SWidget> CreateWidget(TSharedRef<SDockTab> DockTab);
+		TSharedRef<SWidget> CreateWidget(TSharedRef<SDockTab> DockTab, const FSpawnTabArgs& SpawnTabArgs);
 
 		UAudioBus* GetAudioBus() const;
 
diff --git a/Engine/Plugins/Runtime/Metasound/Source/MetasoundEditor/Private/MetasoundEditor.cpp b/Engine/Plugins/Runtime/Metasound/Source/MetasoundEditor/Private/MetasoundEditor.cpp
--- a/Engine/Plugins/Runtime/Metasound/Source/MetasoundEditor/Private/MetasoundEditor.cpp
+++ b/Engine/Plugins/Runtime/Metasound/Source/MetasoundEditor/Private/MetasoundEditor.cpp
@@ -981,12 +981,12 @@ namespace Metasound
 			.SetGroup(WorkspaceMenuCategoryRef)
 			.SetIcon(FSlateIcon("MetaSoundStyle", "MetasoundEditor.Metasound.Icon"));
 
-			InTabManager->RegisterTabSpawner(TabNamesPrivate::Analyzers, FOnSpawnTab::CreateLambda([InAnalyzerRack = AnalyzerRack.ToWeakPtr()](const FSpawnTabArgs&)
+			InTabManager->RegisterTabSpawner(TabNamesPrivate::Analyzers, FOnSpawnTab::CreateLambda([InAnalyzerRack = AnalyzerRack.ToWeakPtr()](const FSpawnTabArgs& Args)
 			{
 				TSharedRef<SDockTab> DockTab = SNew(SDockTab).Label(LOCTEXT("MetasoundAnalyzersTitle", "Analyzers"));
 				if (TSharedPtr<AudioWidgets::FAudioAnalyzerRack> PinnedAnalyzerRack = InAnalyzerRack.Pin()) // AnalyzerRack only created for MetaSoundSource, not MetaSoundPatch.
 				{
-					TSharedRef<SWidget> AnalyzerRackWidget = PinnedAnalyzerRack->CreateWidget(DockTab);
+					TSharedRef<SWidget> AnalyzerRackWidget = PinnedAnalyzerRack->CreateWidget(DockTab, Args);
 					DockTab->SetContent(AnalyzerRackWidget);
 				}
 				return DockTab;
