[Audio Insights] FTraceTableDashboardViewFactory: adapted the code so we refresh the items list instead of doing a full rebuild list (expensive!)

diff --git a/Engine/Plugins/AudioInsights/Source/AudioInsights/Private/Views/TableDashboardViewFactory.cpp b/Engine/Plugins/AudioInsights/Source/AudioInsights/Private/Views/TableDashboardViewFactory.cpp
--- a/Engine/Plugins/AudioInsights/Source/AudioInsights/Private/Views/TableDashboardViewFactory.cpp
+++ b/Engine/Plugins/AudioInsights/Source/AudioInsights/Private/Views/TableDashboardViewFactory.cpp
@@ -35,10 +35,10 @@ namespace UE::Audio::Insights
 
 	TSharedRef<SWidget> FTraceTableDashboardViewFactory::GenerateWidgetForColumn(TSharedRef<IDashboardDataViewEntry> InRowData, const FName& Column)
 	{
-		const IDashboardDataViewEntry& RowData = InRowData.Get();
 		const FColumnData& ColumnData = GetColumns()[Column];
-		const FText Label = ColumnData.GetDisplayValue(RowData);
-		if (Label.IsEmpty())
+		const FText ValueText = ColumnData.GetDisplayValue(InRowData.Get());
+
+		if (ValueText.IsEmpty())
 		{
 			return SNullWidget::NullWidget;
 		}
@@ -46,7 +46,14 @@ namespace UE::Audio::Insights
 		return SNew(SHorizontalBox)
 		+ SHorizontalBox::Slot().Padding(2)
 		[
-			SNew(STextBlock).Text(Label)
+			SNew(STextBlock)
+			.Text_Lambda([this, InRowData, Column]()
+			{
+				const FColumnData& ColumnData = GetColumns()[Column];
+				const FText ValueText = ColumnData.GetDisplayValue(InRowData.Get());
+
+				return ValueText;
+			})
 		];
 	}
 
@@ -192,11 +199,11 @@ namespace UE::Audio::Insights
 	}
 
 
-	void FTraceTableDashboardViewFactory::RebuildFilteredEntriesListView()
+	void FTraceTableDashboardViewFactory::RefreshFilteredEntriesListView()
 	{
 		if (FilteredEntriesListView.IsValid())
 		{
-			FilteredEntriesListView->RebuildList();
+			FilteredEntriesListView->RequestListRefresh();
 		}
 	}
 
@@ -237,7 +244,7 @@ namespace UE::Audio::Insights
 				}
 			}
 
-			RebuildFilteredEntriesListView();
+			RefreshFilteredEntriesListView();
 			UpdateFilterReason = EProcessReason::None;
 		}
 
@@ -296,34 +303,41 @@ namespace UE::Audio::Insights
 
 	TSharedRef<SWidget> FTraceObjectTableDashboardViewFactory::GenerateWidgetForColumn(TSharedRef<IDashboardDataViewEntry> InRowData, const FName& Column)
 	{
-		const IDashboardDataViewEntry& RowData = InRowData.Get();
 		const FColumnData& ColumnData = GetColumns()[Column];
-		const FText Label = ColumnData.GetDisplayValue(RowData);
-		if (Label.IsEmpty())
+		const FText ValueText = ColumnData.GetDisplayValue(InRowData.Get());
+
+		if (ValueText.IsEmpty())
 		{
 			return SNullWidget::NullWidget;
 		}
 
-		TSharedPtr<IDashboardDataViewEntry> RowDataPtr = InRowData;
 		return SNew(SHorizontalBox)
 		+ SHorizontalBox::Slot().Padding(2)
 		[
-			SNew(STextBlock).Text(Label)
-			.ColorAndOpacity(GetRowColor(RowDataPtr))
-			.OnDoubleClicked_Lambda([this, RowDataPtr](const FGeometry& MyGeometry, const FPointerEvent& PointerEvent)
+			SNew(STextBlock)
+			.Text_Lambda([this, InRowData, Column]()
+			{
+				const FColumnData& ColumnData = GetColumns()[Column];
+				const FText ValueText = ColumnData.GetDisplayValue(InRowData.Get());
+
+				return ValueText;
+			})
+			.ColorAndOpacity_Lambda([this, InRowData]()
+			{
+				return GetRowColor(InRowData.ToSharedPtr());
+			})
+			.OnDoubleClicked_Lambda([this, InRowData](const FGeometry& MyGeometry, const FPointerEvent& PointerEvent)
 			{
 #if WITH_EDITOR
 				if (GEditor)
 				{
-					TSharedPtr<IObjectDashboardEntry> ObjectData = StaticCastSharedPtr<IObjectDashboardEntry>(RowDataPtr);
-					if (ObjectData.IsValid())
+					const TSharedRef<IObjectDashboardEntry> ObjectData = StaticCastSharedRef<IObjectDashboardEntry>(InRowData);
+					const TObjectPtr<UObject> Object = ObjectData->GetObject();
+
+					if (Object && Object->IsAsset())
 					{
-						const UObject* Object = ObjectData->GetObject();
-						if (Object && Object->IsAsset())
-						{
-							GEditor->GetEditorSubsystem<UAssetEditorSubsystem>()->OpenEditorForAsset(Object);
-							return FReply::Handled();
-						}
+						GEditor->GetEditorSubsystem<UAssetEditorSubsystem>()->OpenEditorForAsset(Object);
+						return FReply::Handled();
 					}
 				}
 #endif // WITH_EDITOR
diff --git a/Engine/Plugins/AudioInsights/Source/AudioInsights/Public/Views/TableDashboardViewFactory.h b/Engine/Plugins/AudioInsights/Source/AudioInsights/Public/Views/TableDashboardViewFactory.h
--- a/Engine/Plugins/AudioInsights/Source/AudioInsights/Public/Views/TableDashboardViewFactory.h
+++ b/Engine/Plugins/AudioInsights/Source/AudioInsights/Public/Views/TableDashboardViewFactory.h
@@ -61,7 +61,7 @@ namespace UE::Audio::Insights
 		virtual TSharedRef<SWidget> GenerateWidgetForColumn(TSharedRef<IDashboardDataViewEntry> InRowData, const FName& Column);
 		virtual FReply OnDataRowKeyInput(const FGeometry& Geometry, const FKeyEvent& KeyEvent) const;
 		virtual TSharedRef<SWidget> MakeWidget(TSharedRef<SDockTab> OwnerTab, const FSpawnTabArgs& SpawnTabArgs) override;
-		virtual void RebuildFilteredEntriesListView();
+		virtual void RefreshFilteredEntriesListView();
 
 		enum class EProcessReason : uint8
 		{
