[Audio Insights] Submixes Dashboard: refactored submixes dashboard so it is 100% trace based instead of asset events based + added activity indicator column

diff --git a/Engine/Plugins/AudioInsights/Source/AudioInsightsEditor/AudioInsightsEditor.Build.cs b/Engine/Plugins/AudioInsights/Source/AudioInsightsEditor/AudioInsightsEditor.Build.cs
--- a/Engine/Plugins/AudioInsights/Source/AudioInsightsEditor/AudioInsightsEditor.Build.cs
+++ b/Engine/Plugins/AudioInsights/Source/AudioInsightsEditor/AudioInsightsEditor.Build.cs
@@ -33,6 +33,7 @@ public class AudioInsightsEditor : ModuleRules
 				"Slate",
 				"SlateCore",
 				"ToolMenus",
+				"TraceAnalysis",
 				"UnrealEd",
 				"WorkspaceMenuStructure",
 			}
diff --git a/Engine/Plugins/AudioInsights/Source/AudioInsightsEditor/Private/Messages/SoundSubmixMessages.h b/Engine/Plugins/AudioInsights/Source/AudioInsightsEditor/Private/Messages/SoundSubmixMessages.h
new file mode 100644
--- /dev/null
+++ b/Engine/Plugins/AudioInsights/Source/AudioInsightsEditor/Private/Messages/SoundSubmixMessages.h
@@ -0,0 +1,79 @@
+// Copyright Epic Games, Inc. All Rights Reserved.
+#pragma once
+
+#include "AudioDefines.h"
+#include "Math/NumericLimits.h"
+#include "Messages/AnalyzerMessageQueue.h"
+#include "Trace/Analyzer.h"
+#include "Views/TableDashboardViewFactory.h"
+
+namespace UE::Audio::Insights
+{
+	struct FSoundSubmixMessageBase
+	{
+		FSoundSubmixMessageBase() = default;
+
+		FSoundSubmixMessageBase(const UE::Trace::IAnalyzer::FOnEventContext& InContext)
+		{
+			const UE::Trace::IAnalyzer::FEventData& EventData = InContext.EventData;
+
+			DeviceId = static_cast<::Audio::FDeviceId>(EventData.GetValue<uint32>("DeviceId"));
+			SubmixId = EventData.GetValue<uint32>("SubmixId");
+			Timestamp = InContext.EventTime.AsSeconds(EventData.GetValue<uint64>("Timestamp"));
+		}
+
+		::Audio::FDeviceId DeviceId = INDEX_NONE;
+		uint32 SubmixId = INDEX_NONE;
+		double Timestamp = 0.0;
+	};
+
+	struct FSoundSubmixActivateMessage : public FSoundSubmixMessageBase
+	{
+		FSoundSubmixActivateMessage() = default;
+
+		FSoundSubmixActivateMessage(const UE::Trace::IAnalyzer::FOnEventContext& InContext)
+			: FSoundSubmixMessageBase(InContext)
+		{
+			const UE::Trace::IAnalyzer::FEventData& EventData = InContext.EventData;
+
+			EventData.GetString("Name", Name);
+		}
+
+		FString Name;
+	};
+
+	using FSoundSubmixDeactivateMessage = FSoundSubmixMessageBase;
+
+	struct FSoundSubmixHasActivityMessage : public FSoundSubmixMessageBase
+	{
+		FSoundSubmixHasActivityMessage() = default;
+
+		FSoundSubmixHasActivityMessage(const UE::Trace::IAnalyzer::FOnEventContext& InContext)
+			: FSoundSubmixMessageBase(InContext)
+		{
+			const UE::Trace::IAnalyzer::FEventData& EventData = InContext.EventData;
+
+			bHasActivity = EventData.GetValue<bool>("HasActivity");
+		}
+
+		bool bHasActivity = false;
+	};
+
+	class FSoundSubmixMessages
+	{
+		UE::Audio::Insights::TAnalyzerMessageQueue<FSoundSubmixActivateMessage> ActivateMessages{ 0.1 };
+		UE::Audio::Insights::TAnalyzerMessageQueue<FSoundSubmixDeactivateMessage> DeactivateMessages{ 0.1 };
+		UE::Audio::Insights::TAnalyzerMessageQueue<FSoundSubmixHasActivityMessage> HasActivityMessages{ 0.1 };
+
+		friend class FSoundSubmixProvider;
+	};
+
+	struct FSoundSubmixAssetDashboardEntry : public UE::Audio::Insights::FSoundAssetDashboardEntry
+	{
+		FSoundSubmixAssetDashboardEntry() = default;
+		virtual ~FSoundSubmixAssetDashboardEntry() = default;
+
+		uint32 SubmixId = INDEX_NONE;
+		bool bHasActivity = false;
+	};
+} // namespace UE::Audio::Insights
diff --git a/Engine/Plugins/AudioInsights/Source/AudioInsightsEditor/Private/Providers/SoundSubmixProvider.cpp b/Engine/Plugins/AudioInsights/Source/AudioInsightsEditor/Private/Providers/SoundSubmixProvider.cpp
--- a/Engine/Plugins/AudioInsights/Source/AudioInsightsEditor/Private/Providers/SoundSubmixProvider.cpp
+++ b/Engine/Plugins/AudioInsights/Source/AudioInsightsEditor/Private/Providers/SoundSubmixProvider.cpp
@@ -13,25 +13,6 @@ namespace UE::Audio::Insights
 	FSoundSubmixProvider::FSoundSubmixProvider()
 		: TDeviceDataMapTraceProvider<uint32, TSharedPtr<FSoundSubmixAssetDashboardEntry>>(GetName_Static())
 	{
-		FAssetRegistryModule& AssetRegistryModule = FModuleManager::LoadModuleChecked<FAssetRegistryModule>("AssetRegistry");
-
-		AssetRegistryModule.Get().OnAssetAdded().AddRaw(this, &FSoundSubmixProvider::OnAssetAdded);
-		AssetRegistryModule.Get().OnAssetRemoved().AddRaw(this, &FSoundSubmixProvider::OnAssetRemoved);
-		AssetRegistryModule.Get().OnFilesLoaded().AddRaw(this, &FSoundSubmixProvider::OnFilesLoaded);
-
-		FEditorDashboardFactory::OnActiveAudioDeviceChanged.AddRaw(this, &FSoundSubmixProvider::OnActiveAudioDeviceChanged);
-	}
-	
-	FSoundSubmixProvider::~FSoundSubmixProvider()
-	{
-		if (FAssetRegistryModule* AssetRegistryModule = FModuleManager::GetModulePtr<FAssetRegistryModule>("AssetRegistry"))
-		{
-			AssetRegistryModule->Get().OnAssetAdded().RemoveAll(this);
-			AssetRegistryModule->Get().OnAssetRemoved().RemoveAll(this);
-			AssetRegistryModule->Get().OnFilesLoaded().RemoveAll(this);
-		}
-
-		FEditorDashboardFactory::OnActiveAudioDeviceChanged.RemoveAll(this);
 	}
 
 	FName FSoundSubmixProvider::GetName_Static()
@@ -39,129 +20,146 @@ namespace UE::Audio::Insights
 		return "SubmixesProvider";
 	}
 
-	void FSoundSubmixProvider::OnAssetAdded(const FAssetData& InAssetData)
-	{
-		if (bAreFilesLoaded && InAssetData.AssetClassPath == FTopLevelAssetPath(USoundSubmix::StaticClass()))
-		{
-			AddSubmixAsset(InAssetData);
-		}
-	}
-
-	void FSoundSubmixProvider::OnAssetRemoved(const FAssetData& InAssetData)
+	bool FSoundSubmixProvider::ProcessMessages()
 	{
-		if (InAssetData.AssetClassPath == FTopLevelAssetPath(USoundSubmix::StaticClass()))
+		// Helper lambdas
+		auto BumpEntryFunc = [this](const FSoundSubmixMessageBase& Msg)
 		{
-			RemoveSubmixAsset(InAssetData);
-		}
-	}
+			TSharedPtr<FSoundSubmixAssetDashboardEntry>* ToReturn = nullptr;
 
-	void FSoundSubmixProvider::OnFilesLoaded()
-	{
-		bAreFilesLoaded = true;
-		UpdateSubmixAssetNames();
-	}
+			UpdateDeviceEntry(Msg.DeviceId, Msg.SubmixId, [&ToReturn, &Msg](TSharedPtr<FSoundSubmixAssetDashboardEntry>& Entry)
+			{
+				if (!Entry.IsValid())
+				{
+					Entry = MakeShared<FSoundSubmixAssetDashboardEntry>();
+					Entry->DeviceId = Msg.DeviceId;
+					Entry->SubmixId = Msg.SubmixId;
+				}
 
-	void FSoundSubmixProvider::OnActiveAudioDeviceChanged()
-	{
-		UpdateSubmixAssetNames();
-	}
+				Entry->Timestamp = Msg.Timestamp;
 
-	void FSoundSubmixProvider::AddSubmixAsset(const FAssetData& InAssetData)
-	{
-		const bool bIsSubmixAssetAlreadAdded = SubmixDataViewEntries.ContainsByPredicate(
-			[AssetName = InAssetData.GetObjectPathString()](const TSharedPtr<FSoundSubmixAssetDashboardEntry>& SoundSubmixAssetDashboardEntry)
-			{
-				return SoundSubmixAssetDashboardEntry->Name == AssetName;
+				*ToReturn = Entry;
 			});
 
-		if (!bIsSubmixAssetAlreadAdded)
-		{
-			const FAudioInsightsEditorModule AudioInsightsEditorModule = FAudioInsightsEditorModule::GetChecked();
-			const ::Audio::FDeviceId AudioDeviceId = AudioInsightsEditorModule.GetDeviceId();
+			return ToReturn;
+		};
 
-			TSharedPtr<FSoundSubmixAssetDashboardEntry> SoundSubmixAssetDashboardEntry = MakeShared<FSoundSubmixAssetDashboardEntry>();
-			SoundSubmixAssetDashboardEntry->DeviceId    = AudioDeviceId;
-			SoundSubmixAssetDashboardEntry->Name        = InAssetData.GetObjectPathString();
-			SoundSubmixAssetDashboardEntry->SoundSubmix = Cast<USoundSubmix>(InAssetData.GetAsset());
+		auto GetEntry = [this](const FSoundSubmixMessageBase& Msg)
+		{
+			return FindDeviceEntry(Msg.DeviceId, Msg.SubmixId);
+		};
 
-			SubmixDataViewEntries.Add(MoveTemp(SoundSubmixAssetDashboardEntry));
 
-			OnSubmixAssetAdded.Broadcast(InAssetData.GetAsset());
-			++LastUpdateId;
-		}
-	}
+		// Process messages
+		ProcessMessageQueue<FSoundSubmixActivateMessage>(TraceMessages.ActivateMessages, BumpEntryFunc,
+		[](const FSoundSubmixActivateMessage& Msg, TSharedPtr<FSoundSubmixAssetDashboardEntry>* OutEntry)
+		{
+			FSoundSubmixAssetDashboardEntry& EntryRef = *OutEntry->Get();
+			EntryRef.Name = *Msg.Name;
+		});
 
-	void FSoundSubmixProvider::RemoveSubmixAsset(const FAssetData& InAssetData)
-	{
-		const int32 FoundSubmixAssetNameIndex = SubmixDataViewEntries.IndexOfByPredicate(
-			[AssetName = InAssetData.GetObjectPathString()](const TSharedPtr<FSoundSubmixAssetDashboardEntry>& SoundSubmixAssetDashboardEntry)
+		ProcessMessageQueue<FSoundSubmixDeactivateMessage>(TraceMessages.DeactivateMessages, GetEntry,
+		[this](const FSoundSubmixDeactivateMessage& Msg, TSharedPtr<FSoundSubmixAssetDashboardEntry>* OutEntry)
+		{
+			if (OutEntry && (*OutEntry)->Timestamp < Msg.Timestamp)
 			{
-				return SoundSubmixAssetDashboardEntry->Name == AssetName;
-			});
+				FSoundSubmixAssetDashboardEntry& EntryRef = *OutEntry->Get();
 
-		if (FoundSubmixAssetNameIndex != INDEX_NONE)
-		{
-			const FAudioInsightsEditorModule AudioInsightsEditorModule = FAudioInsightsEditorModule::GetChecked();
-			const ::Audio::FDeviceId AudioDeviceId = AudioInsightsEditorModule.GetDeviceId();
+				OnSubmixDeactivated.Broadcast(EntryRef.GetObject());
 
-			RemoveDeviceEntry(AudioDeviceId, SubmixDataViewEntries[FoundSubmixAssetNameIndex]->SoundSubmix->GetUniqueID());
+				RemoveDeviceEntry(Msg.DeviceId, Msg.SubmixId);
+			}
+		});
 
-			SubmixDataViewEntries.RemoveAt(FoundSubmixAssetNameIndex);
+		ProcessMessageQueue<FSoundSubmixHasActivityMessage>(TraceMessages.HasActivityMessages, GetEntry,
+		[](const FSoundSubmixHasActivityMessage& Msg, TSharedPtr<FSoundSubmixAssetDashboardEntry>* OutEntry)
+		{
+			if (OutEntry)
+			{
+				FSoundSubmixAssetDashboardEntry& EntryRef = *OutEntry->Get();
+				EntryRef.bHasActivity = Msg.bHasActivity;
+				EntryRef.Timestamp = Msg.Timestamp;
+			}
+		});
 
-			OnSubmixAssetRemoved.Broadcast(InAssetData.GetAsset());
-			++LastUpdateId;
-		}
+		return true;
 	}
 
-	void FSoundSubmixProvider::UpdateSubmixAssetNames()
+	UE::Trace::IAnalyzer* FSoundSubmixProvider::ConstructAnalyzer(TraceServices::IAnalysisSession& InSession)
 	{
-		// Get all USoundSubmix assets
-		TArray<FAssetData> AssetDataArray;
-
-		FAssetRegistryModule& AssetRegistryModule = FModuleManager::LoadModuleChecked<FAssetRegistryModule>("AssetRegistry");
-		AssetRegistryModule.Get().GetAssetsByClass(FTopLevelAssetPath(USoundSubmix::StaticClass()), AssetDataArray);
-
-		// Build SubmixDataViewEntries
-		Reset();
-		SubmixDataViewEntries.Empty();
-
-		for (const FAssetData& AssetData : AssetDataArray)
+		class FSoundSubmixTraceAnalyzer : public FTraceAnalyzerBase
 		{
-			AddSubmixAsset(AssetData);
-		}
+		public:
+			FSoundSubmixTraceAnalyzer(TSharedRef<FSoundSubmixProvider> InProvider, TraceServices::IAnalysisSession& InSession)
+				: FTraceAnalyzerBase(InProvider)
+				, Session(InSession)
+			{
+			}
 
-		SubmixDataViewEntries.Sort([](const TSharedPtr<FSoundSubmixAssetDashboardEntry>& A, const TSharedPtr<FSoundSubmixAssetDashboardEntry>& B)
-		{
-			return A->GetDisplayName().CompareToCaseIgnored(B->GetDisplayName()) < 0;
-		});
+			virtual void OnAnalysisBegin(const FOnAnalysisContext& Context) override
+			{
+				FTraceAnalyzerBase::OnAnalysisBegin(Context);
 
-		OnSubmixAssetListUpdated.Broadcast();
-	}
+				UE::Trace::IAnalyzer::FInterfaceBuilder& Builder = Context.InterfaceBuilder;
 
-	bool FSoundSubmixProvider::ProcessMessages()
-	{
-		const FAudioInsightsEditorModule AudioInsightsEditorModule = FAudioInsightsEditorModule::GetChecked();
-		const ::Audio::FDeviceId AudioDeviceId = AudioInsightsEditorModule.GetDeviceId();
+				Builder.RouteEvent(RouteId_Activate,    "Audio", "SoundSubmixActivate");
+				Builder.RouteEvent(RouteId_Deactivate,  "Audio", "SoundSubmixDeactivate");
+				Builder.RouteEvent(RouteId_HasActivity, "Audio", "SoundSubmixHasActivity");
+			}
 
-		for (const TSharedPtr<FSoundSubmixAssetDashboardEntry>& SubmixDataViewEntry : SubmixDataViewEntries)
-		{
-			if (SubmixDataViewEntry.IsValid() && SubmixDataViewEntry->SoundSubmix.IsValid())
+			virtual bool OnEvent(uint16 RouteId, EStyle Style, const FOnEventContext& Context) override
 			{
-				UpdateDeviceEntry(AudioDeviceId, SubmixDataViewEntry->SoundSubmix->GetUniqueID(), [&SubmixDataViewEntry](TSharedPtr<FSoundSubmixAssetDashboardEntry>& Entry)
+				LLM_SCOPE_BYNAME(TEXT("Insights/FSoundSubmixTraceAnalyzer"));
+
+				FSoundSubmixMessages& Messages = GetProvider<FSoundSubmixProvider>().TraceMessages;
+
+				switch (RouteId)
 				{
-					if (!Entry.IsValid())
+					case RouteId_Activate:
 					{
-						Entry = SubmixDataViewEntry;
+						Messages.ActivateMessages.Enqueue(FSoundSubmixActivateMessage { Context });
+						break;
 					}
-				});
+
+					case RouteId_Deactivate:
+					{
+						Messages.DeactivateMessages.Enqueue(FSoundSubmixDeactivateMessage{ Context });
+						break;
+					}
+
+					case RouteId_HasActivity:
+					{
+						Messages.HasActivityMessages.Enqueue(FSoundSubmixHasActivityMessage{ Context });
+						break;
+					}
+
+					default:
+					{
+						return OnEventFailure(RouteId, Style, Context);
+					}
+				}
+
+				const double Timestamp = Context.EventTime.AsSeconds(Context.EventData.GetValue<uint64>("Timestamp"));
+
+				{
+					TraceServices::FAnalysisSessionEditScope SessionEditScope(Session);
+					Session.UpdateDurationSeconds(Timestamp);
+				}
+
+				return OnEventSuccess(RouteId, Style, Context);
 			}
-		}
 
-		return true;
-	}
+		private:
+			enum : uint16
+			{
+				RouteId_Activate,
+				RouteId_Deactivate,
+				RouteId_HasActivity
+			};
 
-	UE::Trace::IAnalyzer* FSoundSubmixProvider::ConstructAnalyzer(TraceServices::IAnalysisSession& InSession)
-	{
-		return nullptr;
+			TraceServices::IAnalysisSession& Session;
+		};
+
+		return new FSoundSubmixTraceAnalyzer(AsShared(), InSession);
 	}
 } // namespace UE::Audio::Insights
diff --git a/Engine/Plugins/AudioInsights/Source/AudioInsightsEditor/Private/Providers/SoundSubmixProvider.h b/Engine/Plugins/AudioInsights/Source/AudioInsightsEditor/Private/Providers/SoundSubmixProvider.h
--- a/Engine/Plugins/AudioInsights/Source/AudioInsightsEditor/Private/Providers/SoundSubmixProvider.h
+++ b/Engine/Plugins/AudioInsights/Source/AudioInsightsEditor/Private/Providers/SoundSubmixProvider.h
@@ -3,7 +3,7 @@
 
 #include "Delegates/Delegate.h"
 #include "UObject/NameTypes.h"
-#include "Views/SoundSubmixAssetDashboardEntry.h"
+#include "Messages/SoundSubmixMessages.h"
 
 namespace UE::Audio::Insights
 {
@@ -11,36 +11,18 @@ namespace UE::Audio::Insights
 	{
 	public:
 		FSoundSubmixProvider();
-		virtual ~FSoundSubmixProvider();
+		virtual ~FSoundSubmixProvider() = default;
 
 		static FName GetName_Static();
 
 		virtual UE::Trace::IAnalyzer* ConstructAnalyzer(TraceServices::IAnalysisSession& InSession) override;
-		
-		DECLARE_MULTICAST_DELEGATE_OneParam(FOnSubmixAssetAdded, const TWeakObjectPtr<UObject> /*Asset*/);
-		inline static FOnSubmixAssetAdded OnSubmixAssetAdded;
 
-		DECLARE_MULTICAST_DELEGATE_OneParam(FOnSubmixAssetRemoved, const TWeakObjectPtr<UObject> /*Asset*/);
-		inline static FOnSubmixAssetRemoved OnSubmixAssetRemoved;
-
-		DECLARE_MULTICAST_DELEGATE(FOnSubmixAssetListUpdated);
-		inline static FOnSubmixAssetListUpdated OnSubmixAssetListUpdated;
+		DECLARE_MULTICAST_DELEGATE_OneParam(FOnSubmixDeactivated, const TWeakObjectPtr<UObject> /*Asset*/);
+		inline static FOnSubmixDeactivated OnSubmixDeactivated;
 
 	private:
-		void OnAssetAdded(const FAssetData& InAssetData);
-		void OnAssetRemoved(const FAssetData& InAssetData);
-		void OnFilesLoaded();
-		void OnActiveAudioDeviceChanged();
-
-		void AddSubmixAsset(const FAssetData& InAssetData);
-		void RemoveSubmixAsset(const FAssetData& InAssetData);
-
-		void UpdateSubmixAssetNames();
-
 		virtual bool ProcessMessages() override;
 
-		bool bAreFilesLoaded = false;
-
-		TArray<TSharedPtr<FSoundSubmixAssetDashboardEntry>> SubmixDataViewEntries;
+		FSoundSubmixMessages TraceMessages;
 	};
 } // namespace UE::Audio::Insights
diff --git a/Engine/Plugins/AudioInsights/Source/AudioInsightsEditor/Private/Views/AudioMetersDashboardViewFactory.cpp b/Engine/Plugins/AudioInsights/Source/AudioInsightsEditor/Private/Views/AudioMetersDashboardViewFactory.cpp
--- a/Engine/Plugins/AudioInsights/Source/AudioInsightsEditor/Private/Views/AudioMetersDashboardViewFactory.cpp
+++ b/Engine/Plugins/AudioInsights/Source/AudioInsightsEditor/Private/Views/AudioMetersDashboardViewFactory.cpp
@@ -36,9 +36,9 @@ namespace UE::Audio::Insights
 
 	TSharedRef<SWidget> FAudioMetersDashboardViewFactory::MakeWidget(TSharedRef<SDockTab> OwnerTab, const FSpawnTabArgs& SpawnTabArgs)
 	{
-		if (!OnSubmixAssetCheckedHandle.IsValid())
+		if (!OnSubmixItemCheckedHandle.IsValid())
 		{
-			OnSubmixAssetCheckedHandle = FSubmixesDashboardViewFactory::OnSubmixAssetChecked.AddSP(this, &FAudioMetersDashboardViewFactory::HandleOnSubmixAssetChecked);
+			OnSubmixItemCheckedHandle = FSubmixesDashboardViewFactory::OnSubmixItemChecked.AddSP(this, &FAudioMetersDashboardViewFactory::HandleOnSubmixItemChecked);
 		}
 
 		if (!OnAudioBusAssetCheckedHandle.IsValid())
@@ -46,9 +46,9 @@ namespace UE::Audio::Insights
 			OnAudioBusAssetCheckedHandle = FAudioBusesDashboardViewFactory::OnAudioBusAssetChecked.AddSP(this, &FAudioMetersDashboardViewFactory::HandleOnAudioBusAssetChecked);
 		}
 
-		if (!OnSubmixAssetRemovedHandle.IsValid())
+		if (!OnSubmixDeactivatedHandle.IsValid())
 		{
-			OnSubmixAssetRemovedHandle = FSoundSubmixProvider::OnSubmixAssetRemoved.AddSP(this, &FAudioMetersDashboardViewFactory::HandleOnAssetRemoved);
+			OnSubmixDeactivatedHandle = FSoundSubmixProvider::OnSubmixDeactivated.AddSP(this, &FAudioMetersDashboardViewFactory::HandleOnAssetRemoved);
 		}
 
 		if (!OnAudioBusAssetRemovedHandle.IsValid())
@@ -69,7 +69,7 @@ namespace UE::Audio::Insights
 		return MeterViewsScrollBox.ToSharedRef();
 	}
 
-	void FAudioMetersDashboardViewFactory::HandleOnSubmixAssetChecked(const bool bInIsChecked, const TWeakObjectPtr<USoundSubmix> InSoundSubmix)
+	void FAudioMetersDashboardViewFactory::HandleOnSubmixItemChecked(const bool bInIsChecked, const TWeakObjectPtr<USoundSubmix> InSoundSubmix)
 	{
 		if (!InSoundSubmix.IsValid())
 		{
diff --git a/Engine/Plugins/AudioInsights/Source/AudioInsightsEditor/Private/Views/AudioMetersDashboardViewFactory.h b/Engine/Plugins/AudioInsights/Source/AudioInsightsEditor/Private/Views/AudioMetersDashboardViewFactory.h
--- a/Engine/Plugins/AudioInsights/Source/AudioInsightsEditor/Private/Views/AudioMetersDashboardViewFactory.h
+++ b/Engine/Plugins/AudioInsights/Source/AudioInsightsEditor/Private/Views/AudioMetersDashboardViewFactory.h
@@ -26,7 +26,7 @@ namespace UE::Audio::Insights
 		virtual TSharedRef<SWidget> MakeWidget(TSharedRef<SDockTab> OwnerTab, const FSpawnTabArgs& SpawnTabArgs) override;
 
 	private:
-		void HandleOnSubmixAssetChecked(const bool bInIsChecked, const TWeakObjectPtr<USoundSubmix> InSoundSubmix);
+		void HandleOnSubmixItemChecked(const bool bInIsChecked, const TWeakObjectPtr<USoundSubmix> InSoundSubmix);
 		void HandleOnAudioBusAssetChecked(const bool bInIsChecked, const TWeakObjectPtr<UAudioBus> InAudioBus);
 
 		void HandleOnAssetRemoved(const TWeakObjectPtr<UObject> InAsset);
@@ -35,9 +35,9 @@ namespace UE::Audio::Insights
 		TSharedPtr<SHorizontalBox> AudioMeterViewsContainer;
 		TMap<uint32, TSharedRef<FAudioMeterView>> AudioMeterViews;
 
-		FDelegateHandle OnSubmixAssetCheckedHandle;
+		FDelegateHandle OnSubmixItemCheckedHandle;
 		FDelegateHandle OnAudioBusAssetCheckedHandle;
-		FDelegateHandle OnSubmixAssetRemovedHandle;
+		FDelegateHandle OnSubmixDeactivatedHandle;
 		FDelegateHandle OnAudioBusAssetRemovedHandle;
 	};
 } // namespace UE::Audio::Insights
diff --git a/Engine/Plugins/AudioInsights/Source/AudioInsightsEditor/Private/Views/SoundSubmixAssetDashboardEntry.cpp b/Engine/Plugins/AudioInsights/Source/AudioInsightsEditor/Private/Views/SoundSubmixAssetDashboardEntry.cpp
deleted file mode 100644
diff --git a/Engine/Plugins/AudioInsights/Source/AudioInsightsEditor/Private/Views/SoundSubmixAssetDashboardEntry.h b/Engine/Plugins/AudioInsights/Source/AudioInsightsEditor/Private/Views/SoundSubmixAssetDashboardEntry.h
deleted file mode 100644
diff --git a/Engine/Plugins/AudioInsights/Source/AudioInsightsEditor/Private/Views/SubmixesDashboardViewFactory.cpp b/Engine/Plugins/AudioInsights/Source/AudioInsightsEditor/Private/Views/SubmixesDashboardViewFactory.cpp
--- a/Engine/Plugins/AudioInsights/Source/AudioInsightsEditor/Private/Views/SubmixesDashboardViewFactory.cpp
+++ b/Engine/Plugins/AudioInsights/Source/AudioInsightsEditor/Private/Views/SubmixesDashboardViewFactory.cpp
@@ -1,6 +1,7 @@
 // Copyright Epic Games, Inc. All Rights Reserved.
 #include "Views/SubmixesDashboardViewFactory.h"
 
+#include "AudioInsightsEditorModule.h"
 #include "AudioInsightsStyle.h"
 #include "Editor.h"
 #include "Internationalization/Text.h"
@@ -19,30 +20,30 @@ namespace UE::Audio::Insights
 		{
 			return static_cast<const FSoundSubmixAssetDashboardEntry&>(InData);
 		};
+
+		FSoundSubmixAssetDashboardEntry& CastEntry(IDashboardDataViewEntry& InData)
+		{
+			return static_cast<FSoundSubmixAssetDashboardEntry&>(InData);
+		};
 	}
 
 	FSubmixesDashboardViewFactory::FSubmixesDashboardViewFactory()
 	{
-		FSoundSubmixProvider::OnSubmixAssetAdded.AddRaw(this, &FSubmixesDashboardViewFactory::HandleOnSubmixAssetListUpdated);
-		FSoundSubmixProvider::OnSubmixAssetRemoved.AddRaw(this, &FSubmixesDashboardViewFactory::HandleOnSubmixAssetListUpdated);
-		FSoundSubmixProvider::OnSubmixAssetListUpdated.AddRaw(this, &FSubmixesDashboardViewFactory::RequestListRefresh);
+		IAudioInsightsTraceModule& AudioInsightsTraceModule = IAudioInsightsEditorModule::GetChecked().GetTraceModule();
+
+		const TSharedPtr<FSoundSubmixProvider> SoundSubmixProvider = MakeShared<FSoundSubmixProvider>();
+
+		AudioInsightsTraceModule.AddTraceProvider(SoundSubmixProvider);
 
 		Providers = TArray<TSharedPtr<FTraceProviderBase>>
 		{
-			MakeShared<FSoundSubmixProvider>()
+			SoundSubmixProvider
 		};
 
 		SortByColumn = "Name";
 		SortMode     = EColumnSortMode::Ascending;
 	}
 
-	FSubmixesDashboardViewFactory::~FSubmixesDashboardViewFactory()
-	{
-		FSoundSubmixProvider::OnSubmixAssetAdded.RemoveAll(this);
-		FSoundSubmixProvider::OnSubmixAssetRemoved.RemoveAll(this);
-		FSoundSubmixProvider::OnSubmixAssetListUpdated.RemoveAll(this);
-	}
-
 	FName FSubmixesDashboardViewFactory::GetName() const
 	{
 		return "Submixes";
@@ -55,49 +56,83 @@ namespace UE::Audio::Insights
 
 	TSharedRef<SWidget> FSubmixesDashboardViewFactory::GenerateWidgetForColumn(TSharedRef<IDashboardDataViewEntry> InRowData, const FName& InColumnName)
 	{
-		if (InColumnName == "Name")
+		if (InColumnName == "Active")
 		{
-			const TWeakObjectPtr<USoundSubmix> SoundSubmix = SubmixesPrivate::CastEntry(InRowData.Get()).SoundSubmix;
-
-			SubmixCheckboxCheckedStates.FindOrAdd(SoundSubmix, false);
+			static const FLinearColor DarkGreen(0.027f, 0.541f, 0.22f);
+			static const float Radius = 4.0f;
+			static const FVector2f Size(7.0f, 7.0f);
+			static const FSlateRoundedBoxBrush GreenRoundedBrush(DarkGreen, Radius, Size);
 
+			return SNew(SBox)
+				.Clipping(EWidgetClipping::ClipToBounds)
+				.Padding(6.0f)
+				.Visibility_Lambda([InRowData]()
+				{ 
+					const FSoundSubmixAssetDashboardEntry& SoundSubmixAssetDashboardEntry = SubmixesPrivate::CastEntry(InRowData.Get());
+					return SoundSubmixAssetDashboardEntry.bHasActivity ? EVisibility::Visible : EVisibility::Hidden;
+				})
+				[
+					SNew(SImage)
+					.Image(&GreenRoundedBrush)
+				];
+		}
+		else if (InColumnName == "Name")
+		{
 			const FColumnData& ColumnData = GetColumns()[InColumnName];
-			const FText Label = ColumnData.GetDisplayValue(InRowData.Get());
+			const FText ValueText = ColumnData.GetDisplayValue(InRowData.Get());
 
-			return (Label.IsEmpty()) ? 
-				SNullWidget::NullWidget :
-				SNew(SHorizontalBox)
+			if (ValueText.IsEmpty())
+			{
+				return SNullWidget::NullWidget;
+			}
+
+			FSoundSubmixAssetDashboardEntry& SoundSubmixAssetDashboardEntry = SubmixesPrivate::CastEntry(InRowData.Get());
+			const TWeakObjectPtr<USoundSubmix> SoundSubmix(Cast<USoundSubmix>(SoundSubmixAssetDashboardEntry.GetObject()));
+
+			if (!SoundSubmix.IsValid())
+			{
+				return SNullWidget::NullWidget;
+			}
+
+			SubmixCheckboxCheckedStates.FindOrAdd(SoundSubmix, false);
+
+			return SNew(SHorizontalBox)
 				.Clipping(EWidgetClipping::ClipToBounds)
 				+ SHorizontalBox::Slot()
-				.FillWidth(0.1f)
+				.AutoWidth()
 				[
 					SNew(SCheckBox)
-					.IsChecked(SubmixCheckboxCheckedStates.Contains(SoundSubmix) && SubmixCheckboxCheckedStates[SoundSubmix] == true)
-					.OnCheckStateChanged_Lambda([SoundSubmix, this](ECheckBoxState NewState)
+					.OnCheckStateChanged_Lambda([this, SoundSubmix](ECheckBoxState NewState)
 					{
-						if (SubmixCheckboxCheckedStates.Contains(SoundSubmix))
+						if (bool* FoundSubmix = SubmixCheckboxCheckedStates.Find(SoundSubmix))
 						{
-							SubmixCheckboxCheckedStates[SoundSubmix] = NewState == ECheckBoxState::Checked;
+							(*FoundSubmix) = NewState == ECheckBoxState::Checked;
 						}
 
-						OnSubmixAssetChecked.Broadcast(NewState == ECheckBoxState::Checked, SoundSubmix);
+						OnSubmixItemChecked.Broadcast(NewState == ECheckBoxState::Checked, SoundSubmix);
 					})
 				]
 				+ SHorizontalBox::Slot()
-				.FillWidth(0.8f)
+				.AutoWidth()
+				[
+					SNew(SBox)
+					.MinDesiredWidth(5.0f)
+				]
+				+ SHorizontalBox::Slot()
+				.AutoWidth()
 				[
 					SNew(STextBlock)
-					.Text(Label)
+					.Text(ValueText)
 					.MinDesiredWidth(300)
 					.OnDoubleClicked_Lambda([this, InRowData](const FGeometry& MyGeometry, const FPointerEvent& PointerEvent)
 					{
 						if (GEditor)
 						{
-							if (TSharedPtr<IObjectDashboardEntry> ObjectData = StaticCastSharedPtr<IObjectDashboardEntry>(InRowData.ToSharedPtr()); 
-								ObjectData.IsValid())
+							TSharedPtr<IObjectDashboardEntry> ObjectData = StaticCastSharedPtr<IObjectDashboardEntry>(InRowData.ToSharedPtr());
+							if (ObjectData.IsValid())
 							{
-								if (const UObject* Object = ObjectData->GetObject();
-									Object && Object->IsAsset())
+								const UObject* Object = ObjectData->GetObject();
+								if (Object && Object->IsAsset())
 								{
 									GEditor->GetEditorSubsystem<UAssetEditorSubsystem>()->OpenEditorForAsset(Object);
 									return FReply::Handled();
@@ -156,13 +191,23 @@ namespace UE::Audio::Insights
 		{
 			return TMap<FName, FTraceTableDashboardViewFactory::FColumnData>
 			{
+				{
+					"Active",
+					{
+						LOCTEXT("Submixes_ActiveDisplayName", "Active"),
+						[](const IDashboardDataViewEntry& InData) { return FText::GetEmpty(); },
+						false,		/* bDefaultHidden */
+						0.08f,		/* FillWidth */
+						EHorizontalAlignment::HAlign_Center
+					}
+				},
 				{
 					"Name",
 					{
 						LOCTEXT("Submixes_NameColumnDisplayName", "Name"),
 						[](const IDashboardDataViewEntry& InData) { return SubmixesPrivate::CastEntry(InData).GetDisplayName(); },
 						false,		/* bDefaultHidden */
-						0.88f,		/* FillWidth */
+						0.92f,		/* FillWidth */
 						EHorizontalAlignment::HAlign_Left
 					}
 				}
@@ -176,7 +221,30 @@ namespace UE::Audio::Insights
 
 	void FSubmixesDashboardViewFactory::SortTable()
 	{
-		if (SortByColumn == "Name")
+		if (SortByColumn == "Active")
+		{
+			if (SortMode == EColumnSortMode::Ascending)
+			{
+				DataViewEntries.Sort([](const TSharedPtr<IDashboardDataViewEntry>& A, const TSharedPtr<IDashboardDataViewEntry>& B)
+				{
+					const FSoundSubmixAssetDashboardEntry& AData = SubmixesPrivate::CastEntry(*A.Get());
+					const FSoundSubmixAssetDashboardEntry& BData = SubmixesPrivate::CastEntry(*B.Get());
+
+					return !AData.bHasActivity && BData.bHasActivity;
+				});
+			}
+			else if (SortMode == EColumnSortMode::Descending)
+			{
+				DataViewEntries.Sort([](const TSharedPtr<IDashboardDataViewEntry>& A, const TSharedPtr<IDashboardDataViewEntry>& B)
+				{
+					const FSoundSubmixAssetDashboardEntry& AData = SubmixesPrivate::CastEntry(*A.Get());
+					const FSoundSubmixAssetDashboardEntry& BData = SubmixesPrivate::CastEntry(*B.Get());
+
+					return !BData.bHasActivity && AData.bHasActivity;
+				});
+			}
+		}
+		else if (SortByColumn == "Name")
 		{
 			if (SortMode == EColumnSortMode::Ascending)
 			{
@@ -205,23 +273,14 @@ namespace UE::Audio::Insights
 	{
 		if (SelectedItem.IsValid())
 		{
-			const FSoundSubmixAssetDashboardEntry& SoundSubmixAssetDashboardEntry = SubmixesPrivate::CastEntry(*SelectedItem.Get());
-
-			OnSubmixSelectionChanged.Broadcast(SoundSubmixAssetDashboardEntry.SoundSubmix);
-		}
-	}
+			FSoundSubmixAssetDashboardEntry& SoundSubmixAssetDashboardEntry = SubmixesPrivate::CastEntry(*SelectedItem.Get());
 
-	void FSubmixesDashboardViewFactory::RequestListRefresh()
-	{
-		if (FilteredEntriesListView.IsValid())
-		{
-			FilteredEntriesListView->RequestListRefresh();
+			const TWeakObjectPtr<USoundSubmix> SoundSubmixWeakObjectPtr(Cast<USoundSubmix>(SoundSubmixAssetDashboardEntry.GetObject()));
+			if (SoundSubmixWeakObjectPtr.IsValid())
+			{
+				OnSubmixSelectionChanged.Broadcast(SoundSubmixWeakObjectPtr);
+			}
 		}
 	}
-
-	void FSubmixesDashboardViewFactory::HandleOnSubmixAssetListUpdated(const TWeakObjectPtr<UObject> InAsset)
-	{
-		RequestListRefresh();
-	}
 } // namespace UE::Audio::Insights
 #undef LOCTEXT_NAMESPACE
diff --git a/Engine/Plugins/AudioInsights/Source/AudioInsightsEditor/Private/Views/SubmixesDashboardViewFactory.h b/Engine/Plugins/AudioInsights/Source/AudioInsightsEditor/Private/Views/SubmixesDashboardViewFactory.h
--- a/Engine/Plugins/AudioInsights/Source/AudioInsightsEditor/Private/Views/SubmixesDashboardViewFactory.h
+++ b/Engine/Plugins/AudioInsights/Source/AudioInsightsEditor/Private/Views/SubmixesDashboardViewFactory.h
@@ -3,7 +3,7 @@
 
 #include "Delegates/Delegate.h"
 #include "Templates/SharedPointer.h"
-#include "Views/SoundSubmixAssetDashboardEntry.h"
+#include "Messages/SoundSubmixMessages.h"
 #include "Views/TableDashboardViewFactory.h"
 #include "Widgets/Input/SCheckBox.h"
 
@@ -15,7 +15,7 @@ namespace UE::Audio::Insights
 	{
 	public:
 		FSubmixesDashboardViewFactory();
-		virtual ~FSubmixesDashboardViewFactory();
+		virtual ~FSubmixesDashboardViewFactory() = default;
 
 		virtual FName GetName() const override;
 		virtual FText GetDisplayName() const override;
@@ -23,8 +23,8 @@ namespace UE::Audio::Insights
 		virtual EDefaultDashboardTabStack GetDefaultTabStack() const override;
 		virtual TSharedRef<SWidget> MakeWidget(TSharedRef<SDockTab> OwnerTab, const FSpawnTabArgs& SpawnTabArgs) override;
 
-		DECLARE_MULTICAST_DELEGATE_TwoParams(FOnSubmixAssetChecked, const bool /*bIsChecked*/, TWeakObjectPtr<USoundSubmix> /*SoundSubmix*/);
-		inline static FOnSubmixAssetChecked OnSubmixAssetChecked;
+		DECLARE_MULTICAST_DELEGATE_TwoParams(FOnSubmixItemChecked, const bool /*bIsChecked*/, TWeakObjectPtr<USoundSubmix> /*SoundSubmix*/);
+		inline static FOnSubmixItemChecked OnSubmixItemChecked;
 
 		DECLARE_MULTICAST_DELEGATE_OneParam(FOnSubmixSelectionChanged, const TWeakObjectPtr<USoundSubmix> /*SoundSubmix*/);
 		inline static FOnSubmixSelectionChanged OnSubmixSelectionChanged;
@@ -38,10 +38,6 @@ namespace UE::Audio::Insights
 
 		virtual void OnSelectionChanged(TSharedPtr<IDashboardDataViewEntry> SelectedItem, ESelectInfo::Type SelectInfo) override;
 
-		void RequestListRefresh();
-
-		void HandleOnSubmixAssetListUpdated(const TWeakObjectPtr<UObject> InAsset);
-
 		TMap<const TWeakObjectPtr<USoundSubmix>, bool> SubmixCheckboxCheckedStates;
 	};
 } // namespace UE::Audio::Insights
diff --git a/Engine/Source/Runtime/AudioMixer/Private/AudioMixerSubmix.cpp b/Engine/Source/Runtime/AudioMixer/Private/AudioMixerSubmix.cpp
--- a/Engine/Source/Runtime/AudioMixer/Private/AudioMixerSubmix.cpp
+++ b/Engine/Source/Runtime/AudioMixer/Private/AudioMixerSubmix.cpp
@@ -57,6 +57,30 @@ DEFINE_STAT(STAT_AudioMixerSubmixSoundfieldChildren);
 DEFINE_STAT(STAT_AudioMixerSubmixSoundfieldSources);
 DEFINE_STAT(STAT_AudioMixerSubmixSoundfieldProcessors);
 
+#if UE_AUDIO_PROFILERTRACE_ENABLED
+
+UE_TRACE_EVENT_BEGIN(Audio, SoundSubmixActivate)
+	UE_TRACE_EVENT_FIELD(uint32, DeviceId)
+	UE_TRACE_EVENT_FIELD(uint32, SubmixId)
+	UE_TRACE_EVENT_FIELD(double, Timestamp)
+	UE_TRACE_EVENT_FIELD(UE::Trace::WideString, Name)
+UE_TRACE_EVENT_END()
+
+UE_TRACE_EVENT_BEGIN(Audio, SoundSubmixDeactivate)
+	UE_TRACE_EVENT_FIELD(uint32, DeviceId)
+	UE_TRACE_EVENT_FIELD(uint32, SubmixId)
+	UE_TRACE_EVENT_FIELD(double, Timestamp)
+UE_TRACE_EVENT_END()
+
+UE_TRACE_EVENT_BEGIN(Audio, SoundSubmixHasActivity)
+	UE_TRACE_EVENT_FIELD(uint32, DeviceId)
+	UE_TRACE_EVENT_FIELD(uint32, SubmixId)
+	UE_TRACE_EVENT_FIELD(double, Timestamp)
+	UE_TRACE_EVENT_FIELD(bool, HasActivity)
+UE_TRACE_EVENT_END()
+
+#endif // UE_AUDIO_PROFILERTRACE_ENABLED
+
 namespace Audio
 {
 	namespace MixerSubmixIntrinsics
@@ -215,6 +239,16 @@ namespace Audio
 				MutableSubmix->StopRecordingOutput(MixerDevice, EAudioRecordingExportType::WavFile, InterruptedFileName, FString());
 			}
 		}
+
+#if UE_AUDIO_PROFILERTRACE_ENABLED
+		if (MixerDevice && OwningSubmixObject.IsValid())
+		{
+			UE_TRACE_LOG(Audio, SoundSubmixDeactivate, AudioChannel)
+				<< SoundSubmixDeactivate.DeviceId(MixerDevice->DeviceID)
+				<< SoundSubmixDeactivate.SubmixId(OwningSubmixObject->GetUniqueID())
+				<< SoundSubmixDeactivate.Timestamp(FPlatformTime::Cycles64());
+		}
+#endif
 	}
 
 	void FMixerSubmix::Init(const USoundSubmixBase* InSoundSubmix, bool bAllowReInit)
@@ -247,6 +281,18 @@ namespace Audio
 				TArray<USoundfieldEffectBase*> Effects = SoundfieldSubmix->GetSoundfieldProcessors();
 				SetupSoundfieldStreams(EncodingSettings, Effects, SoundfieldFactory);
 			}
+
+#if UE_AUDIO_PROFILERTRACE_ENABLED
+			if (MixerDevice)
+			{
+				UE_LOG(LogAudioMixer, Warning, TEXT("SoundSubmixActivate: %s"), *(OwningSubmixObject->GetName()));
+				UE_TRACE_LOG(Audio, SoundSubmixActivate, AudioChannel)
+					<< SoundSubmixActivate.DeviceId(MixerDevice->DeviceID)
+					<< SoundSubmixActivate.SubmixId(OwningSubmixObject->GetUniqueID())
+					<< SoundSubmixActivate.Timestamp(FPlatformTime::Cycles64())
+					<< SoundSubmixActivate.Name(*OwningSubmixObject->GetPathName());
+			}
+#endif
 		}
 	}
 
@@ -1614,6 +1660,17 @@ namespace Audio
 				bIsSilent = false;
 				SilenceTimeStartSeconds = -1.0;
 			}
+
+#if UE_AUDIO_PROFILERTRACE_ENABLED
+			if (MixerDevice && OwningSubmixObject.IsValid())
+			{
+				UE_TRACE_LOG(Audio, SoundSubmixHasActivity, AudioChannel)
+					<< SoundSubmixHasActivity.DeviceId(MixerDevice->DeviceID)
+					<< SoundSubmixHasActivity.SubmixId(OwningSubmixObject->GetUniqueID())
+					<< SoundSubmixHasActivity.Timestamp(FPlatformTime::Cycles64())
+					<< SoundSubmixHasActivity.HasActivity(!bIsNowSilent);
+			}
+#endif
 		}
 
 	}
