Fixing small mistake in SoundSubmixProvider

diff --git a/Engine/Plugins/AudioInsights/Source/AudioInsightsEditor/Private/Providers/SoundSubmixProvider.cpp b/Engine/Plugins/AudioInsights/Source/AudioInsightsEditor/Private/Providers/SoundSubmixProvider.cpp
--- a/Engine/Plugins/AudioInsights/Source/AudioInsightsEditor/Private/Providers/SoundSubmixProvider.cpp
+++ b/Engine/Plugins/AudioInsights/Source/AudioInsightsEditor/Private/Providers/SoundSubmixProvider.cpp
@@ -38,7 +38,7 @@ namespace UE::Audio::Insights
 
 				Entry->Timestamp = Msg.Timestamp;
 
-				*ToReturn = Entry;
+				ToReturn = &Entry;
 			});
 
 			return ToReturn;
