[Audio Insights] Submixes Dashboard: Bug fix for checkboxes ending up in a wrong state when PIE is stopped

diff --git a/Engine/Plugins/AudioInsights/Source/AudioInsightsEditor/Private/Views/SubmixesDashboardViewFactory.cpp b/Engine/Plugins/AudioInsights/Source/AudioInsightsEditor/Private/Views/SubmixesDashboardViewFactory.cpp
--- a/Engine/Plugins/AudioInsights/Source/AudioInsightsEditor/Private/Views/SubmixesDashboardViewFactory.cpp
+++ b/Engine/Plugins/AudioInsights/Source/AudioInsightsEditor/Private/Views/SubmixesDashboardViewFactory.cpp
@@ -102,11 +102,16 @@ namespace UE::Audio::Insights
 				.AutoWidth()
 				[
 					SNew(SCheckBox)
+					.IsChecked_Lambda([this, SoundSubmix]()
+					{
+						const bool* FoundSubmixCheckedState = SubmixCheckboxCheckedStates.Find(SoundSubmix);
+						return FoundSubmixCheckedState && *FoundSubmixCheckedState == true ? ECheckBoxState::Checked : ECheckBoxState::Unchecked;
+					})
 					.OnCheckStateChanged_Lambda([this, SoundSubmix](ECheckBoxState NewState)
 					{
-						if (bool* FoundSubmix = SubmixCheckboxCheckedStates.Find(SoundSubmix))
+						if (bool* FoundSubmixCheckedState = SubmixCheckboxCheckedStates.Find(SoundSubmix))
 						{
-							(*FoundSubmix) = NewState == ECheckBoxState::Checked;
+							*FoundSubmixCheckedState = NewState == ECheckBoxState::Checked;
 						}
 
 						OnSubmixItemChecked.Broadcast(NewState == ECheckBoxState::Checked, SoundSubmix);
@@ -174,6 +179,8 @@ namespace UE::Audio::Insights
 	{
 		if (!DashboardWidget.IsValid())
 		{
+			FAudioDeviceManagerDelegates::OnAudioDeviceDestroyed.AddSP(this, &FSubmixesDashboardViewFactory::HandleOnAudioDeviceDestroyed);
+
 			DashboardWidget = FTraceTableDashboardViewFactory::MakeWidget(OwnerTab, SpawnTabArgs);
 
 			if (FilteredEntriesListView.IsValid())
@@ -282,5 +289,10 @@ namespace UE::Audio::Insights
 			}
 		}
 	}
+
+	void FSubmixesDashboardViewFactory::HandleOnAudioDeviceDestroyed(::Audio::FDeviceId InAudioDeviceId)
+	{
+		SubmixCheckboxCheckedStates.Empty();
+	}
 } // namespace UE::Audio::Insights
 #undef LOCTEXT_NAMESPACE
diff --git a/Engine/Plugins/AudioInsights/Source/AudioInsightsEditor/Private/Views/SubmixesDashboardViewFactory.h b/Engine/Plugins/AudioInsights/Source/AudioInsightsEditor/Private/Views/SubmixesDashboardViewFactory.h
--- a/Engine/Plugins/AudioInsights/Source/AudioInsightsEditor/Private/Views/SubmixesDashboardViewFactory.h
+++ b/Engine/Plugins/AudioInsights/Source/AudioInsightsEditor/Private/Views/SubmixesDashboardViewFactory.h
@@ -38,6 +38,8 @@ namespace UE::Audio::Insights
 
 		virtual void OnSelectionChanged(TSharedPtr<IDashboardDataViewEntry> SelectedItem, ESelectInfo::Type SelectInfo) override;
 
+		void HandleOnAudioDeviceDestroyed(::Audio::FDeviceId InAudioDeviceId);
+
 		TMap<const TWeakObjectPtr<USoundSubmix>, bool> SubmixCheckboxCheckedStates;
 	};
 } // namespace UE::Audio::Insights
