[Audio Insights] AudioBuses Dashboard: refactored audio buses dashboard so it is 100% trace based instead of asset events based + added activity indicator column

diff --git a/Engine/Plugins/AudioInsights/Source/AudioInsightsEditor/Private/Messages/AudioBusMessages.h b/Engine/Plugins/AudioInsights/Source/AudioInsightsEditor/Private/Messages/AudioBusMessages.h
new file mode 100644
--- /dev/null
+++ b/Engine/Plugins/AudioInsights/Source/AudioInsightsEditor/Private/Messages/AudioBusMessages.h
@@ -0,0 +1,88 @@
+// Copyright Epic Games, Inc. All Rights Reserved.
+#pragma once
+
+#include "AudioDefines.h"
+#include "Math/NumericLimits.h"
+#include "Messages/AnalyzerMessageQueue.h"
+#include "Trace/Analyzer.h"
+#include "Views/TableDashboardViewFactory.h"
+
+namespace UE::Audio::Insights
+{
+	enum class EAudioBusEntryType : uint8
+	{
+		AssetBased,
+		CodeGenerated,
+		None
+	};
+
+	struct FAudioBusMessageBase
+	{
+		FAudioBusMessageBase() = default;
+
+		FAudioBusMessageBase(const UE::Trace::IAnalyzer::FOnEventContext& InContext)
+		{
+			const UE::Trace::IAnalyzer::FEventData& EventData = InContext.EventData;
+
+			DeviceId = static_cast<::Audio::FDeviceId>(EventData.GetValue<uint32>("DeviceId"));
+			AudioBusId = EventData.GetValue<uint32>("AudioBusId");
+			Timestamp = InContext.EventTime.AsSeconds(EventData.GetValue<uint64>("Timestamp"));
+		}
+
+		::Audio::FDeviceId DeviceId = INDEX_NONE;
+		uint32 AudioBusId = INDEX_NONE;
+		double Timestamp = 0.0;
+	};
+
+	struct FAudioBusActivateMessage : public FAudioBusMessageBase
+	{
+		FAudioBusActivateMessage() = default;
+
+		FAudioBusActivateMessage(const UE::Trace::IAnalyzer::FOnEventContext& InContext)
+			: FAudioBusMessageBase(InContext)
+		{
+			const UE::Trace::IAnalyzer::FEventData& EventData = InContext.EventData;
+
+			EventData.GetString("Name", Name);
+		}
+
+		FString Name;
+	};
+
+	using FAudioBusDeactivateMessage = FAudioBusMessageBase;
+
+	struct FAudioBusHasActivityMessage : public FAudioBusMessageBase
+	{
+		FAudioBusHasActivityMessage() = default;
+
+		FAudioBusHasActivityMessage(const UE::Trace::IAnalyzer::FOnEventContext& InContext)
+			: FAudioBusMessageBase(InContext)
+		{
+			const UE::Trace::IAnalyzer::FEventData& EventData = InContext.EventData;
+
+			bHasActivity = EventData.GetValue<bool>("HasActivity");
+		}
+
+		bool bHasActivity = false;
+	};
+
+	class FAudioBusMessages
+	{
+		UE::Audio::Insights::TAnalyzerMessageQueue<FAudioBusActivateMessage> ActivateMessages{ 0.1 };
+		UE::Audio::Insights::TAnalyzerMessageQueue<FAudioBusDeactivateMessage> DeactivateMessages{ 0.1 };
+		UE::Audio::Insights::TAnalyzerMessageQueue<FAudioBusHasActivityMessage> HasActivityMessages{ 0.1 };
+
+		friend class FAudioBusProvider;
+	};
+
+	struct FAudioBusAssetDashboardEntry : public UE::Audio::Insights::FSoundAssetDashboardEntry
+	{
+		FAudioBusAssetDashboardEntry() = default;
+		virtual ~FAudioBusAssetDashboardEntry() = default;
+
+		EAudioBusEntryType EntryType = EAudioBusEntryType::None;
+
+		uint32 AudioBusId = INDEX_NONE;
+		bool bHasActivity = false;
+	};
+} // namespace UE::Audio::Insights
diff --git a/Engine/Plugins/AudioInsights/Source/AudioInsightsEditor/Private/Providers/AudioBusProvider.cpp b/Engine/Plugins/AudioInsights/Source/AudioInsightsEditor/Private/Providers/AudioBusProvider.cpp
--- a/Engine/Plugins/AudioInsights/Source/AudioInsightsEditor/Private/Providers/AudioBusProvider.cpp
+++ b/Engine/Plugins/AudioInsights/Source/AudioInsightsEditor/Private/Providers/AudioBusProvider.cpp
@@ -13,155 +13,155 @@ namespace UE::Audio::Insights
 	FAudioBusProvider::FAudioBusProvider()
 		: TDeviceDataMapTraceProvider<uint32, TSharedPtr<FAudioBusAssetDashboardEntry>>(GetName_Static())
 	{
-		FAssetRegistryModule& AssetRegistryModule = FModuleManager::LoadModuleChecked<FAssetRegistryModule>("AssetRegistry");
-
-		AssetRegistryModule.Get().OnAssetAdded().AddRaw(this, &FAudioBusProvider::OnAssetAdded);
-		AssetRegistryModule.Get().OnAssetRemoved().AddRaw(this, &FAudioBusProvider::OnAssetRemoved);
-		AssetRegistryModule.Get().OnFilesLoaded().AddRaw(this, &FAudioBusProvider::OnFilesLoaded);
-
-		FEditorDashboardFactory::OnActiveAudioDeviceChanged.AddRaw(this, &FAudioBusProvider::OnActiveAudioDeviceChanged);
-	}
-	
-	FAudioBusProvider::~FAudioBusProvider()
-	{
-		if (FAssetRegistryModule* AssetRegistryModule = FModuleManager::GetModulePtr<FAssetRegistryModule>("AssetRegistry"))
-		{
-			AssetRegistryModule->Get().OnAssetAdded().RemoveAll(this);
-			AssetRegistryModule->Get().OnAssetRemoved().RemoveAll(this);
-			AssetRegistryModule->Get().OnFilesLoaded().RemoveAll(this);
-		}
-
-		FEditorDashboardFactory::OnActiveAudioDeviceChanged.RemoveAll(this);
 	}
 
 	FName FAudioBusProvider::GetName_Static()
 	{
-		return "AudioBusProvider";
-	}
-
-	void FAudioBusProvider::OnAssetAdded(const FAssetData& InAssetData)
-	{
-		if (bAreFilesLoaded && InAssetData.AssetClassPath == FTopLevelAssetPath(UAudioBus::StaticClass()))
-		{
-			AddAudioBusAsset(InAssetData);
-		}
+		return "AudioBusesProvider";
 	}
 
-	void FAudioBusProvider::OnAssetRemoved(const FAssetData& InAssetData)
+	bool FAudioBusProvider::ProcessMessages()
 	{
-		if (InAssetData.AssetClassPath == FTopLevelAssetPath(UAudioBus::StaticClass()))
+		// Helper lambdas
+		auto BumpEntryFunc = [this](const FAudioBusMessageBase& Msg)
 		{
-			RemoveAudioBusAsset(InAssetData);
-		}
-	}
+			TSharedPtr<FAudioBusAssetDashboardEntry>* ToReturn = nullptr;
 
-	void FAudioBusProvider::OnFilesLoaded()
-	{
-		bAreFilesLoaded = true;
-		UpdateAudioBusAssetNames();
-	}
+			UpdateDeviceEntry(Msg.DeviceId, Msg.AudioBusId, [&ToReturn, &Msg](TSharedPtr<FAudioBusAssetDashboardEntry>& Entry)
+			{
+				if (!Entry.IsValid())
+				{
+					Entry = MakeShared<FAudioBusAssetDashboardEntry>();
+					Entry->DeviceId = Msg.DeviceId;
+					Entry->AudioBusId = Msg.AudioBusId;
+				}
 
-	void FAudioBusProvider::OnActiveAudioDeviceChanged()
-	{
-		UpdateAudioBusAssetNames();
-	}
+				Entry->Timestamp = Msg.Timestamp;
 
-	void FAudioBusProvider::AddAudioBusAsset(const FAssetData& InAssetData)
-	{
-		const bool bIsAudioBusAssetAlreadAdded = AudioBusDataViewEntries.ContainsByPredicate(
-			[AssetName = InAssetData.GetObjectPathString()](const TSharedPtr<FAudioBusAssetDashboardEntry>& AudioBusAssetDashboardEntry)
-			{
-				return AudioBusAssetDashboardEntry->Name == AssetName;
+				ToReturn = &Entry;
 			});
 
-		if (!bIsAudioBusAssetAlreadAdded)
-		{
-			const FAudioInsightsEditorModule AudioInsightsEditorModule = FAudioInsightsEditorModule::GetChecked();
-			const ::Audio::FDeviceId AudioDeviceId = AudioInsightsEditorModule.GetDeviceId();
+			return ToReturn;
+		};
 
-			TSharedPtr<FAudioBusAssetDashboardEntry> AudioBusAssetDashboardEntry = MakeShared<FAudioBusAssetDashboardEntry>();
-			AudioBusAssetDashboardEntry->DeviceId = AudioDeviceId;
-			AudioBusAssetDashboardEntry->Name     = InAssetData.GetObjectPathString();
-			AudioBusAssetDashboardEntry->AudioBus = Cast<UAudioBus>(InAssetData.GetAsset());
+		auto GetEntry = [this](const FAudioBusMessageBase& Msg)
+		{
+			return FindDeviceEntry(Msg.DeviceId, Msg.AudioBusId);
+		};
 
-			AudioBusDataViewEntries.Add(MoveTemp(AudioBusAssetDashboardEntry));
 
-			OnAudioBusAssetAdded.Broadcast(InAssetData.GetAsset());
-			++LastUpdateId;
-		}
-	}
+		// Process messages
+		ProcessMessageQueue<FAudioBusActivateMessage>(TraceMessages.ActivateMessages, BumpEntryFunc,
+		[](const FAudioBusActivateMessage& Msg, TSharedPtr<FAudioBusAssetDashboardEntry>* OutEntry)
+		{
+			FAudioBusAssetDashboardEntry& EntryRef = *OutEntry->Get();
+			EntryRef.Name = *Msg.Name;
+			EntryRef.EntryType = FSoftObjectPath(EntryRef.Name).IsAsset() ? EAudioBusEntryType::AssetBased : EAudioBusEntryType::CodeGenerated;
+		});
 
-	void FAudioBusProvider::RemoveAudioBusAsset(const FAssetData& InAssetData)
-	{
-		const int32 FoundAudioBusAssetNameIndex = AudioBusDataViewEntries.IndexOfByPredicate(
-			[AssetName = InAssetData.GetObjectPathString()](const TSharedPtr<FAudioBusAssetDashboardEntry>& AudioBusAssetDashboardEntry)
+		ProcessMessageQueue<FAudioBusDeactivateMessage>(TraceMessages.DeactivateMessages, GetEntry,
+		[this](const FAudioBusDeactivateMessage& Msg, TSharedPtr<FAudioBusAssetDashboardEntry>* OutEntry)
+		{
+			if (OutEntry && (*OutEntry)->Timestamp < Msg.Timestamp)
 			{
-				return AudioBusAssetDashboardEntry->Name == AssetName;
-			});
+				FAudioBusAssetDashboardEntry& EntryRef = *OutEntry->Get();
 
-		if (FoundAudioBusAssetNameIndex != INDEX_NONE)
-		{
-			const FAudioInsightsEditorModule AudioInsightsEditorModule = FAudioInsightsEditorModule::GetChecked();
-			const ::Audio::FDeviceId AudioDeviceId = AudioInsightsEditorModule.GetDeviceId();
+				OnAudioBusDeactivated.Broadcast(EntryRef.GetObject());
 
-			RemoveDeviceEntry(AudioDeviceId, AudioBusDataViewEntries[FoundAudioBusAssetNameIndex]->AudioBus->GetUniqueID());
+				RemoveDeviceEntry(Msg.DeviceId, Msg.AudioBusId);
+			}
+		});
 
-			AudioBusDataViewEntries.RemoveAt(FoundAudioBusAssetNameIndex);
+		ProcessMessageQueue<FAudioBusHasActivityMessage>(TraceMessages.HasActivityMessages, GetEntry,
+		[](const FAudioBusHasActivityMessage& Msg, TSharedPtr<FAudioBusAssetDashboardEntry>* OutEntry)
+		{
+			if (OutEntry)
+			{
+				FAudioBusAssetDashboardEntry& EntryRef = *OutEntry->Get();
+				EntryRef.bHasActivity = Msg.bHasActivity;
+				EntryRef.Timestamp = Msg.Timestamp;
+			}
+		});
 
-			OnAudioBusAssetRemoved.Broadcast(InAssetData.GetAsset());
-			++LastUpdateId;
-		}
+		return true;
 	}
 
-	void FAudioBusProvider::UpdateAudioBusAssetNames()
+	UE::Trace::IAnalyzer* FAudioBusProvider::ConstructAnalyzer(TraceServices::IAnalysisSession& InSession)
 	{
-		// Get all UAudioBus assets
-		TArray<FAssetData> AssetDataArray;
-
-		FAssetRegistryModule& AssetRegistryModule = FModuleManager::LoadModuleChecked<FAssetRegistryModule>("AssetRegistry");
-		AssetRegistryModule.Get().GetAssetsByClass(FTopLevelAssetPath(UAudioBus::StaticClass()), AssetDataArray);
-
-		// Build AudioBusDataViewEntries
-		Reset();
-		AudioBusDataViewEntries.Empty();
-
-		for (const FAssetData& AssetData : AssetDataArray)
+		class FAudioBusTraceAnalyzer : public FTraceAnalyzerBase
 		{
-			AddAudioBusAsset(AssetData);
-		}
+		public:
+			FAudioBusTraceAnalyzer(TSharedRef<FAudioBusProvider> InProvider, TraceServices::IAnalysisSession& InSession)
+				: FTraceAnalyzerBase(InProvider)
+				, Session(InSession)
+			{
+			}
 
-		AudioBusDataViewEntries.Sort([](const TSharedPtr<FAudioBusAssetDashboardEntry>& A, const TSharedPtr<FAudioBusAssetDashboardEntry>& B)
-		{
-			return A->GetDisplayName().CompareToCaseIgnored(B->GetDisplayName()) < 0;
-		});
+			virtual void OnAnalysisBegin(const FOnAnalysisContext& Context) override
+			{
+				FTraceAnalyzerBase::OnAnalysisBegin(Context);
 
-		OnAudioBusAssetListUpdated.Broadcast();
-	}
+				UE::Trace::IAnalyzer::FInterfaceBuilder& Builder = Context.InterfaceBuilder;
 
-	bool FAudioBusProvider::ProcessMessages()
-	{
-		const FAudioInsightsEditorModule AudioInsightsEditorModule = FAudioInsightsEditorModule::GetChecked();
-		const ::Audio::FDeviceId AudioDeviceId = AudioInsightsEditorModule.GetDeviceId();
+				Builder.RouteEvent(RouteId_Activate,    "Audio", "AudioBusActivate");
+				Builder.RouteEvent(RouteId_Deactivate,  "Audio", "AudioBusDeactivate");
+				Builder.RouteEvent(RouteId_HasActivity, "Audio", "AudioBusHasActivity");
+			}
 
-		for (const TSharedPtr<FAudioBusAssetDashboardEntry>& AudioBusDataViewEntry : AudioBusDataViewEntries)
-		{
-			if (AudioBusDataViewEntry.IsValid() && AudioBusDataViewEntry->AudioBus.IsValid())
+			virtual bool OnEvent(uint16 RouteId, EStyle Style, const FOnEventContext& Context) override
 			{
-				UpdateDeviceEntry(AudioDeviceId, AudioBusDataViewEntry->AudioBus->GetUniqueID(), [&AudioBusDataViewEntry](TSharedPtr<FAudioBusAssetDashboardEntry>& Entry)
+				LLM_SCOPE_BYNAME(TEXT("Insights/FAudioBusTraceAnalyzer"));
+
+				FAudioBusMessages& Messages = GetProvider<FAudioBusProvider>().TraceMessages;
+
+				switch (RouteId)
 				{
-					if (!Entry.IsValid())
+					case RouteId_Activate:
+					{
+						Messages.ActivateMessages.Enqueue(FAudioBusActivateMessage { Context });
+						break;
+					}
+
+					case RouteId_Deactivate:
 					{
-						Entry = AudioBusDataViewEntry;
+						Messages.DeactivateMessages.Enqueue(FAudioBusDeactivateMessage{ Context });
+						break;
 					}
-				});
+
+					case RouteId_HasActivity:
+					{
+						Messages.HasActivityMessages.Enqueue(FAudioBusHasActivityMessage{ Context });
+						break;
+					}
+
+					default:
+					{
+						return OnEventFailure(RouteId, Style, Context);
+					}
+				}
+
+				const double Timestamp = Context.EventTime.AsSeconds(Context.EventData.GetValue<uint64>("Timestamp"));
+
+				{
+					TraceServices::FAnalysisSessionEditScope SessionEditScope(Session);
+					Session.UpdateDurationSeconds(Timestamp);
+				}
+
+				return OnEventSuccess(RouteId, Style, Context);
 			}
-		}
 
-		return true;
-	}
+		private:
+			enum : uint16
+			{
+				RouteId_Activate,
+				RouteId_Deactivate,
+				RouteId_HasActivity
+			};
 
-	UE::Trace::IAnalyzer* FAudioBusProvider::ConstructAnalyzer(TraceServices::IAnalysisSession& InSession)
-	{
-		return nullptr;
+			TraceServices::IAnalysisSession& Session;
+		};
+
+		return new FAudioBusTraceAnalyzer(AsShared(), InSession);
 	}
 } // namespace UE::Audio::Insights
+
diff --git a/Engine/Plugins/AudioInsights/Source/AudioInsightsEditor/Private/Providers/AudioBusProvider.h b/Engine/Plugins/AudioInsights/Source/AudioInsightsEditor/Private/Providers/AudioBusProvider.h
--- a/Engine/Plugins/AudioInsights/Source/AudioInsightsEditor/Private/Providers/AudioBusProvider.h
+++ b/Engine/Plugins/AudioInsights/Source/AudioInsightsEditor/Private/Providers/AudioBusProvider.h
@@ -2,8 +2,8 @@
 #pragma once
 
 #include "Delegates/Delegate.h"
+#include "Messages/AudioBusMessages.h"
 #include "UObject/NameTypes.h"
-#include "Views/AudioBusAssetDashboardEntry.h"
 
 namespace UE::Audio::Insights
 {
@@ -11,36 +11,18 @@ namespace UE::Audio::Insights
 	{
 	public:
 		FAudioBusProvider();
-		virtual ~FAudioBusProvider();
+		virtual ~FAudioBusProvider() = default;
 
 		static FName GetName_Static();
 
 		virtual UE::Trace::IAnalyzer* ConstructAnalyzer(TraceServices::IAnalysisSession& InSession) override;
-		
-		DECLARE_MULTICAST_DELEGATE_OneParam(FOnAudioBusAssetAdded, const TWeakObjectPtr<UObject> /*Asset*/);
-		inline static FOnAudioBusAssetAdded OnAudioBusAssetAdded;
 
-		DECLARE_MULTICAST_DELEGATE_OneParam(FOnAudioBusAssetRemoved, const TWeakObjectPtr<UObject> /*Asset*/);
-		inline static FOnAudioBusAssetRemoved OnAudioBusAssetRemoved;
-
-		DECLARE_MULTICAST_DELEGATE(FOnAudioBusAssetListUpdated);
-		inline static FOnAudioBusAssetListUpdated OnAudioBusAssetListUpdated;
+		DECLARE_MULTICAST_DELEGATE_OneParam(FOnAudioBusDeactivated, const TWeakObjectPtr<UObject> /*AudioBus*/);
+		inline static FOnAudioBusDeactivated OnAudioBusDeactivated;
 
 	private:
-		void OnAssetAdded(const FAssetData& InAssetData);
-		void OnAssetRemoved(const FAssetData& InAssetData);
-		void OnFilesLoaded();
-		void OnActiveAudioDeviceChanged();
-
-		void AddAudioBusAsset(const FAssetData& InAssetData);
-		void RemoveAudioBusAsset(const FAssetData& InAssetData);
-
-		void UpdateAudioBusAssetNames();
-
 		virtual bool ProcessMessages() override;
 
-		bool bAreFilesLoaded = false;
-
-		TArray<TSharedPtr<FAudioBusAssetDashboardEntry>> AudioBusDataViewEntries;
+		FAudioBusMessages TraceMessages;
 	};
 } // namespace UE::Audio::Insights
diff --git a/Engine/Plugins/AudioInsights/Source/AudioInsightsEditor/Private/Views/AudioBusAssetDashboardEntry.cpp b/Engine/Plugins/AudioInsights/Source/AudioInsightsEditor/Private/Views/AudioBusAssetDashboardEntry.cpp
deleted file mode 100644
diff --git a/Engine/Plugins/AudioInsights/Source/AudioInsightsEditor/Private/Views/AudioBusAssetDashboardEntry.h b/Engine/Plugins/AudioInsights/Source/AudioInsightsEditor/Private/Views/AudioBusAssetDashboardEntry.h
deleted file mode 100644
diff --git a/Engine/Plugins/AudioInsights/Source/AudioInsightsEditor/Private/Views/AudioBusesDashboardViewFactory.cpp b/Engine/Plugins/AudioInsights/Source/AudioInsightsEditor/Private/Views/AudioBusesDashboardViewFactory.cpp
--- a/Engine/Plugins/AudioInsights/Source/AudioInsightsEditor/Private/Views/AudioBusesDashboardViewFactory.cpp
+++ b/Engine/Plugins/AudioInsights/Source/AudioInsightsEditor/Private/Views/AudioBusesDashboardViewFactory.cpp
@@ -1,6 +1,7 @@
 // Copyright Epic Games, Inc. All Rights Reserved.
 #include "Views/AudioBusesDashboardViewFactory.h"
 
+#include "AudioInsightsEditorModule.h"
 #include "AudioInsightsStyle.h"
 #include "Editor.h"
 #include "Internationalization/Text.h"
@@ -8,6 +9,8 @@
 #include "Sound/AudioBus.h"
 #include "Subsystems/AssetEditorSubsystem.h"
 #include "Templates/SharedPointer.h"
+#include "Widgets/Input/SCheckBox.h"
+#include "Widgets/Input/SComboBox.h"
 
 #define LOCTEXT_NAMESPACE "AudioInsights"
 
@@ -19,30 +22,30 @@ namespace UE::Audio::Insights
 		{
 			return static_cast<const FAudioBusAssetDashboardEntry&>(InData);
 		};
+
+		FAudioBusAssetDashboardEntry& CastEntry(IDashboardDataViewEntry& InData)
+		{
+			return static_cast<FAudioBusAssetDashboardEntry&>(InData);
+		};
 	}
 
 	FAudioBusesDashboardViewFactory::FAudioBusesDashboardViewFactory()
 	{
-		FAudioBusProvider::OnAudioBusAssetAdded.AddRaw(this, &FAudioBusesDashboardViewFactory::HandleOnAudioBusAssetListUpdated);
-		FAudioBusProvider::OnAudioBusAssetRemoved.AddRaw(this, &FAudioBusesDashboardViewFactory::HandleOnAudioBusAssetListUpdated);
-		FAudioBusProvider::OnAudioBusAssetListUpdated.AddRaw(this, &FAudioBusesDashboardViewFactory::RequestListRefresh);
+		IAudioInsightsTraceModule& AudioInsightsTraceModule = IAudioInsightsEditorModule::GetChecked().GetTraceModule();
+
+		const TSharedPtr<FAudioBusProvider> AudioBusProvider = MakeShared<FAudioBusProvider>();
+
+		AudioInsightsTraceModule.AddTraceProvider(AudioBusProvider);
 
 		Providers = TArray<TSharedPtr<FTraceProviderBase>>
 		{
-			MakeShared<FAudioBusProvider>()
+			AudioBusProvider
 		};
 
 		SortByColumn = "Name";
 		SortMode     = EColumnSortMode::Ascending;
 	}
 
-	FAudioBusesDashboardViewFactory::~FAudioBusesDashboardViewFactory()
-	{
-		FAudioBusProvider::OnAudioBusAssetAdded.RemoveAll(this);
-		FAudioBusProvider::OnAudioBusAssetRemoved.RemoveAll(this);
-		FAudioBusProvider::OnAudioBusAssetListUpdated.RemoveAll(this);
-	}
-
 	FName FAudioBusesDashboardViewFactory::GetName() const
 	{
 		return "AudioBuses";
@@ -55,59 +58,115 @@ namespace UE::Audio::Insights
 
 	TSharedRef<SWidget> FAudioBusesDashboardViewFactory::GenerateWidgetForColumn(TSharedRef<IDashboardDataViewEntry> InRowData, const FName& InColumnName)
 	{
-		if (InColumnName == "Name")
+		if (InColumnName == "Active")
 		{
-			const TWeakObjectPtr<UAudioBus> AudioBus = AudioBusesPrivate::CastEntry(InRowData.Get()).AudioBus;
+			static const FLinearColor DarkGreen(0.027f, 0.541f, 0.22f);
+			static const float Radius = 4.0f;
+			static const FVector2f Size(7.0f, 7.0f);
+			static const FSlateRoundedBoxBrush GreenRoundedBrush(DarkGreen, Radius, Size);
 
-			AudioBusCheckboxCheckedStates.FindOrAdd(AudioBus, false);
-
-			const FColumnData& ColumnData = GetColumns()[InColumnName];
-			const FText Label = ColumnData.GetDisplayValue(InRowData.Get());
-
-			return (Label.IsEmpty()) ? 
-				SNullWidget::NullWidget :
-				SNew(SHorizontalBox)
+			return SNew(SBox)
 				.Clipping(EWidgetClipping::ClipToBounds)
-				+ SHorizontalBox::Slot()
-				.FillWidth(0.1f)
+				.Padding(6.0f)
+				.Visibility_Lambda([InRowData]()
+				{ 
+					const FAudioBusAssetDashboardEntry& AudioBusAssetDashboardEntry = AudioBusesPrivate::CastEntry(InRowData.Get());
+					return AudioBusAssetDashboardEntry.bHasActivity ? EVisibility::Visible : EVisibility::Hidden;
+				})
 				[
-					SNew(SCheckBox)
-					.IsChecked(AudioBusCheckboxCheckedStates.Contains(AudioBus) && AudioBusCheckboxCheckedStates[AudioBus] == true)
-					.OnCheckStateChanged_Lambda([AudioBus, this](ECheckBoxState NewState)
-					{
-						if (AudioBusCheckboxCheckedStates.Contains(AudioBus))
+					SNew(SImage)
+					.Image(&GreenRoundedBrush)
+				];
+		}
+		else if (InColumnName == "Name")
+		{
+			const FColumnData& ColumnData = GetColumns()[InColumnName];
+			const FText ValueText = ColumnData.GetDisplayValue(InRowData.Get());
+
+			if (ValueText.IsEmpty())
+			{
+				return SNullWidget::NullWidget;
+			}
+
+			FAudioBusAssetDashboardEntry& AudioBusAssetDashboardEntry = AudioBusesPrivate::CastEntry(InRowData.Get());
+
+			if (AudioBusAssetDashboardEntry.EntryType == EAudioBusEntryType::AssetBased)
+			{
+				const TWeakObjectPtr<UAudioBus> AudioBus(Cast<UAudioBus>(AudioBusAssetDashboardEntry.GetObject()));
+
+				if (!AudioBus.IsValid())
+				{
+					return SNullWidget::NullWidget;
+				}
+
+				AudioBusCheckboxCheckedStates.FindOrAdd(AudioBus, false);
+
+				return SNew(SHorizontalBox)
+					.Clipping(EWidgetClipping::ClipToBounds)
+					+ SHorizontalBox::Slot()
+					.AutoWidth()
+					[
+						SNew(SCheckBox)
+						.IsChecked_Lambda([this, AudioBus]()
 						{
-							AudioBusCheckboxCheckedStates[AudioBus] = NewState == ECheckBoxState::Checked;
-						}
+							const bool* FoundAudioBusCheckedState = AudioBusCheckboxCheckedStates.Find(AudioBus);
+							return FoundAudioBusCheckedState && *FoundAudioBusCheckedState == true ? ECheckBoxState::Checked : ECheckBoxState::Unchecked;
+						})
+						.OnCheckStateChanged_Lambda([this, AudioBus](ECheckBoxState NewState)
+						{
+							if (bool* FoundAudioBusCheckedState = AudioBusCheckboxCheckedStates.Find(AudioBus))
+							{
+								*FoundAudioBusCheckedState = NewState == ECheckBoxState::Checked;
+							}
 
-						OnAudioBusAssetChecked.Broadcast(NewState == ECheckBoxState::Checked, AudioBus);
-					})
-				]
-				+ SHorizontalBox::Slot()
-				.FillWidth(0.8f)
-				[
-					SNew(STextBlock)
-					.Text(Label)
-					.MinDesiredWidth(300)
-					.OnDoubleClicked_Lambda([this, InRowData](const FGeometry& MyGeometry, const FPointerEvent& PointerEvent)
-					{
-						if (GEditor)
+							OnAudioBusItemChecked.Broadcast(NewState == ECheckBoxState::Checked, AudioBus);
+						})
+					]
+					+ SHorizontalBox::Slot()
+					.AutoWidth()
+					[
+						SNew(SBox)
+						.MinDesiredWidth(5.0f)
+					]
+					+ SHorizontalBox::Slot()
+					.AutoWidth()
+					[
+						SNew(STextBlock)
+						.Text(ValueText)
+						.MinDesiredWidth(300)
+						.OnDoubleClicked_Lambda([this, InRowData](const FGeometry& MyGeometry, const FPointerEvent& PointerEvent)
 						{
-							if (TSharedPtr<IObjectDashboardEntry> ObjectData = StaticCastSharedPtr<IObjectDashboardEntry>(InRowData.ToSharedPtr()); 
-								ObjectData.IsValid())
+							if (GEditor)
 							{
-								if (const UObject* Object = ObjectData->GetObject();
-									Object && Object->IsAsset())
+								TSharedPtr<IObjectDashboardEntry> ObjectData = StaticCastSharedPtr<IObjectDashboardEntry>(InRowData.ToSharedPtr());
+								if (ObjectData.IsValid())
 								{
-									GEditor->GetEditorSubsystem<UAssetEditorSubsystem>()->OpenEditorForAsset(Object);
-									return FReply::Handled();
+									const UObject* Object = ObjectData->GetObject();
+									if (Object && Object->IsAsset())
+									{
+										GEditor->GetEditorSubsystem<UAssetEditorSubsystem>()->OpenEditorForAsset(Object);
+										return FReply::Handled();
+									}
 								}
 							}
-						}
 
-						return FReply::Unhandled();
-					})
-				];
+							return FReply::Unhandled();
+						})
+					];
+			}
+			else if (AudioBusAssetDashboardEntry.EntryType == EAudioBusEntryType::CodeGenerated)
+			{
+				// Show only the name since we can't click the checkbox to activate the audio meter nor open with double click a code generated audio bus
+				return SNew(SHorizontalBox)
+					.Clipping(EWidgetClipping::ClipToBounds)
+					+ SHorizontalBox::Slot()
+					.AutoWidth()
+					[
+						SNew(STextBlock)
+						.Text(ValueText)
+						.MinDesiredWidth(300)
+					];
+			}
 		}
 
 		return SNullWidget::NullWidget;
@@ -115,14 +174,8 @@ namespace UE::Audio::Insights
 
 	void FAudioBusesDashboardViewFactory::ProcessEntries(FTraceTableDashboardViewFactory::EProcessReason Reason)
 	{
-		const FString FilterString = GetSearchFilterText().ToString();
-		
-		FTraceTableDashboardViewFactory::FilterEntries<FAudioBusProvider>([&FilterString](const IDashboardDataViewEntry& Entry)
-		{
-			const FAudioBusAssetDashboardEntry& AudioBusEntry = static_cast<const FAudioBusAssetDashboardEntry&>(Entry);
-			
-			return !AudioBusEntry.GetDisplayName().ToString().Contains(FilterString);
-		});
+		FilterByAudioBusName();
+		FilterByAudioBusType();
 	}
 
 	FSlateIcon FAudioBusesDashboardViewFactory::GetIcon() const
@@ -135,11 +188,89 @@ namespace UE::Audio::Insights
 		return EDefaultDashboardTabStack::Analysis;
 	}
 
+	TSharedRef<SWidget> FAudioBusesDashboardViewFactory::MakeAudioBusTypeFilterWidget()
+	{
+		if (AudioBusTypes.IsEmpty())
+		{
+			AudioBusTypes.Emplace(MakeShared<FComboboxSelectionItem>(EAudioBusTypeComboboxSelection::AssetBased,    LOCTEXT("AudioBusesDashboard_AudioBusTypeAssetBased",    "Asset")));
+			AudioBusTypes.Emplace(MakeShared<FComboboxSelectionItem>(EAudioBusTypeComboboxSelection::CodeGenerated, LOCTEXT("AudioBusesDashboard_AudioBusTypeCodeGenerated", "Code Generated")));
+			AudioBusTypes.Emplace(MakeShared<FComboboxSelectionItem>(EAudioBusTypeComboboxSelection::All,           LOCTEXT("AudioBusesDashboard_AudioBusTypeAll",           "All")));
+
+			SelectedAudioBusType = AudioBusTypes[0];
+		}
+
+		return SNew(SHorizontalBox)
+			+ SHorizontalBox::Slot()
+			.AutoWidth()
+			.Padding(2.0f, 10.0f, 0.0f, 0.0f)
+			[
+				SNew(STextBlock)
+				.Margin(FMargin(0.0, 2.0, 0.0, 0.0))
+				.Text(LOCTEXT("AudioBusesDashboard_TypeFilterText", "Type Filter:"))
+			]
+			+ SHorizontalBox::Slot()
+			.MaxWidth(2.0f)
+			.Padding(0.0f, 10.0f, 0.0f, 0.0f)
+			[
+				SNew(SBox)
+			]
+			+ SHorizontalBox::Slot()
+			.AutoWidth()
+			.HAlign(EHorizontalAlignment::HAlign_Center)
+			.Padding(0.0f, 10.0f, 0.0f, 0.0f)
+			[
+				SNew(SComboBox<TSharedPtr<FComboboxSelectionItem>>)
+				.OptionsSource(&AudioBusTypes)
+				.OnGenerateWidget_Lambda([this](const TSharedPtr<FComboboxSelectionItem>& AudioBusTypePtr)
+				{
+					const FText AudioBusTypeDisplayName = AudioBusTypePtr.IsValid() ? AudioBusTypePtr->Value /*DisplayName*/ : FText::GetEmpty();
+
+					return SNew(STextBlock)
+						.Text(AudioBusTypeDisplayName);
+				})
+				.OnSelectionChanged_Lambda([this](TSharedPtr<FComboboxSelectionItem> InSelectedAudioBusTypePtr, ESelectInfo::Type)
+				{
+					if (InSelectedAudioBusTypePtr.IsValid())
+					{
+						SelectedAudioBusType = InSelectedAudioBusTypePtr;
+						UpdateFilterReason = EProcessReason::FilterUpdated;
+					}
+				})
+				[
+					SNew(STextBlock)
+					.Text_Lambda([this]()
+					{
+						const int32 FoundIndex = AudioBusTypes.Find(SelectedAudioBusType);
+						if (AudioBusTypes.IsValidIndex(FoundIndex) && AudioBusTypes[FoundIndex].IsValid())
+						{
+							return AudioBusTypes[FoundIndex]->Value;
+						}
+
+						return FText::GetEmpty();
+					})
+				]
+			];
+	}
+
 	TSharedRef<SWidget> FAudioBusesDashboardViewFactory::MakeWidget(TSharedRef<SDockTab> OwnerTab, const FSpawnTabArgs& SpawnTabArgs)
 	{
 		if (!DashboardWidget.IsValid())
 		{
-			DashboardWidget = FTraceTableDashboardViewFactory::MakeWidget(OwnerTab, SpawnTabArgs);
+			FAudioDeviceManagerDelegates::OnAudioDeviceDestroyed.AddSP(this, &FAudioBusesDashboardViewFactory::HandleOnAudioDeviceDestroyed);
+
+			SAssignNew(DashboardWidget, SVerticalBox)
+			+ SVerticalBox::Slot()
+			.AutoHeight()
+			.HAlign(HAlign_Fill)
+			.Padding(0.0f, 0.0f, 0.0f, 6.0f)
+			[
+				MakeAudioBusTypeFilterWidget()
+			]
+			+ SVerticalBox::Slot()
+			.HAlign(HAlign_Fill)
+			[
+				FTraceTableDashboardViewFactory::MakeWidget(OwnerTab, SpawnTabArgs)
+			];
 
 			if (FilteredEntriesListView.IsValid())
 			{
@@ -156,13 +287,27 @@ namespace UE::Audio::Insights
 		{
 			return TMap<FName, FTraceTableDashboardViewFactory::FColumnData>
 			{
+				{
+					"Active",
+					{
+						LOCTEXT("AudioBuses_ActiveDisplayName", "Active"),
+						[](const IDashboardDataViewEntry& InData) { return FText::GetEmpty(); },
+						false,		/* bDefaultHidden */
+						0.08f,		/* FillWidth */
+						EHorizontalAlignment::HAlign_Center
+					}
+				},
 				{
 					"Name",
 					{
 						LOCTEXT("AudioBuses_NameColumnDisplayName", "Name"),
-						[](const IDashboardDataViewEntry& InData) { return AudioBusesPrivate::CastEntry(InData).GetDisplayName(); },
+						[](const IDashboardDataViewEntry& InData)
+						{
+							const FAudioBusAssetDashboardEntry& AudioBusAssetDashboardEntry = AudioBusesPrivate::CastEntry(InData);
+							return AudioBusAssetDashboardEntry.EntryType == EAudioBusEntryType::AssetBased ? AudioBusAssetDashboardEntry.GetDisplayName() : FText::FromString(AudioBusAssetDashboardEntry.Name);
+						},
 						false,		/* bDefaultHidden */
-						0.88f,		/* FillWidth */
+						0.92f,		/* FillWidth */
 						EHorizontalAlignment::HAlign_Left
 					}
 				}
@@ -176,7 +321,30 @@ namespace UE::Audio::Insights
 
 	void FAudioBusesDashboardViewFactory::SortTable()
 	{
-		if (SortByColumn == "Name")
+		if (SortByColumn == "Active")
+		{
+			if (SortMode == EColumnSortMode::Ascending)
+			{
+				DataViewEntries.Sort([](const TSharedPtr<IDashboardDataViewEntry>& A, const TSharedPtr<IDashboardDataViewEntry>& B)
+				{
+					const FAudioBusAssetDashboardEntry& AData = AudioBusesPrivate::CastEntry(*A.Get());
+					const FAudioBusAssetDashboardEntry& BData = AudioBusesPrivate::CastEntry(*B.Get());
+
+					return !AData.bHasActivity && BData.bHasActivity;
+				});
+			}
+			else if (SortMode == EColumnSortMode::Descending)
+			{
+				DataViewEntries.Sort([](const TSharedPtr<IDashboardDataViewEntry>& A, const TSharedPtr<IDashboardDataViewEntry>& B)
+				{
+					const FAudioBusAssetDashboardEntry& AData = AudioBusesPrivate::CastEntry(*A.Get());
+					const FAudioBusAssetDashboardEntry& BData = AudioBusesPrivate::CastEntry(*B.Get());
+
+					return !BData.bHasActivity && AData.bHasActivity;
+				});
+			}
+		}
+		else if (SortByColumn == "Name")
 		{
 			if (SortMode == EColumnSortMode::Ascending)
 			{
@@ -201,17 +369,52 @@ namespace UE::Audio::Insights
 		}
 	}
 
-	void FAudioBusesDashboardViewFactory::RequestListRefresh()
+	void FAudioBusesDashboardViewFactory::FilterByAudioBusName()
+	{
+		const FString FilterString = GetSearchFilterText().ToString();
+		
+		FTraceTableDashboardViewFactory::FilterEntries<FAudioBusProvider>([&FilterString](const IDashboardDataViewEntry& Entry)
+		{
+			const FAudioBusAssetDashboardEntry& AudioBusEntry = static_cast<const FAudioBusAssetDashboardEntry&>(Entry);
+			
+			return !AudioBusEntry.GetDisplayName().ToString().Contains(FilterString);
+		});
+	}
+
+	void FAudioBusesDashboardViewFactory::FilterByAudioBusType()
 	{
-		if (FilteredEntriesListView.IsValid())
+		using namespace UE::Audio::Insights;
+
+		TArray<TSharedPtr<IDashboardDataViewEntry>> EntriesToFilterOut;
+
+		const EAudioBusTypeComboboxSelection SelectedAudioBusTypeEnum = SelectedAudioBusType.IsValid()? SelectedAudioBusType->Key : EAudioBusTypeComboboxSelection::All;
+
+		for (const TSharedPtr<IDashboardDataViewEntry>& Entry : DataViewEntries)
+		{
+			if (Entry.IsValid())
+			{
+				const FAudioBusAssetDashboardEntry& AudioBusAssetDashboardEntry = AudioBusesPrivate::CastEntry(*Entry);
+
+				if (SelectedAudioBusTypeEnum != EAudioBusTypeComboboxSelection::All)
+				{
+					if ((SelectedAudioBusTypeEnum == EAudioBusTypeComboboxSelection::AssetBased    && AudioBusAssetDashboardEntry.EntryType != EAudioBusEntryType::AssetBased) ||
+						(SelectedAudioBusTypeEnum == EAudioBusTypeComboboxSelection::CodeGenerated && AudioBusAssetDashboardEntry.EntryType != EAudioBusEntryType::CodeGenerated))
+					{
+						EntriesToFilterOut.Emplace(Entry);
+					}
+				}
+			}
+		}
+
+		for (const TSharedPtr<IDashboardDataViewEntry>& Entry : EntriesToFilterOut)
 		{
-			FilteredEntriesListView->RequestListRefresh();
+			DataViewEntries.Remove(Entry);
 		}
 	}
 
-	void FAudioBusesDashboardViewFactory::HandleOnAudioBusAssetListUpdated(const TWeakObjectPtr<UObject> InAsset)
+	void FAudioBusesDashboardViewFactory::HandleOnAudioDeviceDestroyed(::Audio::FDeviceId InAudioDeviceId)
 	{
-		RequestListRefresh();
+		AudioBusCheckboxCheckedStates.Empty();
 	}
 } // namespace UE::Audio::Insights
 #undef LOCTEXT_NAMESPACE
diff --git a/Engine/Plugins/AudioInsights/Source/AudioInsightsEditor/Private/Views/AudioBusesDashboardViewFactory.h b/Engine/Plugins/AudioInsights/Source/AudioInsightsEditor/Private/Views/AudioBusesDashboardViewFactory.h
--- a/Engine/Plugins/AudioInsights/Source/AudioInsightsEditor/Private/Views/AudioBusesDashboardViewFactory.h
+++ b/Engine/Plugins/AudioInsights/Source/AudioInsightsEditor/Private/Views/AudioBusesDashboardViewFactory.h
@@ -3,10 +3,8 @@
 
 #include "Delegates/Delegate.h"
 #include "Templates/SharedPointer.h"
-#include "Views/AudioBusAssetDashboardEntry.h"
-#include "Views/AudioBusAssetDashboardEntry.h"
+#include "Messages/AudioBusMessages.h"
 #include "Views/TableDashboardViewFactory.h"
-#include "Widgets/Input/SCheckBox.h"
 
 class UAudioBus;
 
@@ -16,7 +14,7 @@ namespace UE::Audio::Insights
 	{
 	public:
 		FAudioBusesDashboardViewFactory();
-		virtual ~FAudioBusesDashboardViewFactory();
+		virtual ~FAudioBusesDashboardViewFactory() = default;
 
 		virtual FName GetName() const override;
 		virtual FText GetDisplayName() const override;
@@ -24,20 +22,34 @@ namespace UE::Audio::Insights
 		virtual EDefaultDashboardTabStack GetDefaultTabStack() const override;
 		virtual TSharedRef<SWidget> MakeWidget(TSharedRef<SDockTab> OwnerTab, const FSpawnTabArgs& SpawnTabArgs) override;
 
-		DECLARE_MULTICAST_DELEGATE_TwoParams(FOnAudioBusAssetChecked, const bool /*bIsChecked*/, const TWeakObjectPtr<UAudioBus> /*AudioBus*/);
-		inline static FOnAudioBusAssetChecked OnAudioBusAssetChecked;
+		DECLARE_MULTICAST_DELEGATE_TwoParams(FOnAudioBusItemChecked, const bool /*bIsChecked*/, const TWeakObjectPtr<UAudioBus> /*AudioBus*/);
+		inline static FOnAudioBusItemChecked OnAudioBusItemChecked;
+
+	private:
+		enum class EAudioBusTypeComboboxSelection : uint8
+		{
+			AssetBased,
+			CodeGenerated,
+			All
+		};
+
+		TSharedRef<SWidget> MakeAudioBusTypeFilterWidget();
 
-	protected:
 		virtual TSharedRef<SWidget> GenerateWidgetForColumn(TSharedRef<IDashboardDataViewEntry> InRowData, const FName& InColumnName) override;
 		virtual void ProcessEntries(FTraceTableDashboardViewFactory::EProcessReason Reason) override;
 		virtual const TMap<FName, FTraceTableDashboardViewFactory::FColumnData>& GetColumns() const override;
 
 		virtual void SortTable() override;
 
-		void RequestListRefresh();
+		void FilterByAudioBusName();
+		void FilterByAudioBusType();
 
-		void HandleOnAudioBusAssetListUpdated(const TWeakObjectPtr<UObject> InAsset);
+		void HandleOnAudioDeviceDestroyed(::Audio::FDeviceId InAudioDeviceId);
 
 		TMap<const TWeakObjectPtr<UAudioBus>, bool> AudioBusCheckboxCheckedStates;
+
+		using FComboboxSelectionItem = TPair<EAudioBusTypeComboboxSelection, FText>;
+		TArray<TSharedPtr<FComboboxSelectionItem>> AudioBusTypes;
+		TSharedPtr<FComboboxSelectionItem> SelectedAudioBusType;
 	};
 } // namespace UE::Audio::Insights
diff --git a/Engine/Plugins/AudioInsights/Source/AudioInsightsEditor/Private/Views/AudioMetersDashboardViewFactory.cpp b/Engine/Plugins/AudioInsights/Source/AudioInsightsEditor/Private/Views/AudioMetersDashboardViewFactory.cpp
--- a/Engine/Plugins/AudioInsights/Source/AudioInsightsEditor/Private/Views/AudioMetersDashboardViewFactory.cpp
+++ b/Engine/Plugins/AudioInsights/Source/AudioInsightsEditor/Private/Views/AudioMetersDashboardViewFactory.cpp
@@ -41,9 +41,9 @@ namespace UE::Audio::Insights
 			OnSubmixItemCheckedHandle = FSubmixesDashboardViewFactory::OnSubmixItemChecked.AddSP(this, &FAudioMetersDashboardViewFactory::HandleOnSubmixItemChecked);
 		}
 
-		if (!OnAudioBusAssetCheckedHandle.IsValid())
+		if (!OnAudioBusItemCheckedHandle.IsValid())
 		{
-			OnAudioBusAssetCheckedHandle = FAudioBusesDashboardViewFactory::OnAudioBusAssetChecked.AddSP(this, &FAudioMetersDashboardViewFactory::HandleOnAudioBusAssetChecked);
+			OnAudioBusItemCheckedHandle = FAudioBusesDashboardViewFactory::OnAudioBusItemChecked.AddSP(this, &FAudioMetersDashboardViewFactory::HandleOnAudioBusItemChecked);
 		}
 
 		if (!OnSubmixDeactivatedHandle.IsValid())
@@ -51,9 +51,14 @@ namespace UE::Audio::Insights
 			OnSubmixDeactivatedHandle = FSoundSubmixProvider::OnSubmixDeactivated.AddSP(this, &FAudioMetersDashboardViewFactory::HandleOnAssetRemoved);
 		}
 
-		if (!OnAudioBusAssetRemovedHandle.IsValid())
+		if (!OnAudioBusDeactivatedHandle.IsValid())
 		{
-			OnAudioBusAssetRemovedHandle = FAudioBusProvider::OnAudioBusAssetRemoved.AddSP(this, &FAudioMetersDashboardViewFactory::HandleOnAssetRemoved);
+			OnAudioBusDeactivatedHandle = FAudioBusProvider::OnAudioBusDeactivated.AddSP(this, &FAudioMetersDashboardViewFactory::HandleOnAssetRemoved);
+		}
+
+		if (!OnAudioDeviceDestroyedHandle.IsValid())
+		{
+			OnAudioDeviceDestroyedHandle = FAudioDeviceManagerDelegates::OnAudioDeviceDestroyed.AddSP(this, &FAudioMetersDashboardViewFactory::HandleOnAudioDeviceDestroyed);
 		}
 		
 		if (!MeterViewsScrollBox.IsValid())
@@ -104,7 +109,7 @@ namespace UE::Audio::Insights
 		}
 	}
 
-	void FAudioMetersDashboardViewFactory::HandleOnAudioBusAssetChecked(const bool bInIsChecked, const TWeakObjectPtr<UAudioBus> InAudioBus)
+	void FAudioMetersDashboardViewFactory::HandleOnAudioBusItemChecked(const bool bInIsChecked, const TWeakObjectPtr<UAudioBus> InAudioBus)
 	{
 		if (!InAudioBus.IsValid())
 		{
@@ -158,5 +163,20 @@ namespace UE::Audio::Insights
 			MeterViewsScrollBox->Invalidate(EInvalidateWidget::Layout);
 		}
 	}
+
+	void FAudioMetersDashboardViewFactory::HandleOnAudioDeviceDestroyed(::Audio::FDeviceId InAudioDeviceId)
+	{
+		for (const auto& [AssetUniqueId, AudioMeterView] : AudioMeterViews)
+		{
+			AudioMeterViewsContainer->RemoveSlot(AudioMeterView->GetWidget());
+		}
+
+		AudioMeterViews.Empty();
+
+		if (MeterViewsScrollBox.IsValid())
+		{
+			MeterViewsScrollBox->Invalidate(EInvalidateWidget::Layout);
+		}
+	}
 } // namespace UE::Audio::Insights
 #undef LOCTEXT_NAMESPACE
diff --git a/Engine/Plugins/AudioInsights/Source/AudioInsightsEditor/Private/Views/AudioMetersDashboardViewFactory.h b/Engine/Plugins/AudioInsights/Source/AudioInsightsEditor/Private/Views/AudioMetersDashboardViewFactory.h
--- a/Engine/Plugins/AudioInsights/Source/AudioInsightsEditor/Private/Views/AudioMetersDashboardViewFactory.h
+++ b/Engine/Plugins/AudioInsights/Source/AudioInsightsEditor/Private/Views/AudioMetersDashboardViewFactory.h
@@ -27,17 +27,21 @@ namespace UE::Audio::Insights
 
 	private:
 		void HandleOnSubmixItemChecked(const bool bInIsChecked, const TWeakObjectPtr<USoundSubmix> InSoundSubmix);
-		void HandleOnAudioBusAssetChecked(const bool bInIsChecked, const TWeakObjectPtr<UAudioBus> InAudioBus);
+		void HandleOnAudioBusItemChecked(const bool bInIsChecked, const TWeakObjectPtr<UAudioBus> InAudioBus);
 
 		void HandleOnAssetRemoved(const TWeakObjectPtr<UObject> InAsset);
 
+		void HandleOnAudioDeviceDestroyed(::Audio::FDeviceId InAudioDeviceId);
+
 		TSharedPtr<SScrollBox> MeterViewsScrollBox;
 		TSharedPtr<SHorizontalBox> AudioMeterViewsContainer;
 		TMap<uint32, TSharedRef<FAudioMeterView>> AudioMeterViews;
 
 		FDelegateHandle OnSubmixItemCheckedHandle;
-		FDelegateHandle OnAudioBusAssetCheckedHandle;
+		FDelegateHandle OnAudioBusItemCheckedHandle;
 		FDelegateHandle OnSubmixDeactivatedHandle;
-		FDelegateHandle OnAudioBusAssetRemovedHandle;
+		FDelegateHandle OnAudioBusDeactivatedHandle;
+
+		FDelegateHandle OnAudioDeviceDestroyedHandle;
 	};
 } // namespace UE::Audio::Insights
diff --git a/Engine/Source/Runtime/AudioMixer/Private/AudioMixerSourceManager.cpp b/Engine/Source/Runtime/AudioMixer/Private/AudioMixerSourceManager.cpp
--- a/Engine/Source/Runtime/AudioMixer/Private/AudioMixerSourceManager.cpp
+++ b/Engine/Source/Runtime/AudioMixer/Private/AudioMixerSourceManager.cpp
@@ -183,6 +183,7 @@ DEFINE_STAT(STAT_AudioMixerSourceManagerUpdate);
 DEFINE_STAT(STAT_AudioMixerSourceOutputBuffers);
 
 #if UE_AUDIO_PROFILERTRACE_ENABLED
+// Mixer Source messages
 UE_TRACE_EVENT_BEGIN(Audio, MixerSourceVolume)
 	UE_TRACE_EVENT_FIELD(uint32, DeviceId)
 	UE_TRACE_EVENT_FIELD(uint64, Timestamp)
@@ -218,6 +219,27 @@ UE_TRACE_EVENT_BEGIN(Audio, MixerSourceEnvelope)
 	UE_TRACE_EVENT_FIELD(uint32, PlayOrder)
 	UE_TRACE_EVENT_FIELD(float, Envelope)
 UE_TRACE_EVENT_END()
+
+// Audio Bus messages
+UE_TRACE_EVENT_BEGIN(Audio, AudioBusActivate)
+	UE_TRACE_EVENT_FIELD(uint32, DeviceId)
+	UE_TRACE_EVENT_FIELD(uint32, AudioBusId)
+	UE_TRACE_EVENT_FIELD(double, Timestamp)
+	UE_TRACE_EVENT_FIELD(UE::Trace::WideString, Name)
+UE_TRACE_EVENT_END()
+
+UE_TRACE_EVENT_BEGIN(Audio, AudioBusDeactivate)
+	UE_TRACE_EVENT_FIELD(uint32, DeviceId)
+	UE_TRACE_EVENT_FIELD(uint32, AudioBusId)
+	UE_TRACE_EVENT_FIELD(double, Timestamp)
+UE_TRACE_EVENT_END()
+
+UE_TRACE_EVENT_BEGIN(Audio, AudioBusHasActivity)
+	UE_TRACE_EVENT_FIELD(uint32, DeviceId)
+	UE_TRACE_EVENT_FIELD(uint32, AudioBusId)
+	UE_TRACE_EVENT_FIELD(double, Timestamp)
+	UE_TRACE_EVENT_FIELD(bool, HasActivity)
+UE_TRACE_EVENT_END()
 #endif // UE_AUDIO_PROFILERTRACE_ENABLED
 
 
@@ -240,6 +262,45 @@ namespace Audio
 	{
 		return FMath::Clamp(CommandBufferInitialCapacityCvar, 0, 10000);
 	}
+
+	bool IsAudioBufferSilent(const float* AudioBuffer, const int32 NumSamples)
+	{
+		bool bIsSilent = true;
+
+		int32 Index = 0;
+
+#if PLATFORM_ENABLE_VECTORINTRINSICS
+		const int32 SimdNum = NumSamples & 0xFFFFFFF0;
+		for (; Index < SimdNum; Index += 16)
+		{
+			const VectorRegister4x4Float Samples = VectorLoad16(&AudioBuffer[Index]);
+
+			if (VectorAnyGreaterThan(VectorAbs(Samples.val[0]), GlobalVectorConstants::SmallNumber) ||
+			    VectorAnyGreaterThan(VectorAbs(Samples.val[1]), GlobalVectorConstants::SmallNumber) ||
+			    VectorAnyGreaterThan(VectorAbs(Samples.val[2]), GlobalVectorConstants::SmallNumber) ||
+			    VectorAnyGreaterThan(VectorAbs(Samples.val[3]), GlobalVectorConstants::SmallNumber))
+			{
+				bIsSilent = false;
+				Index = INT_MAX;
+				break;
+			}
+		}
+#endif
+
+		// Finish to the end of the buffer or check each sample if vector intrinsics are disabled
+		for (; Index < NumSamples; ++Index)
+		{
+			// As soon as we hit a non-silent sample, we're not silent
+			if (FMath::Abs(AudioBuffer[Index]) > SMALL_NUMBER)
+			{
+				bIsSilent = false;
+				break;
+			}
+		}
+
+		return bIsSilent;
+	}
+
 	/*************************************************************************
 	* FMixerSourceManager
 	**************************************************************************/
@@ -1216,6 +1277,14 @@ namespace Audio
 
 		AudioBusKeys_AudioThread.Add(InAudioBusKey);
 
+#if UE_AUDIO_PROFILERTRACE_ENABLED
+		UE_TRACE_LOG(Audio, AudioBusActivate, AudioChannel)
+			<< AudioBusActivate.DeviceId(MixerDevice->DeviceID)
+			<< AudioBusActivate.AudioBusId(InAudioBusKey.ObjectId)
+			<< AudioBusActivate.Timestamp(FPlatformTime::Cycles64())
+			<< AudioBusActivate.Name(*InAudioBusName);
+#endif
+
 		AudioMixerThreadCommand([this, InAudioBusKey, InNumChannels, bInIsAutomatic]()
 		{
 			// If this audio bus id already exists, set it to not be automatic and return it
@@ -1245,6 +1314,13 @@ namespace Audio
 
 		AudioBusKeys_AudioThread.Remove(InAudioBusKey);
 
+#if UE_AUDIO_PROFILERTRACE_ENABLED
+		UE_TRACE_LOG(Audio, AudioBusDeactivate, AudioChannel)
+			<< AudioBusDeactivate.DeviceId(MixerDevice->DeviceID)
+			<< AudioBusDeactivate.AudioBusId(InAudioBusKey.ObjectId)
+			<< AudioBusDeactivate.Timestamp(FPlatformTime::Cycles64());
+#endif
+
 		AudioMixerThreadCommand([this, InAudioBusKey]()
 		{
 			TSharedPtr<FMixerAudioBus>* AudioBusPtr = AudioBuses.Find(InAudioBusKey);
@@ -2423,10 +2499,22 @@ namespace Audio
 		RenderThreadPhase = ESourceManagerRenderThreadPhase::UpdateBusses;
 
 		// Update the bus states post mixing. This flips the current/previous buffer indices.
-		for (auto& Entry : AudioBuses)
+		for (auto& [AudioBusKey, AudioBus] : AudioBuses)
 		{
-			TSharedPtr<FMixerAudioBus>& AudioBus = Entry.Value;
 			AudioBus->Update();
+			
+#if UE_AUDIO_PROFILERTRACE_ENABLED
+			const float* AudioBuffer = AudioBus->GetCurrentBusBuffer();
+			const int32 NumSamples   = AudioBus->NumChannels * AudioBus->NumFrames;
+
+			const bool bIsAudioBufferSilent = IsAudioBufferSilent(AudioBuffer, NumSamples);
+
+			UE_TRACE_LOG(Audio, AudioBusHasActivity, AudioChannel)
+				<< AudioBusHasActivity.DeviceId(MixerDevice->DeviceID)
+				<< AudioBusHasActivity.AudioBusId(AudioBusKey.ObjectId)
+				<< AudioBusHasActivity.Timestamp(FPlatformTime::Cycles64())
+				<< AudioBusHasActivity.HasActivity(!bIsAudioBufferSilent);
+#endif // UE_AUDIO_PROFILERTRACE_ENABLED
 		}
 	}
 
