[Audio Insights] - Fix for: Mixer Source Tab is extremely delayed when running Audio Insights through Unreal Insights

- AudioInsightsTimingViewExtender: Added GetCurrentDurationSeconds()
- AudioInsightsTraceModule: Added OnAnalysisStarting delegate so dashboards can know when the session is starting instead of using a custom FirstTimeStamp )based on the first source event registered)
- SAudioCurveView: Pass FScrubRangeToScreen by reference and use cached RangeToScreen.ViewInput instead of invoking the ViewRange lambda
- MixerSourceDashboardViewFactory: Use OnAnalysisStarting to set the BeginTimestamp so now timestamps are relative to the session (and not to the first source event registered) + adjusted GetViewRange to avoid plot delays in standalone

diff --git a/Engine/Plugins/AudioInsights/Source/AudioInsights/Private/AudioInsightsTimingViewExtender.cpp b/Engine/Plugins/AudioInsights/Source/AudioInsights/Private/AudioInsightsTimingViewExtender.cpp
--- a/Engine/Plugins/AudioInsights/Source/AudioInsights/Private/AudioInsightsTimingViewExtender.cpp
+++ b/Engine/Plugins/AudioInsights/Source/AudioInsights/Private/AudioInsightsTimingViewExtender.cpp
@@ -19,6 +19,16 @@ namespace UE::Audio::Insights
 	void FAudioInsightsTimingViewExtender::Tick(UE::Insights::Timing::ITimingViewSession& InSession, const TraceServices::IAnalysisSession& InAnalysisSession)
 	{
 		AnalysisSession = &InAnalysisSession;
+
+		{
+			TraceServices::FAnalysisSessionReadScope SessionReadScope(InAnalysisSession);
+			TraceDurationSeconds = InAnalysisSession.GetDurationSeconds();
+		}
+	}
+
+	double FAudioInsightsTimingViewExtender::GetCurrentDurationSeconds() const
+	{
+		return TraceDurationSeconds;
 	}
 
 	void FAudioInsightsTimingViewExtender::OnTimeMarkerChanged(UE::Insights::Timing::ETimeChangedFlags InFlags, double InTimeMarker)
diff --git a/Engine/Plugins/AudioInsights/Source/AudioInsights/Private/AudioInsightsTimingViewExtender.h b/Engine/Plugins/AudioInsights/Source/AudioInsights/Private/AudioInsightsTimingViewExtender.h
--- a/Engine/Plugins/AudioInsights/Source/AudioInsights/Private/AudioInsightsTimingViewExtender.h
+++ b/Engine/Plugins/AudioInsights/Source/AudioInsights/Private/AudioInsightsTimingViewExtender.h
@@ -22,6 +22,8 @@ namespace UE::Audio::Insights
 		virtual void OnEndSession(UE::Insights::Timing::ITimingViewSession& InSession) override;
 		virtual void Tick(UE::Insights::Timing::ITimingViewSession& InSession, const TraceServices::IAnalysisSession& InAnalysisSession) override;
 
+		double GetCurrentDurationSeconds() const;
+
 		DECLARE_MULTICAST_DELEGATE_OneParam(FOnTimingViewTimeMarkerChanged, double /*TimeMarker*/);
 		FOnTimingViewTimeMarkerChanged OnTimingViewTimeMarkerChanged;
 
@@ -29,5 +31,7 @@ namespace UE::Audio::Insights
 		void OnTimeMarkerChanged(UE::Insights::Timing::ETimeChangedFlags InFlags, double InTimeMarker);
 
 		const TraceServices::IAnalysisSession* AnalysisSession = nullptr;
+
+		double TraceDurationSeconds = 0.0;
 	};
 } // namespace UE::Audio::Insights
diff --git a/Engine/Plugins/AudioInsights/Source/AudioInsights/Private/AudioInsightsTraceModule.cpp b/Engine/Plugins/AudioInsights/Source/AudioInsights/Private/AudioInsightsTraceModule.cpp
--- a/Engine/Plugins/AudioInsights/Source/AudioInsights/Private/AudioInsightsTraceModule.cpp
+++ b/Engine/Plugins/AudioInsights/Source/AudioInsights/Private/AudioInsightsTraceModule.cpp
@@ -80,7 +80,7 @@ namespace UE::Audio::Insights
 			InSession.AddAnalyzer(Provider->ConstructAnalyzer(InSession));
 		}
 
-		FirstTimeStamp = -TNumericLimits<double>::Min();
+		OnAnalysisStarting.Broadcast(FPlatformTime::Seconds());
 	}
 
 	void FTraceModule::StartTraceAnalysis() const
diff --git a/Engine/Plugins/AudioInsights/Source/AudioInsights/Private/Providers/MixerSourceTraceProvider.cpp b/Engine/Plugins/AudioInsights/Source/AudioInsights/Private/Providers/MixerSourceTraceProvider.cpp
--- a/Engine/Plugins/AudioInsights/Source/AudioInsights/Private/Providers/MixerSourceTraceProvider.cpp
+++ b/Engine/Plugins/AudioInsights/Source/AudioInsights/Private/Providers/MixerSourceTraceProvider.cpp
@@ -462,27 +462,12 @@ namespace UE::Audio::Insights
 				Builder.RouteEvent(RouteId_Start, "Audio", "MixerSourceStart");
 				Builder.RouteEvent(RouteId_Stop, "Audio", "MixerSourceStop");
 				Builder.RouteEvent(RouteId_Volume, "Audio", "MixerSourceVolume");
-
-				bIsFirstTimestampSet = false;
 			}
 
 			virtual bool OnEvent(uint16 RouteId, EStyle Style, const FOnEventContext& Context) override
 			{
 				LLM_SCOPE_BYNAME(TEXT("Insights/FMixerSourceTraceAnalyzer"));
 
-				if (!bIsFirstTimestampSet)
-				{
-					FTraceModule& TraceModule = static_cast<FTraceModule&>(FAudioInsightsModule::GetChecked().GetTraceModule());
-
-					if (TraceModule.GetFirstTimeStamp() < 0.0)
-					{
-						const double CurrentTime = FPlatformTime::Seconds();
-						TraceModule.SetFirstTimeStamp(CurrentTime);
-
-						bIsFirstTimestampSet = true;
-					}
-				}
-
 				FMixerSourceMessages& Messages = GetProvider<FMixerSourceTraceProvider>().TraceMessages;
 				switch (RouteId)
 				{
@@ -557,8 +542,6 @@ namespace UE::Audio::Insights
 				RouteId_Volume,
 			};
 
-			bool bIsFirstTimestampSet = false;
-
 			TraceServices::IAnalysisSession& Session;
 		};
 
diff --git a/Engine/Plugins/AudioInsights/Source/AudioInsights/Private/Views/MixerSourceDashboardViewFactory.cpp b/Engine/Plugins/AudioInsights/Source/AudioInsights/Private/Views/MixerSourceDashboardViewFactory.cpp
--- a/Engine/Plugins/AudioInsights/Source/AudioInsights/Private/Views/MixerSourceDashboardViewFactory.cpp
+++ b/Engine/Plugins/AudioInsights/Source/AudioInsights/Private/Views/MixerSourceDashboardViewFactory.cpp
@@ -5,6 +5,7 @@
 #include "AudioInsightsModule.h"
 #include "AudioInsightsStyle.h"
 #include "DSP/Dsp.h"
+#include "Features/IModularFeatures.h"
 #include "Framework/Docking/LayoutService.h"
 #include "Framework/Docking/TabManager.h"
 #include "Internationalization/Text.h"
@@ -65,11 +66,22 @@ namespace UE::Audio::Insights
 
 	FMixerSourceDashboardViewFactory::FMixerSourceDashboardViewFactory()
 	{
-		const FTraceModule& TraceModule = static_cast<FTraceModule&>(FAudioInsightsModule::GetChecked().GetTraceModule());
+		FTraceModule& TraceModule = static_cast<FTraceModule&>(FAudioInsightsModule::GetChecked().GetTraceModule());
 		Providers = TArray<TSharedPtr<FTraceProviderBase>>
 		{
 			TraceModule.FindAudioTraceProvider<FMixerSourceTraceProvider>()
 		};
+
+		TraceModule.OnAnalysisStarting.AddRaw(this, &FMixerSourceDashboardViewFactory::OnAnalysisStarting);
+	}
+
+	FMixerSourceDashboardViewFactory::~FMixerSourceDashboardViewFactory()
+	{
+		if (FModuleManager::Get().IsModuleLoaded("AudioInsights") && IModularFeatures::Get().IsModularFeatureAvailable(TraceServices::ModuleFeatureName))
+		{
+			FTraceModule& TraceModule = static_cast<FTraceModule&>(FAudioInsightsModule::GetChecked().GetTraceModule());
+			TraceModule.OnAnalysisStarting.RemoveAll(this);
+		}
 	}
 
 	FName FMixerSourceDashboardViewFactory::GetName() const
@@ -401,8 +413,17 @@ namespace UE::Audio::Insights
 			PlotWidgetMetadataPerCurve->Empty();
 		}
 
-		BeginTimestamp = TNumericLimits<double>::Max();
-		CurrentTimestamp = 0;
+		BeginTimestamp   = TNumericLimits<double>::Max();
+		CurrentTimestamp = TNumericLimits<double>::Lowest();
+	}
+
+	void FMixerSourceDashboardViewFactory::OnAnalysisStarting(const double Timestamp)
+	{
+#if WITH_EDITOR
+		BeginTimestamp = Timestamp - GStartTime;
+#else
+		BeginTimestamp = 0.0;
+#endif // WITH_EDITOR
 	}
 
 #if WITH_EDITOR
@@ -520,7 +541,6 @@ namespace UE::Audio::Insights
 					
 					for (const auto& [Timestamp, Value] : TimeStampedValuesDisjointedArrayView.FirstBuffer)
 					{
-						BeginTimestamp   = FMath::Min(BeginTimestamp, Timestamp);
 						CurrentTimestamp = FMath::Max(CurrentTimestamp, Timestamp);
 
 #if WITH_EDITOR
@@ -720,27 +740,48 @@ namespace UE::Audio::Insights
 				return TRange<double>(0, MaxPlotHistorySeconds);
 			}
 
-			if (GameState == EGameState::Running)
-			{
-				const FTraceModule& TraceModule = static_cast<FTraceModule&>(FAudioInsightsModule::GetChecked().GetTraceModule());
-				const double FirstTimestamp = TraceModule.GetFirstTimeStamp();
+			double RangeUpperBound = 0.0;
 
-				const double CurrentTime = FPlatformTime::Seconds();
-				const double RelativeCurrentTime = CurrentTime - FirstTimestamp;
-				const double TimestampsDiff = RelativeCurrentTime - (CurrentTimestamp - BeginTimestamp);
+#if WITH_EDITOR
+			const double CurrentTime = FPlatformTime::Seconds() - GStartTime;
+
+			double TimestampsDiff = 0.0;
 
-				CurrentTimestamp += TimestampsDiff;
+			const bool bAnyMessageReceived = CurrentTimestamp != TNumericLimits<double>::Lowest();
+
+			if (GameState == EGameState::Running && bAnyMessageReceived)
+			{
+				const double RelativeCurrentTime = CurrentTime - BeginTimestamp;
+				TimestampsDiff = RelativeCurrentTime - (CurrentTimestamp - BeginTimestamp);
 			}
 
-#if !WITH_EDITOR
+			const double FinalCurrentTime = bAnyMessageReceived ? CurrentTimestamp : CurrentTime;
+
+			constexpr double RangeAlignmentOffset = 0.2;
+			RangeUpperBound = FinalCurrentTime - BeginTimestamp + TimestampsDiff - RangeAlignmentOffset;
+#else
 			const TSharedPtr<const FAudioInsightsComponent> AudioInsightsComponent = FAudioInsightsModule::GetChecked().GetAudioInsightsComponent();
 			if (AudioInsightsComponent.IsValid() && !AudioInsightsComponent->GetIsLiveSession())
 			{
 				return TRange<double>(CurrentTimestamp - MaxPlotHistorySeconds, CurrentTimestamp);
 			}
-#endif // !WITH_EDITOR
+			else
+			{
+				const double CurrentTime = FPlatformTime::Seconds();
+				const double DeltaTime   = CurrentTime - PreviousTime;
+
+				PreviousTime = CurrentTime;
+
+				const double TraceCurrentDurationSeconds = FAudioInsightsModule::GetChecked().GetTimingViewExtender().GetCurrentDurationSeconds();
+
+				CurrentRangeUpperBound = FMath::FInterpTo(CurrentRangeUpperBound, TraceCurrentDurationSeconds, DeltaTime, 1.0);
+
+				constexpr double RangeAlignmentOffset = 0.9;
+				RangeUpperBound = CurrentRangeUpperBound + RangeAlignmentOffset;
+			}
+#endif // WITH_EDITOR
 
-			return TRange<double>(CurrentTimestamp - MaxPlotHistorySeconds - BeginTimestamp, CurrentTimestamp - BeginTimestamp);
+			return TRange<double>(RangeUpperBound - MaxPlotHistorySeconds, RangeUpperBound);
 		};
 
 		if (PlotWidgets.IsEmpty())
diff --git a/Engine/Plugins/AudioInsights/Source/AudioInsights/Private/Views/SAudioCurveView.cpp b/Engine/Plugins/AudioInsights/Source/AudioInsights/Private/Views/SAudioCurveView.cpp
--- a/Engine/Plugins/AudioInsights/Source/AudioInsights/Private/Views/SAudioCurveView.cpp
+++ b/Engine/Plugins/AudioInsights/Source/AudioInsights/Private/Views/SAudioCurveView.cpp
@@ -287,7 +287,7 @@ float SAudioCurveView::LocalYToValue(const FVector2f AllottedLocalSize, const fl
 	return FMath::GetMappedRangeValueUnclamped(FVector2f(MarginTop, MarginBase), YDataRange, LocalY);
 }
 
-int32 SAudioCurveView::PaintGridLines(const FGeometry& AllottedGeometry, const FSlateRect& MyCullingRect, FSlateWindowElementList& OutDrawElements, int32 LayerId, const FWidgetStyle& InWidgetStyle, bool bParentEnabled, const SSimpleTimeSlider::FScrubRangeToScreen RangeToScreen) const
+int32 SAudioCurveView::PaintGridLines(const FGeometry& AllottedGeometry, const FSlateRect& MyCullingRect, FSlateWindowElementList& OutDrawElements, int32 LayerId, const FWidgetStyle& InWidgetStyle, bool bParentEnabled, const SSimpleTimeSlider::FScrubRangeToScreen& RangeToScreen) const
 {
 	const FVector2f Size = AllottedGeometry.GetLocalSize();
 	const int32 GridLineLayer = LayerId++;
@@ -299,8 +299,8 @@ int32 SAudioCurveView::PaintGridLines(const FGeometry& AllottedGeometry, const F
 	{
 		const double Factor = 1.0 / HorizontalAxisIncrement.Get();
 		// Start at rounded nearest HorizontalAxisIncrement from the lower bound and increment by HorizontalAxisIncrement (ex. 1.5, 2.0, 2.5... if increment is 0.5)
-		double VerticalLineValue = FMath::RoundToDouble(ViewRange.Get().GetLowerBoundValue() * Factor) / Factor;
-		while (VerticalLineValue < ViewRange.Get().GetUpperBoundValue())
+		double VerticalLineValue = FMath::RoundToDouble(RangeToScreen.ViewInput.GetLowerBoundValue() * Factor) / Factor;
+		while (VerticalLineValue < RangeToScreen.ViewInput.GetUpperBoundValue())
 		{
 			const float WidgetX = RangeToScreen.InputToLocalX(VerticalLineValue);
 			GridPoints[0].X = WidgetX;
diff --git a/Engine/Plugins/AudioInsights/Source/AudioInsights/Public/AudioInsightsTraceModule.h b/Engine/Plugins/AudioInsights/Source/AudioInsights/Public/AudioInsightsTraceModule.h
--- a/Engine/Plugins/AudioInsights/Source/AudioInsights/Public/AudioInsightsTraceModule.h
+++ b/Engine/Plugins/AudioInsights/Source/AudioInsights/Public/AudioInsightsTraceModule.h
@@ -37,8 +37,8 @@ namespace UE::Audio::Insights
 		bool IsTraceAnalysisActive() const override;
 		void StopTraceAnalysis() const override;
 
-		void SetFirstTimeStamp(double InFirstTimeStamp) { FirstTimeStamp = InFirstTimeStamp; }
-		double GetFirstTimeStamp() const { return FirstTimeStamp; }
+		DECLARE_MULTICAST_DELEGATE_OneParam(FOnAnalysisStarting, const double /*Timestamp*/);
+		FOnAnalysisStarting OnAnalysisStarting;
 
 	private:
 		static const FName GetName();
@@ -48,7 +48,5 @@ namespace UE::Audio::Insights
 
 		TSharedRef<FTraceChannelManager> ChannelManager;
 		TMap<FName, TSharedPtr<FTraceProviderBase>> TraceProviders;
-
-		double FirstTimeStamp = -TNumericLimits<double>::Min();
 	};
 } // namespace UE::Audio::Insights
diff --git a/Engine/Plugins/AudioInsights/Source/AudioInsights/Public/Views/MixerSourceDashboardViewFactory.h b/Engine/Plugins/AudioInsights/Source/AudioInsights/Public/Views/MixerSourceDashboardViewFactory.h
--- a/Engine/Plugins/AudioInsights/Source/AudioInsights/Public/Views/MixerSourceDashboardViewFactory.h
+++ b/Engine/Plugins/AudioInsights/Source/AudioInsights/Public/Views/MixerSourceDashboardViewFactory.h
@@ -19,7 +19,7 @@ namespace UE::Audio::Insights
 			const FNumberFormattingOptions* FormatOptions;
 		};
 
-		virtual ~FMixerSourceDashboardViewFactory() = default;
+		virtual ~FMixerSourceDashboardViewFactory();
 
 		virtual FName GetName() const override;
 		virtual FText GetDisplayName() const override;
@@ -49,6 +49,8 @@ namespace UE::Audio::Insights
 		using FPointDataPerCurveMap = TMap<int32, TArray<FPlotCurvePoint>>; // Map of source id to data point array 
 		using FPlotCurveMetadata = SAudioCurveView::FCurveMetadata;
 
+		void OnAnalysisStarting(const double Timestamp);
+
 #if WITH_EDITOR
 		void OnPIEStarted(bool bSimulating);
 		void OnPIEStopped(bool bSimulating);
@@ -124,7 +126,7 @@ namespace UE::Audio::Insights
 		EGameState GameState = EGameState::Stopped;
 
 		double BeginTimestamp = TNumericLimits<double>::Max();
-		double CurrentTimestamp = 0.0;
+		double CurrentTimestamp = TNumericLimits<double>::Lowest();
 
 		const static int32 NumPlotWidgets = 1;
 		TArray<FName> SelectedPlotColumnNames;
@@ -140,6 +142,9 @@ namespace UE::Audio::Insights
 		ECheckBoxState MuteState = ECheckBoxState::Unchecked;
 		ECheckBoxState SoloState = ECheckBoxState::Unchecked;
 		FString CurrentFilterString;
+#else
+		double PreviousTime = 0.0;
+		double CurrentRangeUpperBound = 0.0;
 #endif // WITH_EDITOR
 	};
 } // namespace UE::Audio::Insights
diff --git a/Engine/Plugins/AudioInsights/Source/AudioInsights/Public/Views/SAudioCurveView.h b/Engine/Plugins/AudioInsights/Source/AudioInsights/Public/Views/SAudioCurveView.h
--- a/Engine/Plugins/AudioInsights/Source/AudioInsights/Public/Views/SAudioCurveView.h
+++ b/Engine/Plugins/AudioInsights/Source/AudioInsights/Public/Views/SAudioCurveView.h
@@ -80,7 +80,7 @@ protected:
 	virtual FReply OnMouseMove(const FGeometry& MyGeometry, const FPointerEvent& MouseEvent) override;
 	
 	int32 PaintCurves(const FGeometry& AllottedGeometry, const FSlateRect& MyCullingRect, FSlateWindowElementList& OutDrawElements, int32 LayerId, const FWidgetStyle& InWidgetStyle, bool bParentEnabled) const;
-	int32 PaintGridLines(const FGeometry& AllottedGeometry, const FSlateRect& MyCullingRect, FSlateWindowElementList& OutDrawElements, int32 LayerId, const FWidgetStyle& InWidgetStyle, bool bParentEnabled, const SSimpleTimeSlider::FScrubRangeToScreen RangeToScreen) const;
+	int32 PaintGridLines(const FGeometry& AllottedGeometry, const FSlateRect& MyCullingRect, FSlateWindowElementList& OutDrawElements, int32 LayerId, const FWidgetStyle& InWidgetStyle, bool bParentEnabled, const SSimpleTimeSlider::FScrubRangeToScreen& RangeToScreen) const;
 	int32 PaintYAxisLabels(const FGeometry& AllottedGeometry, FSlateWindowElementList& OutDrawElements, int32 LayerId) const;
 	void UpdateCurveToolTip(const FGeometry& MyGeometry, const FPointerEvent& MouseEvent);
 	void UpdateYDataRange();
