[Audio Insights] - Audio Insights will end Unreal Insights trace after stopping PIE

- We now only Start/Stop a trace session during PIE if one wasn't already in progress. This is to stop Audio Insights from cancelling trace sessions started by the user (fixing the original bug).
- Original options "Start with PIE" and "Stop with PIE" removed - if Audio Insights is open, we always want to enable the relevant traces (Audio, Audio Mixer and CPU) so that Audio Insights is functional.
- Added a new option "Only trace audio channels during PIE:" to the Audio Insights window. This provides an easy way for users to reduce the size of their trace sessions when using Audio Insights. Note: we only disable channels that are not read-only (i.e. non-command line channels).
- We now cache all enabled traces when starting an Audio Insights trace. This allows us to restore the user's trace channel settings after using Audio Insights.

diff --git a/Engine/Plugins/AudioInsights/Source/AudioInsights/Private/AudioInsightsTraceModule.cpp b/Engine/Plugins/AudioInsights/Source/AudioInsights/Private/AudioInsightsTraceModule.cpp
--- a/Engine/Plugins/AudioInsights/Source/AudioInsights/Private/AudioInsightsTraceModule.cpp
+++ b/Engine/Plugins/AudioInsights/Source/AudioInsights/Private/AudioInsightsTraceModule.cpp
@@ -12,6 +12,13 @@
 
 namespace UE::Audio::Insights
 {
+	namespace FTraceModulePrivate
+	{
+		static const TCHAR* CPUChannelName = TEXT("CpuChannel");
+		static const TCHAR* AudioChannelName = TEXT("AudioChannel");
+		static const TCHAR* AudioMixerChannelName = TEXT("AudioMixerChannel");
+	}
+
 	FTraceModule::FTraceModule()
 		: ChannelManager(MakeShared<FTraceChannelManager>())
 	{
@@ -50,22 +57,66 @@ namespace UE::Audio::Insights
 
 	void FTraceModule::DisableAllTraceChannels()
 	{
-		UE::Trace::EnumerateChannels([](const ANSICHAR* ChannelName, bool bEnabled, void*)
+		UE::Trace::EnumerateChannels([](const UE::Trace::FChannelInfo& ChannelInfo, void* User)
+		{
+			// Only disable channels that are not read only (i.e. set from the command line)
+			if (!ChannelInfo.bIsReadOnly && ChannelInfo.bIsEnabled)
 			{
-				if (bEnabled)
-				{
-					FString ChannelNameFString(ChannelName);
-					UE::Trace::ToggleChannel(ChannelNameFString.GetCharArray().GetData(), false);
-				}
+				UE::Trace::ToggleChannel(ANSI_TO_TCHAR(ChannelInfo.Name), false);
 			}
+			return true;
+		}
 		, nullptr);
 	}
 
 	void FTraceModule::EnableAudioInsightsTraceChannels()
 	{
-		UE::Trace::ToggleChannel(TEXT("Cpu"), true);
-		UE::Trace::ToggleChannel(TEXT("Audio"), true);
-		UE::Trace::ToggleChannel(TEXT("AudioMixer"), true);
+		using namespace FTraceModulePrivate;
+
+		UE::Trace::ToggleChannel(CPUChannelName, true);
+		UE::Trace::ToggleChannel(AudioChannelName, true);
+		UE::Trace::ToggleChannel(AudioMixerChannelName, true);
+	}
+
+	void FTraceModule::DisableAudioInsightsTraceChannels() const
+	{
+		using namespace FTraceModulePrivate;
+
+		auto DisableIfNotMarkedToRestore = [this](const TCHAR* ChannelName)
+		{
+			if (!ChannelsToRestore.Contains(ChannelName))
+			{
+				UE::Trace::ToggleChannel(ChannelName, false);
+			}
+		};
+
+		DisableIfNotMarkedToRestore(CPUChannelName);
+		DisableIfNotMarkedToRestore(AudioChannelName);
+		DisableIfNotMarkedToRestore(AudioMixerChannelName);
+	}
+
+	void FTraceModule::CacheCurrentlyEnabledTraceChannels()
+	{	
+		ChannelsToRestore.Empty();
+
+		UE::Trace::EnumerateChannels([](const UE::Trace::FChannelInfo& ChannelInfo, void* User)
+		{
+			TArray<FString>* EnabledChannels = static_cast<TArray<FString>*>(User);
+			if (!ChannelInfo.bIsReadOnly && ChannelInfo.bIsEnabled)
+			{
+				EnabledChannels->Emplace(ANSI_TO_TCHAR(ChannelInfo.Name));
+			}
+
+			return true;
+		}, &ChannelsToRestore);
+	}
+
+	void FTraceModule::RestoreCachedChannels() const
+	{
+		for (const FString& Channel : ChannelsToRestore)
+		{
+			UE::Trace::ToggleChannel(Channel.GetCharArray().GetData(), true);
+		}
 	}
 
 	void FTraceModule::OnAnalysisBegin(TraceServices::IAnalysisSession& InSession)
@@ -83,13 +134,28 @@ namespace UE::Audio::Insights
 		OnAnalysisStarting.Broadcast(FPlatformTime::Seconds());
 	}
 
-	void FTraceModule::StartTraceAnalysis() const
+	void FTraceModule::StartTraceAnalysis(const bool bOnlyTraceAudioChannels)
 	{
-		if (!FTraceAuxiliary::IsConnected())
+		if (bAudioTracesAreAcive)
+		{
+			return;
+		}
+
+		// We don't want to stop any traces that were started prior to running Audio Insights
+		// Check the original state here
+		bStopTracingAfterAudioInsightsIsFinished = !FTraceAuxiliary::IsConnected();
+		
+		CacheCurrentlyEnabledTraceChannels();
+
+		if (bOnlyTraceAudioChannels)
 		{
 			DisableAllTraceChannels();
-			EnableAudioInsightsTraceChannels();
+		}
+
+		EnableAudioInsightsTraceChannels();
 
+		if (!FTraceAuxiliary::IsConnected())
+		{
 			// Clear all buffered data and prevent data from previous recordings from leaking into the new recording
 			FTraceAuxiliary::FOptions Options;
 			Options.bExcludeTail = true;
@@ -99,19 +165,51 @@ namespace UE::Audio::Insights
 			IUnrealInsightsModule& UnrealInsightsModule = FModuleManager::LoadModuleChecked<IUnrealInsightsModule>("TraceInsights");
 			UnrealInsightsModule.StartAnalysisForLastLiveSession();
 		}
+
+		bAudioTracesAreAcive = true;
 	}
 
-	void FTraceModule::StopTraceAnalysis() const
+	void FTraceModule::StopTraceAnalysis()
 	{
 		if (FTraceAuxiliary::IsConnected())
 		{
-			FTraceAuxiliary::Stop();
+			DisableAudioInsightsTraceChannels();
+
+			if (bStopTracingAfterAudioInsightsIsFinished)
+			{
+				FTraceAuxiliary::Stop();
+			}
+
+			RestoreCachedChannels();
+		}
+
+		ChannelsToRestore.Empty();
+		bAudioTracesAreAcive = false;
+	}
+
+	void FTraceModule::OnOnlyTraceAudioChannelsStateChanged(const bool bOnlyTraceAudioChannels)
+	{
+		if (!bAudioTracesAreAcive)
+		{
+			return;
+		}
+
+		if (bOnlyTraceAudioChannels)
+		{
+			// Re-cache the current settings for enabled channels. This may have changed since Audio Insights began.
+			CacheCurrentlyEnabledTraceChannels();
+			DisableAllTraceChannels();
+			EnableAudioInsightsTraceChannels();
+		}
+		else
+		{
+			RestoreCachedChannels();
 		}
 	}
 
 	bool FTraceModule::IsTraceAnalysisActive() const
 	{
-		return FTraceAuxiliary::IsConnected();
+		return FTraceAuxiliary::IsConnected() && bAudioTracesAreAcive;
 	}
 
 	void FTraceModule::GetLoggers(TArray<const TCHAR *>& OutLoggers)
diff --git a/Engine/Plugins/AudioInsights/Source/AudioInsights/Public/AudioInsightsTraceModule.h b/Engine/Plugins/AudioInsights/Source/AudioInsights/Public/AudioInsightsTraceModule.h
--- a/Engine/Plugins/AudioInsights/Source/AudioInsights/Public/AudioInsightsTraceModule.h
+++ b/Engine/Plugins/AudioInsights/Source/AudioInsights/Public/AudioInsightsTraceModule.h
@@ -33,9 +33,10 @@ namespace UE::Audio::Insights
 
 		virtual TSharedRef<FTraceChannelManager> GetChannelManager() override;
 
-		void StartTraceAnalysis() const override;
-		bool IsTraceAnalysisActive() const override;
-		void StopTraceAnalysis() const override;
+		virtual void StartTraceAnalysis(const bool bOnlyTraceAudioChannels) override;
+		virtual bool IsTraceAnalysisActive() const override;
+		virtual void StopTraceAnalysis() override;
+		virtual void OnOnlyTraceAudioChannelsStateChanged(const bool bOnlyTraceAudioChannels) override;
 
 		DECLARE_MULTICAST_DELEGATE_OneParam(FOnAnalysisStarting, const double /*Timestamp*/);
 		FOnAnalysisStarting OnAnalysisStarting;
@@ -46,7 +47,16 @@ namespace UE::Audio::Insights
 		static void DisableAllTraceChannels();
 		static void EnableAudioInsightsTraceChannels();
 
+		void DisableAudioInsightsTraceChannels() const;
+
+		void CacheCurrentlyEnabledTraceChannels();
+		void RestoreCachedChannels() const;
+
 		TSharedRef<FTraceChannelManager> ChannelManager;
 		TMap<FName, TSharedPtr<FTraceProviderBase>> TraceProviders;
+		TArray<FString> ChannelsToRestore;
+
+		bool bAudioTracesAreAcive = false;
+		bool bStopTracingAfterAudioInsightsIsFinished = false;
 	};
 } // namespace UE::Audio::Insights
diff --git a/Engine/Plugins/AudioInsights/Source/AudioInsights/Public/IAudioInsightsTraceModule.h b/Engine/Plugins/AudioInsights/Source/AudioInsights/Public/IAudioInsightsTraceModule.h
--- a/Engine/Plugins/AudioInsights/Source/AudioInsights/Public/IAudioInsightsTraceModule.h
+++ b/Engine/Plugins/AudioInsights/Source/AudioInsights/Public/IAudioInsightsTraceModule.h
@@ -15,9 +15,11 @@ class AUDIOINSIGHTS_API IAudioInsightsTraceModule : public TraceServices::IModul
 public:
 	virtual void AddTraceProvider(TSharedPtr<UE::Audio::Insights::FTraceProviderBase> TraceProvider) = 0;
 
-	virtual void StartTraceAnalysis() const = 0;
+	virtual void StartTraceAnalysis(const bool bOnlyTraceAudioChannels) = 0;
 	virtual bool IsTraceAnalysisActive() const = 0;
-	virtual void StopTraceAnalysis() const = 0;
+	virtual void StopTraceAnalysis() = 0;
+	
+	virtual void OnOnlyTraceAudioChannelsStateChanged(const bool bOnlyTraceAudioChannels) = 0;
 
 	virtual TSharedRef<UE::Audio::Insights::FTraceChannelManager> GetChannelManager() = 0;
 };
diff --git a/Engine/Plugins/AudioInsights/Source/AudioInsightsEditor/Private/AudioInsightsEditorDashboardFactory.cpp b/Engine/Plugins/AudioInsights/Source/AudioInsightsEditor/Private/AudioInsightsEditorDashboardFactory.cpp
--- a/Engine/Plugins/AudioInsights/Source/AudioInsightsEditor/Private/AudioInsightsEditorDashboardFactory.cpp
+++ b/Engine/Plugins/AudioInsights/Source/AudioInsightsEditor/Private/AudioInsightsEditorDashboardFactory.cpp
@@ -33,6 +33,9 @@ namespace UE::Audio::Insights
 		static const FText PreviewDeviceDisplayName = LOCTEXT("AudioDashboard_PreviewDevice", "[Preview Audio]");
 		static const FText DashboardWorldSelectDescription = LOCTEXT("AudioDashboard_SelectWorldDescription", "Select world(s) to monitor (worlds may share audio output).");
 
+		static const FText OnlyTraceAudioChannelsName = LOCTEXT("AudioDashboard_OnlyTraceAudioChannelsDisplayName", "Only trace audio channels during PIE:");
+		static const FText OnlyTraceAudioChannelsDescription = LOCTEXT("AudioDashboard_OnlyTraceAudioChannelsDescription", "Disable all non-command line trace channels apart from Audio, Audio Mixer and CPU during PIE. This will reduce the file sizes of trace sessions while using Audio Insights.");
+
 		FText GetDebugNameFromDeviceId(::Audio::FDeviceId InDeviceId)
 		{
 			FString WorldName;
@@ -62,10 +65,10 @@ namespace UE::Audio::Insights
 	{
 		if (InDeviceId != INDEX_NONE)
 		{
-			if (bStartWithPIE && InWorld && InWorld->WorldType == EWorldType::PIE)
+			if (InWorld && InWorld->WorldType == EWorldType::PIE)
 			{
-				const IAudioInsightsTraceModule& TraceModule = FAudioInsightsEditorModule::GetChecked().GetTraceModule();
-				TraceModule.StartTraceAnalysis();
+				IAudioInsightsTraceModule& TraceModule = FAudioInsightsEditorModule::GetChecked().GetTraceModule();
+				TraceModule.StartTraceAnalysis(bOnlyTraceAudioChannels);
 
 				const TObjectPtr<const UAudioInsightsEditorSettings> AudioInsightsEditorSettings = GetDefault<UAudioInsightsEditorSettings>();
 
@@ -89,20 +92,14 @@ namespace UE::Audio::Insights
 
 	void FEditorDashboardFactory::OnPIEStarted(bool bSimulating)
 	{
-		if (bStartWithPIE)
-		{
-			const IAudioInsightsTraceModule& TraceModule = FAudioInsightsEditorModule::GetChecked().GetTraceModule();
-			TraceModule.StartTraceAnalysis();
-		}
+		IAudioInsightsTraceModule& TraceModule = FAudioInsightsEditorModule::GetChecked().GetTraceModule();
+		TraceModule.StartTraceAnalysis(bOnlyTraceAudioChannels);
 	}
 
 	void FEditorDashboardFactory::OnPIEStopped(bool bSimulating)
 	{
-		if (bStopWithPIE)
-		{
-			const IAudioInsightsTraceModule& TraceModule = FAudioInsightsEditorModule::GetChecked().GetTraceModule();
-			TraceModule.StopTraceAnalysis();
-		}
+		IAudioInsightsTraceModule& TraceModule = FAudioInsightsEditorModule::GetChecked().GetTraceModule();
+		TraceModule.StopTraceAnalysis();
 
 		RefreshDeviceSelector();
 	}
@@ -446,7 +443,7 @@ namespace UE::Audio::Insights
 				.Padding(2.0f, 0.0f)
 				[
 					SNew(STextBlock)
-					.Text(LOCTEXT("StartOnPIE_DisplayName", "Start with PIE:"))
+					.Text(OnlyTraceAudioChannelsName)
 					.Font(IPropertyTypeCustomizationUtils::GetRegularFont())
 				]
 				+ SHorizontalBox::Slot()
@@ -456,28 +453,15 @@ namespace UE::Audio::Insights
 				.Padding(2.0f, 0.0f)
 				[
 					SNew(SCheckBox)
-					.IsChecked_Lambda([this]() { return bStartWithPIE ? ECheckBoxState::Checked : ECheckBoxState::Unchecked; })
-					.OnCheckStateChanged_Lambda([this](ECheckBoxState NewState) { bStartWithPIE = NewState == ECheckBoxState::Checked; })
-				]
-				+ SHorizontalBox::Slot()
-				.HAlign(HAlign_Right)
-				.VAlign(VAlign_Center)
-				.AutoWidth()
-				.Padding(2.0f, 0.0f)
-				[
-					SNew(STextBlock)
-					.Text(LOCTEXT("StopOnPIE_DisplayName", "Stop with PIE:"))
-					.Font(IPropertyTypeCustomizationUtils::GetRegularFont())
-				]
-				+ SHorizontalBox::Slot()
-				.HAlign(HAlign_Left)
-				.VAlign(VAlign_Center)
-				.AutoWidth()
-				.Padding(2.0f, 0.0f)
-				[
-					SNew(SCheckBox)
-					.IsChecked_Lambda([this]() { return bStopWithPIE ? ECheckBoxState::Checked : ECheckBoxState::Unchecked; })
-					.OnCheckStateChanged_Lambda([this](ECheckBoxState NewState) { bStopWithPIE = NewState == ECheckBoxState::Checked; })
+					.ToolTipText(OnlyTraceAudioChannelsDescription)
+					.IsChecked_Lambda([this]() { return bOnlyTraceAudioChannels ? ECheckBoxState::Checked : ECheckBoxState::Unchecked; })
+					.OnCheckStateChanged_Lambda([this](ECheckBoxState NewState) 
+					{ 
+						bOnlyTraceAudioChannels = NewState == ECheckBoxState::Checked;
+
+						IAudioInsightsTraceModule& TraceModule = FAudioInsightsEditorModule::GetChecked().GetTraceModule();
+						TraceModule.OnOnlyTraceAudioChannelsStateChanged(bOnlyTraceAudioChannels);
+					})
 				]
 				+ SHorizontalBox::Slot()
 				.HAlign(HAlign_Right)
diff --git a/Engine/Plugins/AudioInsights/Source/AudioInsightsEditor/Private/AudioInsightsEditorDashboardFactory.h b/Engine/Plugins/AudioInsights/Source/AudioInsightsEditor/Private/AudioInsightsEditorDashboardFactory.h
--- a/Engine/Plugins/AudioInsights/Source/AudioInsightsEditor/Private/AudioInsightsEditorDashboardFactory.h
+++ b/Engine/Plugins/AudioInsights/Source/AudioInsightsEditor/Private/AudioInsightsEditorDashboardFactory.h
@@ -62,8 +62,7 @@ namespace UE::Audio::Insights
 		FDelegateHandle OnPIEStartedHandle;
 		FDelegateHandle OnPIEStoppedHandle;
 
-		bool bStartWithPIE = true;
-		bool bStopWithPIE = true;
+		bool bOnlyTraceAudioChannels = false;
 
 		TSharedPtr<FTabManager> DashboardTabManager;
 		TSharedPtr<FWorkspaceItem> DashboardWorkspace;
