[Audio Insights] Submixes Dashboard: adjusted submixes dashboard with a hybrid approach so it is a mix of asset based and trace based.

diff --git a/Engine/Plugins/AudioInsights/Source/AudioInsightsEditor/Private/Messages/SoundSubmixMessages.h b/Engine/Plugins/AudioInsights/Source/AudioInsightsEditor/Private/Messages/SoundSubmixMessages.h
--- a/Engine/Plugins/AudioInsights/Source/AudioInsightsEditor/Private/Messages/SoundSubmixMessages.h
+++ b/Engine/Plugins/AudioInsights/Source/AudioInsightsEditor/Private/Messages/SoundSubmixMessages.h
@@ -7,6 +7,8 @@
 #include "Trace/Analyzer.h"
 #include "Views/TableDashboardViewFactory.h"
 
+class USoundSubmix;
+
 namespace UE::Audio::Insights
 {
 	struct FSoundSubmixMessageBase
@@ -27,23 +29,6 @@ namespace UE::Audio::Insights
 		double Timestamp = 0.0;
 	};
 
-	struct FSoundSubmixActivateMessage : public FSoundSubmixMessageBase
-	{
-		FSoundSubmixActivateMessage() = default;
-
-		FSoundSubmixActivateMessage(const UE::Trace::IAnalyzer::FOnEventContext& InContext)
-			: FSoundSubmixMessageBase(InContext)
-		{
-			const UE::Trace::IAnalyzer::FEventData& EventData = InContext.EventData;
-
-			EventData.GetString("Name", Name);
-		}
-
-		FString Name;
-	};
-
-	using FSoundSubmixDeactivateMessage = FSoundSubmixMessageBase;
-
 	struct FSoundSubmixHasActivityMessage : public FSoundSubmixMessageBase
 	{
 		FSoundSubmixHasActivityMessage() = default;
@@ -61,8 +46,6 @@ namespace UE::Audio::Insights
 
 	class FSoundSubmixMessages
 	{
-		UE::Audio::Insights::TAnalyzerMessageQueue<FSoundSubmixActivateMessage> ActivateMessages{ 0.1 };
-		UE::Audio::Insights::TAnalyzerMessageQueue<FSoundSubmixDeactivateMessage> DeactivateMessages{ 0.1 };
 		UE::Audio::Insights::TAnalyzerMessageQueue<FSoundSubmixHasActivityMessage> HasActivityMessages{ 0.1 };
 
 		friend class FSoundSubmixProvider;
@@ -73,6 +56,7 @@ namespace UE::Audio::Insights
 		FSoundSubmixAssetDashboardEntry() = default;
 		virtual ~FSoundSubmixAssetDashboardEntry() = default;
 
+		TWeakObjectPtr<USoundSubmix> SoundSubmix;
 		uint32 SubmixId = INDEX_NONE;
 		bool bHasActivity = false;
 	};
diff --git a/Engine/Plugins/AudioInsights/Source/AudioInsightsEditor/Private/Providers/SoundSubmixProvider.cpp b/Engine/Plugins/AudioInsights/Source/AudioInsightsEditor/Private/Providers/SoundSubmixProvider.cpp
--- a/Engine/Plugins/AudioInsights/Source/AudioInsightsEditor/Private/Providers/SoundSubmixProvider.cpp
+++ b/Engine/Plugins/AudioInsights/Source/AudioInsightsEditor/Private/Providers/SoundSubmixProvider.cpp
@@ -13,6 +13,25 @@ namespace UE::Audio::Insights
 	FSoundSubmixProvider::FSoundSubmixProvider()
 		: TDeviceDataMapTraceProvider<uint32, TSharedPtr<FSoundSubmixAssetDashboardEntry>>(GetName_Static())
 	{
+		FAssetRegistryModule& AssetRegistryModule = FModuleManager::LoadModuleChecked<FAssetRegistryModule>("AssetRegistry");
+
+		AssetRegistryModule.Get().OnAssetAdded().AddRaw(this, &FSoundSubmixProvider::OnAssetAdded);
+		AssetRegistryModule.Get().OnAssetRemoved().AddRaw(this, &FSoundSubmixProvider::OnAssetRemoved);
+		AssetRegistryModule.Get().OnFilesLoaded().AddRaw(this, &FSoundSubmixProvider::OnFilesLoaded);
+
+		FEditorDashboardFactory::OnActiveAudioDeviceChanged.AddRaw(this, &FSoundSubmixProvider::OnActiveAudioDeviceChanged);
+	}
+	
+	FSoundSubmixProvider::~FSoundSubmixProvider()
+	{
+		if (FAssetRegistryModule* AssetRegistryModule = FModuleManager::GetModulePtr<FAssetRegistryModule>("AssetRegistry"))
+		{
+			AssetRegistryModule->Get().OnAssetAdded().RemoveAll(this);
+			AssetRegistryModule->Get().OnAssetRemoved().RemoveAll(this);
+			AssetRegistryModule->Get().OnFilesLoaded().RemoveAll(this);
+		}
+
+		FEditorDashboardFactory::OnActiveAudioDeviceChanged.RemoveAll(this);
 	}
 
 	FName FSoundSubmixProvider::GetName_Static()
@@ -20,30 +39,113 @@ namespace UE::Audio::Insights
 		return "SubmixesProvider";
 	}
 
-	bool FSoundSubmixProvider::ProcessMessages()
+	void FSoundSubmixProvider::OnAssetAdded(const FAssetData& InAssetData)
 	{
-		// Helper lambdas
-		auto BumpEntryFunc = [this](const FSoundSubmixMessageBase& Msg)
+		if (bAreFilesLoaded && InAssetData.AssetClassPath == FTopLevelAssetPath(USoundSubmix::StaticClass()))
+		{
+			AddSubmixAsset(InAssetData);
+		}
+	}
+
+	void FSoundSubmixProvider::OnAssetRemoved(const FAssetData& InAssetData)
+	{
+		if (InAssetData.AssetClassPath == FTopLevelAssetPath(USoundSubmix::StaticClass()))
 		{
-			TSharedPtr<FSoundSubmixAssetDashboardEntry>* ToReturn = nullptr;
+			RemoveSubmixAsset(InAssetData);
+		}
+	}
 
-			UpdateDeviceEntry(Msg.DeviceId, Msg.SubmixId, [&ToReturn, &Msg](TSharedPtr<FSoundSubmixAssetDashboardEntry>& Entry)
+	void FSoundSubmixProvider::OnFilesLoaded()
+	{
+		bAreFilesLoaded = true;
+		UpdateSubmixAssetNames();
+	}
+
+	void FSoundSubmixProvider::OnActiveAudioDeviceChanged()
+	{
+		UpdateSubmixAssetNames();
+	}
+
+	void FSoundSubmixProvider::AddSubmixAsset(const FAssetData& InAssetData)
+	{
+		const bool bIsSubmixAssetAlreadAdded = SubmixDataViewEntries.ContainsByPredicate(
+			[AssetName = InAssetData.GetObjectPathString()](const TSharedPtr<FSoundSubmixAssetDashboardEntry>& SoundSubmixAssetDashboardEntry)
 			{
-				if (!Entry.IsValid())
-				{
-					Entry = MakeShared<FSoundSubmixAssetDashboardEntry>();
-					Entry->DeviceId = Msg.DeviceId;
-					Entry->SubmixId = Msg.SubmixId;
-				}
+				return SoundSubmixAssetDashboardEntry->Name == AssetName;
+			});
+
+		if (!bIsSubmixAssetAlreadAdded)
+		{
+			const FAudioInsightsEditorModule AudioInsightsEditorModule = FAudioInsightsEditorModule::GetChecked();
+			const ::Audio::FDeviceId AudioDeviceId = AudioInsightsEditorModule.GetDeviceId();
+
+			TSharedPtr<FSoundSubmixAssetDashboardEntry> SoundSubmixAssetDashboardEntry = MakeShared<FSoundSubmixAssetDashboardEntry>();
+			SoundSubmixAssetDashboardEntry->DeviceId    = AudioDeviceId;
+			SoundSubmixAssetDashboardEntry->Name        = InAssetData.GetObjectPathString();
+			SoundSubmixAssetDashboardEntry->SoundSubmix = Cast<USoundSubmix>(InAssetData.GetAsset());
+			SoundSubmixAssetDashboardEntry->SubmixId    = SoundSubmixAssetDashboardEntry->SoundSubmix->GetUniqueID();
+
+			SubmixDataViewEntries.Add(MoveTemp(SoundSubmixAssetDashboardEntry));
 
-				Entry->Timestamp = Msg.Timestamp;
+			OnSubmixAssetAdded.Broadcast(InAssetData.GetAsset());
 
-				ToReturn = &Entry;
+			bAssetsUpdated = true;
+			++LastUpdateId;
+		}
+	}
+
+	void FSoundSubmixProvider::RemoveSubmixAsset(const FAssetData& InAssetData)
+	{
+		const int32 FoundSubmixAssetNameIndex = SubmixDataViewEntries.IndexOfByPredicate(
+			[AssetName = InAssetData.GetObjectPathString()](const TSharedPtr<FSoundSubmixAssetDashboardEntry>& SoundSubmixAssetDashboardEntry)
+			{
+				return SoundSubmixAssetDashboardEntry->Name == AssetName;
 			});
 
-			return ToReturn;
-		};
+		if (FoundSubmixAssetNameIndex != INDEX_NONE)
+		{
+			const FAudioInsightsEditorModule AudioInsightsEditorModule = FAudioInsightsEditorModule::GetChecked();
+			const ::Audio::FDeviceId AudioDeviceId = AudioInsightsEditorModule.GetDeviceId();
+
+			RemoveDeviceEntry(AudioDeviceId, SubmixDataViewEntries[FoundSubmixAssetNameIndex]->SubmixId);
+
+			SubmixDataViewEntries.RemoveAt(FoundSubmixAssetNameIndex);
 
+			OnSubmixAssetRemoved.Broadcast(InAssetData.GetAsset());
+
+			bAssetsUpdated = true;
+			++LastUpdateId;
+		}
+	}
+
+	void FSoundSubmixProvider::UpdateSubmixAssetNames()
+	{
+		// Get all USoundSubmix assets
+		TArray<FAssetData> AssetDataArray;
+
+		FAssetRegistryModule& AssetRegistryModule = FModuleManager::LoadModuleChecked<FAssetRegistryModule>("AssetRegistry");
+		AssetRegistryModule.Get().GetAssetsByClass(FTopLevelAssetPath(USoundSubmix::StaticClass()), AssetDataArray);
+
+		// Build SubmixDataViewEntries
+		Reset();
+		SubmixDataViewEntries.Empty();
+
+		for (const FAssetData& AssetData : AssetDataArray)
+		{
+			AddSubmixAsset(AssetData);
+		}
+
+		SubmixDataViewEntries.Sort([](const TSharedPtr<FSoundSubmixAssetDashboardEntry>& A, const TSharedPtr<FSoundSubmixAssetDashboardEntry>& B)
+		{
+			return A->GetDisplayName().CompareToCaseIgnored(B->GetDisplayName()) < 0;
+		});
+
+		OnSubmixAssetListUpdated.Broadcast();
+	}
+
+	bool FSoundSubmixProvider::ProcessMessages()
+	{
+		// Helper lambdas
 		auto GetEntry = [this](const FSoundSubmixMessageBase& Msg)
 		{
 			return FindDeviceEntry(Msg.DeviceId, Msg.SubmixId);
@@ -51,25 +153,27 @@ namespace UE::Audio::Insights
 
 
 		// Process messages
-		ProcessMessageQueue<FSoundSubmixActivateMessage>(TraceMessages.ActivateMessages, BumpEntryFunc,
-		[](const FSoundSubmixActivateMessage& Msg, TSharedPtr<FSoundSubmixAssetDashboardEntry>* OutEntry)
+		if (bAssetsUpdated)
 		{
-			FSoundSubmixAssetDashboardEntry& EntryRef = *OutEntry->Get();
-			EntryRef.Name = *Msg.Name;
-		});
+			const FAudioInsightsEditorModule AudioInsightsEditorModule = FAudioInsightsEditorModule::GetChecked();
+			const ::Audio::FDeviceId AudioDeviceId = AudioInsightsEditorModule.GetDeviceId();
 
-		ProcessMessageQueue<FSoundSubmixDeactivateMessage>(TraceMessages.DeactivateMessages, GetEntry,
-		[this](const FSoundSubmixDeactivateMessage& Msg, TSharedPtr<FSoundSubmixAssetDashboardEntry>* OutEntry)
-		{
-			if (OutEntry && (*OutEntry)->Timestamp < Msg.Timestamp)
+			for (const TSharedPtr<FSoundSubmixAssetDashboardEntry>& SubmixDataViewEntry : SubmixDataViewEntries)
 			{
-				FSoundSubmixAssetDashboardEntry& EntryRef = *OutEntry->Get();
-
-				OnSubmixDeactivated.Broadcast(EntryRef.GetObject());
-
-				RemoveDeviceEntry(Msg.DeviceId, Msg.SubmixId);
+				if (SubmixDataViewEntry.IsValid() && SubmixDataViewEntry->SoundSubmix.IsValid())
+				{
+					UpdateDeviceEntry(AudioDeviceId, SubmixDataViewEntry->SubmixId, [&SubmixDataViewEntry](TSharedPtr<FSoundSubmixAssetDashboardEntry>& Entry)
+					{
+						if (!Entry.IsValid())
+						{
+							Entry = SubmixDataViewEntry;
+						}
+					});
+				}
 			}
-		});
+
+			bAssetsUpdated = false;
+		}
 
 		ProcessMessageQueue<FSoundSubmixHasActivityMessage>(TraceMessages.HasActivityMessages, GetEntry,
 		[](const FSoundSubmixHasActivityMessage& Msg, TSharedPtr<FSoundSubmixAssetDashboardEntry>* OutEntry)
@@ -102,8 +206,6 @@ namespace UE::Audio::Insights
 
 				UE::Trace::IAnalyzer::FInterfaceBuilder& Builder = Context.InterfaceBuilder;
 
-				Builder.RouteEvent(RouteId_Activate,    "Audio", "SoundSubmixActivate");
-				Builder.RouteEvent(RouteId_Deactivate,  "Audio", "SoundSubmixDeactivate");
 				Builder.RouteEvent(RouteId_HasActivity, "Audio", "SoundSubmixHasActivity");
 			}
 
@@ -115,18 +217,6 @@ namespace UE::Audio::Insights
 
 				switch (RouteId)
 				{
-					case RouteId_Activate:
-					{
-						Messages.ActivateMessages.Enqueue(FSoundSubmixActivateMessage { Context });
-						break;
-					}
-
-					case RouteId_Deactivate:
-					{
-						Messages.DeactivateMessages.Enqueue(FSoundSubmixDeactivateMessage{ Context });
-						break;
-					}
-
 					case RouteId_HasActivity:
 					{
 						Messages.HasActivityMessages.Enqueue(FSoundSubmixHasActivityMessage{ Context });
@@ -152,8 +242,6 @@ namespace UE::Audio::Insights
 		private:
 			enum : uint16
 			{
-				RouteId_Activate,
-				RouteId_Deactivate,
 				RouteId_HasActivity
 			};
 
diff --git a/Engine/Plugins/AudioInsights/Source/AudioInsightsEditor/Private/Providers/SoundSubmixProvider.h b/Engine/Plugins/AudioInsights/Source/AudioInsightsEditor/Private/Providers/SoundSubmixProvider.h
--- a/Engine/Plugins/AudioInsights/Source/AudioInsightsEditor/Private/Providers/SoundSubmixProvider.h
+++ b/Engine/Plugins/AudioInsights/Source/AudioInsightsEditor/Private/Providers/SoundSubmixProvider.h
@@ -2,8 +2,8 @@
 #pragma once
 
 #include "Delegates/Delegate.h"
-#include "UObject/NameTypes.h"
 #include "Messages/SoundSubmixMessages.h"
+#include "UObject/NameTypes.h"
 
 namespace UE::Audio::Insights
 {
@@ -11,18 +11,39 @@ namespace UE::Audio::Insights
 	{
 	public:
 		FSoundSubmixProvider();
-		virtual ~FSoundSubmixProvider() = default;
+		virtual ~FSoundSubmixProvider();
 
 		static FName GetName_Static();
 
 		virtual UE::Trace::IAnalyzer* ConstructAnalyzer(TraceServices::IAnalysisSession& InSession) override;
+		
+		DECLARE_MULTICAST_DELEGATE_OneParam(FOnSubmixAssetAdded, const TWeakObjectPtr<UObject> /*Asset*/);
+		inline static FOnSubmixAssetAdded OnSubmixAssetAdded;
 
-		DECLARE_MULTICAST_DELEGATE_OneParam(FOnSubmixDeactivated, const TWeakObjectPtr<UObject> /*Asset*/);
-		inline static FOnSubmixDeactivated OnSubmixDeactivated;
+		DECLARE_MULTICAST_DELEGATE_OneParam(FOnSubmixAssetRemoved, const TWeakObjectPtr<UObject> /*Asset*/);
+		inline static FOnSubmixAssetRemoved OnSubmixAssetRemoved;
+
+		DECLARE_MULTICAST_DELEGATE(FOnSubmixAssetListUpdated);
+		inline static FOnSubmixAssetListUpdated OnSubmixAssetListUpdated;
 
 	private:
+		void OnAssetAdded(const FAssetData& InAssetData);
+		void OnAssetRemoved(const FAssetData& InAssetData);
+		void OnFilesLoaded();
+		void OnActiveAudioDeviceChanged();
+
+		void AddSubmixAsset(const FAssetData& InAssetData);
+		void RemoveSubmixAsset(const FAssetData& InAssetData);
+
+		void UpdateSubmixAssetNames();
+
 		virtual bool ProcessMessages() override;
 
+		bool bAreFilesLoaded = false;
+		bool bAssetsUpdated  = false;
+
+		TArray<TSharedPtr<FSoundSubmixAssetDashboardEntry>> SubmixDataViewEntries;
+
 		FSoundSubmixMessages TraceMessages;
 	};
 } // namespace UE::Audio::Insights
diff --git a/Engine/Plugins/AudioInsights/Source/AudioInsightsEditor/Private/Views/AudioMetersDashboardViewFactory.cpp b/Engine/Plugins/AudioInsights/Source/AudioInsightsEditor/Private/Views/AudioMetersDashboardViewFactory.cpp
--- a/Engine/Plugins/AudioInsights/Source/AudioInsightsEditor/Private/Views/AudioMetersDashboardViewFactory.cpp
+++ b/Engine/Plugins/AudioInsights/Source/AudioInsightsEditor/Private/Views/AudioMetersDashboardViewFactory.cpp
@@ -36,9 +36,9 @@ namespace UE::Audio::Insights
 
 	TSharedRef<SWidget> FAudioMetersDashboardViewFactory::MakeWidget(TSharedRef<SDockTab> OwnerTab, const FSpawnTabArgs& SpawnTabArgs)
 	{
-		if (!OnSubmixItemCheckedHandle.IsValid())
+		if (!OnSubmixAssetCheckedHandle.IsValid())
 		{
-			OnSubmixItemCheckedHandle = FSubmixesDashboardViewFactory::OnSubmixItemChecked.AddSP(this, &FAudioMetersDashboardViewFactory::HandleOnSubmixItemChecked);
+			OnSubmixAssetCheckedHandle = FSubmixesDashboardViewFactory::OnSubmixAssetChecked.AddSP(this, &FAudioMetersDashboardViewFactory::HandleOnSubmixAssetChecked);
 		}
 
 		if (!OnAudioBusItemCheckedHandle.IsValid())
@@ -46,20 +46,15 @@ namespace UE::Audio::Insights
 			OnAudioBusItemCheckedHandle = FAudioBusesDashboardViewFactory::OnAudioBusItemChecked.AddSP(this, &FAudioMetersDashboardViewFactory::HandleOnAudioBusItemChecked);
 		}
 
-		if (!OnSubmixDeactivatedHandle.IsValid())
+		if (!OnSubmixAssetRemovedHandle.IsValid())
 		{
-			OnSubmixDeactivatedHandle = FSoundSubmixProvider::OnSubmixDeactivated.AddSP(this, &FAudioMetersDashboardViewFactory::HandleOnAssetRemoved);
+			OnSubmixAssetRemovedHandle = FSoundSubmixProvider::OnSubmixAssetRemoved.AddSP(this, &FAudioMetersDashboardViewFactory::HandleOnAssetRemoved);
 		}
 
 		if (!OnAudioBusDeactivatedHandle.IsValid())
 		{
 			OnAudioBusDeactivatedHandle = FAudioBusProvider::OnAudioBusDeactivated.AddSP(this, &FAudioMetersDashboardViewFactory::HandleOnAssetRemoved);
 		}
-
-		if (!OnAudioDeviceDestroyedHandle.IsValid())
-		{
-			OnAudioDeviceDestroyedHandle = FAudioDeviceManagerDelegates::OnAudioDeviceDestroyed.AddSP(this, &FAudioMetersDashboardViewFactory::HandleOnAudioDeviceDestroyed);
-		}
 		
 		if (!MeterViewsScrollBox.IsValid())
 		{
@@ -74,7 +69,7 @@ namespace UE::Audio::Insights
 		return MeterViewsScrollBox.ToSharedRef();
 	}
 
-	void FAudioMetersDashboardViewFactory::HandleOnSubmixItemChecked(const bool bInIsChecked, const TWeakObjectPtr<USoundSubmix> InSoundSubmix)
+	void FAudioMetersDashboardViewFactory::HandleOnSubmixAssetChecked(const bool bInIsChecked, const TWeakObjectPtr<USoundSubmix> InSoundSubmix)
 	{
 		if (!InSoundSubmix.IsValid())
 		{
@@ -163,20 +158,5 @@ namespace UE::Audio::Insights
 			MeterViewsScrollBox->Invalidate(EInvalidateWidget::Layout);
 		}
 	}
-
-	void FAudioMetersDashboardViewFactory::HandleOnAudioDeviceDestroyed(::Audio::FDeviceId InAudioDeviceId)
-	{
-		for (const auto& [AssetUniqueId, AudioMeterView] : AudioMeterViews)
-		{
-			AudioMeterViewsContainer->RemoveSlot(AudioMeterView->GetWidget());
-		}
-
-		AudioMeterViews.Empty();
-
-		if (MeterViewsScrollBox.IsValid())
-		{
-			MeterViewsScrollBox->Invalidate(EInvalidateWidget::Layout);
-		}
-	}
 } // namespace UE::Audio::Insights
 #undef LOCTEXT_NAMESPACE
diff --git a/Engine/Plugins/AudioInsights/Source/AudioInsightsEditor/Private/Views/AudioMetersDashboardViewFactory.h b/Engine/Plugins/AudioInsights/Source/AudioInsightsEditor/Private/Views/AudioMetersDashboardViewFactory.h
--- a/Engine/Plugins/AudioInsights/Source/AudioInsightsEditor/Private/Views/AudioMetersDashboardViewFactory.h
+++ b/Engine/Plugins/AudioInsights/Source/AudioInsightsEditor/Private/Views/AudioMetersDashboardViewFactory.h
@@ -26,22 +26,18 @@ namespace UE::Audio::Insights
 		virtual TSharedRef<SWidget> MakeWidget(TSharedRef<SDockTab> OwnerTab, const FSpawnTabArgs& SpawnTabArgs) override;
 
 	private:
-		void HandleOnSubmixItemChecked(const bool bInIsChecked, const TWeakObjectPtr<USoundSubmix> InSoundSubmix);
+		void HandleOnSubmixAssetChecked(const bool bInIsChecked, const TWeakObjectPtr<USoundSubmix> InSoundSubmix);
 		void HandleOnAudioBusItemChecked(const bool bInIsChecked, const TWeakObjectPtr<UAudioBus> InAudioBus);
 
 		void HandleOnAssetRemoved(const TWeakObjectPtr<UObject> InAsset);
 
-		void HandleOnAudioDeviceDestroyed(::Audio::FDeviceId InAudioDeviceId);
-
 		TSharedPtr<SScrollBox> MeterViewsScrollBox;
 		TSharedPtr<SHorizontalBox> AudioMeterViewsContainer;
 		TMap<uint32, TSharedRef<FAudioMeterView>> AudioMeterViews;
 
-		FDelegateHandle OnSubmixItemCheckedHandle;
+		FDelegateHandle OnSubmixAssetCheckedHandle;
 		FDelegateHandle OnAudioBusItemCheckedHandle;
-		FDelegateHandle OnSubmixDeactivatedHandle;
+		FDelegateHandle OnSubmixAssetRemovedHandle;
 		FDelegateHandle OnAudioBusDeactivatedHandle;
-
-		FDelegateHandle OnAudioDeviceDestroyedHandle;
 	};
 } // namespace UE::Audio::Insights
diff --git a/Engine/Plugins/AudioInsights/Source/AudioInsightsEditor/Private/Views/SubmixesDashboardViewFactory.cpp b/Engine/Plugins/AudioInsights/Source/AudioInsightsEditor/Private/Views/SubmixesDashboardViewFactory.cpp
--- a/Engine/Plugins/AudioInsights/Source/AudioInsightsEditor/Private/Views/SubmixesDashboardViewFactory.cpp
+++ b/Engine/Plugins/AudioInsights/Source/AudioInsightsEditor/Private/Views/SubmixesDashboardViewFactory.cpp
@@ -20,15 +20,14 @@ namespace UE::Audio::Insights
 		{
 			return static_cast<const FSoundSubmixAssetDashboardEntry&>(InData);
 		};
-
-		FSoundSubmixAssetDashboardEntry& CastEntry(IDashboardDataViewEntry& InData)
-		{
-			return static_cast<FSoundSubmixAssetDashboardEntry&>(InData);
-		};
 	}
 
 	FSubmixesDashboardViewFactory::FSubmixesDashboardViewFactory()
 	{
+		FSoundSubmixProvider::OnSubmixAssetAdded.AddRaw(this, &FSubmixesDashboardViewFactory::HandleOnSubmixAssetListUpdated);
+		FSoundSubmixProvider::OnSubmixAssetRemoved.AddRaw(this, &FSubmixesDashboardViewFactory::HandleOnSubmixAssetListUpdated);
+		FSoundSubmixProvider::OnSubmixAssetListUpdated.AddRaw(this, &FSubmixesDashboardViewFactory::RequestListRefresh);
+
 		IAudioInsightsTraceModule& AudioInsightsTraceModule = IAudioInsightsEditorModule::GetChecked().GetTraceModule();
 
 		const TSharedPtr<FSoundSubmixProvider> SoundSubmixProvider = MakeShared<FSoundSubmixProvider>();
@@ -44,6 +43,13 @@ namespace UE::Audio::Insights
 		SortMode     = EColumnSortMode::Ascending;
 	}
 
+	FSubmixesDashboardViewFactory::~FSubmixesDashboardViewFactory()
+	{
+		FSoundSubmixProvider::OnSubmixAssetAdded.RemoveAll(this);
+		FSoundSubmixProvider::OnSubmixAssetRemoved.RemoveAll(this);
+		FSoundSubmixProvider::OnSubmixAssetListUpdated.RemoveAll(this);
+	}
+
 	FName FSubmixesDashboardViewFactory::GetName() const
 	{
 		return "Submixes";
@@ -86,8 +92,7 @@ namespace UE::Audio::Insights
 				return SNullWidget::NullWidget;
 			}
 
-			FSoundSubmixAssetDashboardEntry& SoundSubmixAssetDashboardEntry = SubmixesPrivate::CastEntry(InRowData.Get());
-			const TWeakObjectPtr<USoundSubmix> SoundSubmix(Cast<USoundSubmix>(SoundSubmixAssetDashboardEntry.GetObject()));
+			const TWeakObjectPtr<USoundSubmix> SoundSubmix = SubmixesPrivate::CastEntry(InRowData.Get()).SoundSubmix;
 
 			if (!SoundSubmix.IsValid())
 			{
@@ -105,7 +110,9 @@ namespace UE::Audio::Insights
 					.IsChecked_Lambda([this, SoundSubmix]()
 					{
 						const bool* FoundSubmixCheckedState = SubmixCheckboxCheckedStates.Find(SoundSubmix);
-						return FoundSubmixCheckedState && *FoundSubmixCheckedState == true ? ECheckBoxState::Checked : ECheckBoxState::Unchecked;
+						const bool bIsSubmixChecked = FoundSubmixCheckedState && *FoundSubmixCheckedState == true;
+
+						return bIsSubmixChecked ? ECheckBoxState::Checked : ECheckBoxState::Unchecked;
 					})
 					.OnCheckStateChanged_Lambda([this, SoundSubmix](ECheckBoxState NewState)
 					{
@@ -114,7 +121,7 @@ namespace UE::Audio::Insights
 							*FoundSubmixCheckedState = NewState == ECheckBoxState::Checked;
 						}
 
-						OnSubmixItemChecked.Broadcast(NewState == ECheckBoxState::Checked, SoundSubmix);
+						OnSubmixAssetChecked.Broadcast(NewState == ECheckBoxState::Checked, SoundSubmix);
 					})
 				]
 				+ SHorizontalBox::Slot()
@@ -133,10 +140,10 @@ namespace UE::Audio::Insights
 					{
 						if (GEditor)
 						{
-							TSharedPtr<IObjectDashboardEntry> ObjectData = StaticCastSharedPtr<IObjectDashboardEntry>(InRowData.ToSharedPtr());
+							const TSharedPtr<IObjectDashboardEntry> ObjectData = StaticCastSharedPtr<IObjectDashboardEntry>(InRowData.ToSharedPtr());
 							if (ObjectData.IsValid())
 							{
-								const UObject* Object = ObjectData->GetObject();
+								const TObjectPtr<UObject> Object = ObjectData->GetObject();
 								if (Object && Object->IsAsset())
 								{
 									GEditor->GetEditorSubsystem<UAssetEditorSubsystem>()->OpenEditorForAsset(Object);
@@ -179,8 +186,6 @@ namespace UE::Audio::Insights
 	{
 		if (!DashboardWidget.IsValid())
 		{
-			FAudioDeviceManagerDelegates::OnAudioDeviceDestroyed.AddSP(this, &FSubmixesDashboardViewFactory::HandleOnAudioDeviceDestroyed);
-
 			DashboardWidget = FTraceTableDashboardViewFactory::MakeWidget(OwnerTab, SpawnTabArgs);
 
 			if (FilteredEntriesListView.IsValid())
@@ -280,19 +285,23 @@ namespace UE::Audio::Insights
 	{
 		if (SelectedItem.IsValid())
 		{
-			FSoundSubmixAssetDashboardEntry& SoundSubmixAssetDashboardEntry = SubmixesPrivate::CastEntry(*SelectedItem.Get());
+			const FSoundSubmixAssetDashboardEntry& SoundSubmixAssetDashboardEntry = SubmixesPrivate::CastEntry(*SelectedItem.Get());
 
-			const TWeakObjectPtr<USoundSubmix> SoundSubmixWeakObjectPtr(Cast<USoundSubmix>(SoundSubmixAssetDashboardEntry.GetObject()));
-			if (SoundSubmixWeakObjectPtr.IsValid())
-			{
-				OnSubmixSelectionChanged.Broadcast(SoundSubmixWeakObjectPtr);
-			}
+			OnSubmixSelectionChanged.Broadcast(SoundSubmixAssetDashboardEntry.SoundSubmix);
+		}
+	}
+
+	void FSubmixesDashboardViewFactory::RequestListRefresh()
+	{
+		if (FilteredEntriesListView.IsValid())
+		{
+			FilteredEntriesListView->RequestListRefresh();
 		}
 	}
 
-	void FSubmixesDashboardViewFactory::HandleOnAudioDeviceDestroyed(::Audio::FDeviceId InAudioDeviceId)
+	void FSubmixesDashboardViewFactory::HandleOnSubmixAssetListUpdated(const TWeakObjectPtr<UObject> InAsset)
 	{
-		SubmixCheckboxCheckedStates.Empty();
+		RequestListRefresh();
 	}
 } // namespace UE::Audio::Insights
 #undef LOCTEXT_NAMESPACE
diff --git a/Engine/Plugins/AudioInsights/Source/AudioInsightsEditor/Private/Views/SubmixesDashboardViewFactory.h b/Engine/Plugins/AudioInsights/Source/AudioInsightsEditor/Private/Views/SubmixesDashboardViewFactory.h
--- a/Engine/Plugins/AudioInsights/Source/AudioInsightsEditor/Private/Views/SubmixesDashboardViewFactory.h
+++ b/Engine/Plugins/AudioInsights/Source/AudioInsightsEditor/Private/Views/SubmixesDashboardViewFactory.h
@@ -3,7 +3,6 @@
 
 #include "Delegates/Delegate.h"
 #include "Templates/SharedPointer.h"
-#include "Messages/SoundSubmixMessages.h"
 #include "Views/TableDashboardViewFactory.h"
 #include "Widgets/Input/SCheckBox.h"
 
@@ -15,7 +14,7 @@ namespace UE::Audio::Insights
 	{
 	public:
 		FSubmixesDashboardViewFactory();
-		virtual ~FSubmixesDashboardViewFactory() = default;
+		virtual ~FSubmixesDashboardViewFactory();
 
 		virtual FName GetName() const override;
 		virtual FText GetDisplayName() const override;
@@ -23,8 +22,8 @@ namespace UE::Audio::Insights
 		virtual EDefaultDashboardTabStack GetDefaultTabStack() const override;
 		virtual TSharedRef<SWidget> MakeWidget(TSharedRef<SDockTab> OwnerTab, const FSpawnTabArgs& SpawnTabArgs) override;
 
-		DECLARE_MULTICAST_DELEGATE_TwoParams(FOnSubmixItemChecked, const bool /*bIsChecked*/, TWeakObjectPtr<USoundSubmix> /*SoundSubmix*/);
-		inline static FOnSubmixItemChecked OnSubmixItemChecked;
+		DECLARE_MULTICAST_DELEGATE_TwoParams(FOnSubmixAssetChecked, const bool /*bIsChecked*/, TWeakObjectPtr<USoundSubmix> /*SoundSubmix*/);
+		inline static FOnSubmixAssetChecked OnSubmixAssetChecked;
 
 		DECLARE_MULTICAST_DELEGATE_OneParam(FOnSubmixSelectionChanged, const TWeakObjectPtr<USoundSubmix> /*SoundSubmix*/);
 		inline static FOnSubmixSelectionChanged OnSubmixSelectionChanged;
@@ -38,7 +37,9 @@ namespace UE::Audio::Insights
 
 		virtual void OnSelectionChanged(TSharedPtr<IDashboardDataViewEntry> SelectedItem, ESelectInfo::Type SelectInfo) override;
 
-		void HandleOnAudioDeviceDestroyed(::Audio::FDeviceId InAudioDeviceId);
+		void RequestListRefresh();
+
+		void HandleOnSubmixAssetListUpdated(const TWeakObjectPtr<UObject> InAsset);
 
 		TMap<const TWeakObjectPtr<USoundSubmix>, bool> SubmixCheckboxCheckedStates;
 	};
diff --git a/Engine/Source/Runtime/AudioMixer/Private/AudioMixerSubmix.cpp b/Engine/Source/Runtime/AudioMixer/Private/AudioMixerSubmix.cpp
--- a/Engine/Source/Runtime/AudioMixer/Private/AudioMixerSubmix.cpp
+++ b/Engine/Source/Runtime/AudioMixer/Private/AudioMixerSubmix.cpp
@@ -59,19 +59,6 @@ DEFINE_STAT(STAT_AudioMixerSubmixSoundfieldProcessors);
 
 #if UE_AUDIO_PROFILERTRACE_ENABLED
 
-UE_TRACE_EVENT_BEGIN(Audio, SoundSubmixActivate)
-	UE_TRACE_EVENT_FIELD(uint32, DeviceId)
-	UE_TRACE_EVENT_FIELD(uint32, SubmixId)
-	UE_TRACE_EVENT_FIELD(double, Timestamp)
-	UE_TRACE_EVENT_FIELD(UE::Trace::WideString, Name)
-UE_TRACE_EVENT_END()
-
-UE_TRACE_EVENT_BEGIN(Audio, SoundSubmixDeactivate)
-	UE_TRACE_EVENT_FIELD(uint32, DeviceId)
-	UE_TRACE_EVENT_FIELD(uint32, SubmixId)
-	UE_TRACE_EVENT_FIELD(double, Timestamp)
-UE_TRACE_EVENT_END()
-
 UE_TRACE_EVENT_BEGIN(Audio, SoundSubmixHasActivity)
 	UE_TRACE_EVENT_FIELD(uint32, DeviceId)
 	UE_TRACE_EVENT_FIELD(uint32, SubmixId)
@@ -239,16 +226,6 @@ namespace Audio
 				MutableSubmix->StopRecordingOutput(MixerDevice, EAudioRecordingExportType::WavFile, InterruptedFileName, FString());
 			}
 		}
-
-#if UE_AUDIO_PROFILERTRACE_ENABLED
-		if (MixerDevice && OwningSubmixObject.IsValid())
-		{
-			UE_TRACE_LOG(Audio, SoundSubmixDeactivate, AudioChannel)
-				<< SoundSubmixDeactivate.DeviceId(MixerDevice->DeviceID)
-				<< SoundSubmixDeactivate.SubmixId(OwningSubmixObject->GetUniqueID())
-				<< SoundSubmixDeactivate.Timestamp(FPlatformTime::Cycles64());
-		}
-#endif
 	}
 
 	void FMixerSubmix::Init(const USoundSubmixBase* InSoundSubmix, bool bAllowReInit)
@@ -281,17 +258,6 @@ namespace Audio
 				TArray<USoundfieldEffectBase*> Effects = SoundfieldSubmix->GetSoundfieldProcessors();
 				SetupSoundfieldStreams(EncodingSettings, Effects, SoundfieldFactory);
 			}
-
-#if UE_AUDIO_PROFILERTRACE_ENABLED
-			if (MixerDevice)
-			{
-				UE_TRACE_LOG(Audio, SoundSubmixActivate, AudioChannel)
-					<< SoundSubmixActivate.DeviceId(MixerDevice->DeviceID)
-					<< SoundSubmixActivate.SubmixId(OwningSubmixObject->GetUniqueID())
-					<< SoundSubmixActivate.Timestamp(FPlatformTime::Cycles64())
-					<< SoundSubmixActivate.Name(*OwningSubmixObject->GetPathName());
-			}
-#endif
 		}
 	}
 
