[Audio Insights] AudioBuses Dashboard: adjusted audio buses dashboard with a hybrid approach so it is a mix of asset based and trace based + added colorization for active audio buses during runtime.

diff --git a/Engine/Plugins/AudioInsights/Source/AudioInsightsEditor/Private/Messages/AudioBusMessages.h b/Engine/Plugins/AudioInsights/Source/AudioInsightsEditor/Private/Messages/AudioBusMessages.h
--- a/Engine/Plugins/AudioInsights/Source/AudioInsightsEditor/Private/Messages/AudioBusMessages.h
+++ b/Engine/Plugins/AudioInsights/Source/AudioInsightsEditor/Private/Messages/AudioBusMessages.h
@@ -2,11 +2,12 @@
 #pragma once
 
 #include "AudioDefines.h"
-#include "Math/NumericLimits.h"
 #include "Messages/AnalyzerMessageQueue.h"
 #include "Trace/Analyzer.h"
 #include "Views/TableDashboardViewFactory.h"
 
+class UAudioBus;
+
 namespace UE::Audio::Insights
 {
 	enum class EAudioBusEntryType : uint8
@@ -82,7 +83,9 @@ namespace UE::Audio::Insights
 
 		EAudioBusEntryType EntryType = EAudioBusEntryType::None;
 
+		TWeakObjectPtr<UAudioBus> AudioBus;
 		uint32 AudioBusId = INDEX_NONE;
+		bool bIsRegistered = false;
 		bool bHasActivity = false;
 	};
 } // namespace UE::Audio::Insights
diff --git a/Engine/Plugins/AudioInsights/Source/AudioInsightsEditor/Private/Providers/AudioBusProvider.cpp b/Engine/Plugins/AudioInsights/Source/AudioInsightsEditor/Private/Providers/AudioBusProvider.cpp
--- a/Engine/Plugins/AudioInsights/Source/AudioInsightsEditor/Private/Providers/AudioBusProvider.cpp
+++ b/Engine/Plugins/AudioInsights/Source/AudioInsightsEditor/Private/Providers/AudioBusProvider.cpp
@@ -13,6 +13,25 @@ namespace UE::Audio::Insights
 	FAudioBusProvider::FAudioBusProvider()
 		: TDeviceDataMapTraceProvider<uint32, TSharedPtr<FAudioBusAssetDashboardEntry>>(GetName_Static())
 	{
+		FAssetRegistryModule& AssetRegistryModule = FModuleManager::LoadModuleChecked<FAssetRegistryModule>("AssetRegistry");
+
+		AssetRegistryModule.Get().OnAssetAdded().AddRaw(this, &FAudioBusProvider::OnAssetAdded);
+		AssetRegistryModule.Get().OnAssetRemoved().AddRaw(this, &FAudioBusProvider::OnAssetRemoved);
+		AssetRegistryModule.Get().OnFilesLoaded().AddRaw(this, &FAudioBusProvider::OnFilesLoaded);
+
+		FEditorDashboardFactory::OnActiveAudioDeviceChanged.AddRaw(this, &FAudioBusProvider::OnActiveAudioDeviceChanged);
+	}
+
+	FAudioBusProvider::~FAudioBusProvider()
+	{
+		if (FAssetRegistryModule* AssetRegistryModule = FModuleManager::GetModulePtr<FAssetRegistryModule>("AssetRegistry"))
+		{
+			AssetRegistryModule->Get().OnAssetAdded().RemoveAll(this);
+			AssetRegistryModule->Get().OnAssetRemoved().RemoveAll(this);
+			AssetRegistryModule->Get().OnFilesLoaded().RemoveAll(this);
+		}
+
+		FEditorDashboardFactory::OnActiveAudioDeviceChanged.RemoveAll(this);
 	}
 
 	FName FAudioBusProvider::GetName_Static()
@@ -20,6 +39,111 @@ namespace UE::Audio::Insights
 		return "AudioBusesProvider";
 	}
 
+		void FAudioBusProvider::OnAssetAdded(const FAssetData& InAssetData)
+	{
+		if (bAreFilesLoaded && InAssetData.AssetClassPath == FTopLevelAssetPath(UAudioBus::StaticClass()))
+		{
+			AddAudioBusAsset(InAssetData);
+		}
+	}
+
+	void FAudioBusProvider::OnAssetRemoved(const FAssetData& InAssetData)
+	{
+		if (InAssetData.AssetClassPath == FTopLevelAssetPath(UAudioBus::StaticClass()))
+		{
+			RemoveAudioBusAsset(InAssetData);
+		}
+	}
+
+	void FAudioBusProvider::OnFilesLoaded()
+	{
+		bAreFilesLoaded = true;
+		UpdateAudioBusAssetNames();
+	}
+
+	void FAudioBusProvider::OnActiveAudioDeviceChanged()
+	{
+		UpdateAudioBusAssetNames();
+	}
+
+	void FAudioBusProvider::AddAudioBusAsset(const FAssetData& InAssetData)
+	{
+		const bool bIsAudioBusAssetAlreadAdded = AudioBusDataViewEntries.ContainsByPredicate(
+			[AssetName = InAssetData.GetObjectPathString()](const TSharedPtr<FAudioBusAssetDashboardEntry>& AudioBusAssetDashboardEntry)
+			{
+				return AudioBusAssetDashboardEntry->Name == AssetName;
+			});
+
+		if (!bIsAudioBusAssetAlreadAdded)
+		{
+			const FAudioInsightsEditorModule AudioInsightsEditorModule = FAudioInsightsEditorModule::GetChecked();
+			const ::Audio::FDeviceId AudioDeviceId = AudioInsightsEditorModule.GetDeviceId();
+
+			TSharedPtr<FAudioBusAssetDashboardEntry> AudioBusAssetDashboardEntry = MakeShared<FAudioBusAssetDashboardEntry>();
+			AudioBusAssetDashboardEntry->DeviceId   = AudioDeviceId;
+			AudioBusAssetDashboardEntry->EntryType  = EAudioBusEntryType::AssetBased;
+			AudioBusAssetDashboardEntry->Name       = InAssetData.GetObjectPathString();
+			AudioBusAssetDashboardEntry->AudioBus   = Cast<UAudioBus>(InAssetData.GetAsset());
+			AudioBusAssetDashboardEntry->AudioBusId = AudioBusAssetDashboardEntry->AudioBus->GetUniqueID();
+
+			AudioBusDataViewEntries.Add(MoveTemp(AudioBusAssetDashboardEntry));
+
+			OnAudioBusAssetAdded.Broadcast(InAssetData.GetAsset());
+
+			bAssetsUpdated = true;
+			++LastUpdateId;
+		}
+	}
+
+	void FAudioBusProvider::RemoveAudioBusAsset(const FAssetData& InAssetData)
+	{
+		const int32 FoundAudioBusAssetNameIndex = AudioBusDataViewEntries.IndexOfByPredicate(
+			[AssetName = InAssetData.GetObjectPathString()](const TSharedPtr<FAudioBusAssetDashboardEntry>& AudioBusAssetDashboardEntry)
+			{
+				return AudioBusAssetDashboardEntry->Name == AssetName;
+			});
+
+		if (FoundAudioBusAssetNameIndex != INDEX_NONE)
+		{
+			const FAudioInsightsEditorModule AudioInsightsEditorModule = FAudioInsightsEditorModule::GetChecked();
+			const ::Audio::FDeviceId AudioDeviceId = AudioInsightsEditorModule.GetDeviceId();
+
+			RemoveDeviceEntry(AudioDeviceId, AudioBusDataViewEntries[FoundAudioBusAssetNameIndex]->AudioBus->GetUniqueID());
+
+			AudioBusDataViewEntries.RemoveAt(FoundAudioBusAssetNameIndex);
+
+			OnAudioBusAssetRemoved.Broadcast(InAssetData.GetAsset());
+
+			bAssetsUpdated = true;
+			++LastUpdateId;
+		}
+	}
+
+	void FAudioBusProvider::UpdateAudioBusAssetNames()
+	{
+		// Get all UAudioBus assets
+		TArray<FAssetData> AssetDataArray;
+
+		FAssetRegistryModule& AssetRegistryModule = FModuleManager::LoadModuleChecked<FAssetRegistryModule>("AssetRegistry");
+		AssetRegistryModule.Get().GetAssetsByClass(FTopLevelAssetPath(UAudioBus::StaticClass()), AssetDataArray);
+
+		// Build AudioBusDataViewEntries
+		Reset();
+		AudioBusDataViewEntries.Empty();
+
+		for (const FAssetData& AssetData : AssetDataArray)
+		{
+			AddAudioBusAsset(AssetData);
+		}
+
+		AudioBusDataViewEntries.Sort([](const TSharedPtr<FAudioBusAssetDashboardEntry>& A, const TSharedPtr<FAudioBusAssetDashboardEntry>& B)
+		{
+			return A->GetDisplayName().CompareToCaseIgnored(B->GetDisplayName()) < 0;
+		});
+
+		OnAudioBusAssetListUpdated.Broadcast();
+	}
+
 	bool FAudioBusProvider::ProcessMessages()
 	{
 		// Helper lambdas
@@ -51,12 +175,44 @@ namespace UE::Audio::Insights
 
 
 		// Process messages
+		if (bAssetsUpdated)
+		{
+			const FAudioInsightsEditorModule AudioInsightsEditorModule = FAudioInsightsEditorModule::GetChecked();
+			const ::Audio::FDeviceId AudioDeviceId = AudioInsightsEditorModule.GetDeviceId();
+
+			for (const TSharedPtr<FAudioBusAssetDashboardEntry>& AudioBusDataViewEntry : AudioBusDataViewEntries)
+			{
+				if (AudioBusDataViewEntry.IsValid() && AudioBusDataViewEntry->AudioBus.IsValid())
+				{
+					UpdateDeviceEntry(AudioDeviceId, AudioBusDataViewEntry->AudioBusId, [&AudioBusDataViewEntry](TSharedPtr<FAudioBusAssetDashboardEntry>& Entry)
+					{
+						if (!Entry.IsValid())
+						{
+							Entry = AudioBusDataViewEntry;
+						}
+					});
+				}
+			}
+
+			bAssetsUpdated = false;
+		}
+
 		ProcessMessageQueue<FAudioBusActivateMessage>(TraceMessages.ActivateMessages, BumpEntryFunc,
 		[](const FAudioBusActivateMessage& Msg, TSharedPtr<FAudioBusAssetDashboardEntry>* OutEntry)
 		{
 			FAudioBusAssetDashboardEntry& EntryRef = *OutEntry->Get();
-			EntryRef.Name = *Msg.Name;
-			EntryRef.EntryType = FSoftObjectPath(EntryRef.Name).IsAsset() ? EAudioBusEntryType::AssetBased : EAudioBusEntryType::CodeGenerated;
+
+			if (EntryRef.Name.IsEmpty())
+			{
+				EntryRef.Name = *Msg.Name;
+			}
+
+			if (EntryRef.EntryType == EAudioBusEntryType::None)
+			{
+				EntryRef.EntryType = FSoftObjectPath(EntryRef.Name).IsAsset() ? EAudioBusEntryType::AssetBased : EAudioBusEntryType::CodeGenerated;
+			}
+
+			EntryRef.bIsRegistered = true;
 		});
 
 		ProcessMessageQueue<FAudioBusDeactivateMessage>(TraceMessages.DeactivateMessages, GetEntry,
@@ -66,9 +222,12 @@ namespace UE::Audio::Insights
 			{
 				FAudioBusAssetDashboardEntry& EntryRef = *OutEntry->Get();
 
-				OnAudioBusDeactivated.Broadcast(EntryRef.GetObject());
+				EntryRef.bIsRegistered = false;
 
-				RemoveDeviceEntry(Msg.DeviceId, Msg.AudioBusId);
+				if (EntryRef.EntryType == EAudioBusEntryType::CodeGenerated)
+				{
+					RemoveDeviceEntry(Msg.DeviceId, Msg.AudioBusId);
+				}
 			}
 		});
 
@@ -78,8 +237,12 @@ namespace UE::Audio::Insights
 			if (OutEntry)
 			{
 				FAudioBusAssetDashboardEntry& EntryRef = *OutEntry->Get();
-				EntryRef.bHasActivity = Msg.bHasActivity;
-				EntryRef.Timestamp = Msg.Timestamp;
+
+				if (EntryRef.bIsRegistered)
+				{
+					EntryRef.bHasActivity = Msg.bHasActivity;
+					EntryRef.Timestamp = Msg.Timestamp;
+				}
 			}
 		});
 
diff --git a/Engine/Plugins/AudioInsights/Source/AudioInsightsEditor/Private/Providers/AudioBusProvider.h b/Engine/Plugins/AudioInsights/Source/AudioInsightsEditor/Private/Providers/AudioBusProvider.h
--- a/Engine/Plugins/AudioInsights/Source/AudioInsightsEditor/Private/Providers/AudioBusProvider.h
+++ b/Engine/Plugins/AudioInsights/Source/AudioInsightsEditor/Private/Providers/AudioBusProvider.h
@@ -11,18 +11,39 @@ namespace UE::Audio::Insights
 	{
 	public:
 		FAudioBusProvider();
-		virtual ~FAudioBusProvider() = default;
+		virtual ~FAudioBusProvider();
 
 		static FName GetName_Static();
 
 		virtual UE::Trace::IAnalyzer* ConstructAnalyzer(TraceServices::IAnalysisSession& InSession) override;
 
-		DECLARE_MULTICAST_DELEGATE_OneParam(FOnAudioBusDeactivated, const TWeakObjectPtr<UObject> /*AudioBus*/);
-		inline static FOnAudioBusDeactivated OnAudioBusDeactivated;
+		DECLARE_MULTICAST_DELEGATE_OneParam(FOnAudioBusAssetAdded, const TWeakObjectPtr<UObject> /*Asset*/);
+		inline static FOnAudioBusAssetAdded OnAudioBusAssetAdded;
+
+		DECLARE_MULTICAST_DELEGATE_OneParam(FOnAudioBusAssetRemoved, const TWeakObjectPtr<UObject> /*Asset*/);
+		inline static FOnAudioBusAssetRemoved OnAudioBusAssetRemoved;
+
+		DECLARE_MULTICAST_DELEGATE(FOnAudioBusAssetListUpdated);
+		inline static FOnAudioBusAssetListUpdated OnAudioBusAssetListUpdated;
 
 	private:
+		void OnAssetAdded(const FAssetData& InAssetData);
+		void OnAssetRemoved(const FAssetData& InAssetData);
+		void OnFilesLoaded();
+		void OnActiveAudioDeviceChanged();
+
+		void AddAudioBusAsset(const FAssetData& InAssetData);
+		void RemoveAudioBusAsset(const FAssetData& InAssetData);
+
+		void UpdateAudioBusAssetNames();
+
 		virtual bool ProcessMessages() override;
 
+		bool bAreFilesLoaded = false;
+		bool bAssetsUpdated = false;
+
+		TArray<TSharedPtr<FAudioBusAssetDashboardEntry>> AudioBusDataViewEntries;
+
 		FAudioBusMessages TraceMessages;
 	};
 } // namespace UE::Audio::Insights
diff --git a/Engine/Plugins/AudioInsights/Source/AudioInsightsEditor/Private/Views/AudioBusesDashboardViewFactory.cpp b/Engine/Plugins/AudioInsights/Source/AudioInsightsEditor/Private/Views/AudioBusesDashboardViewFactory.cpp
--- a/Engine/Plugins/AudioInsights/Source/AudioInsightsEditor/Private/Views/AudioBusesDashboardViewFactory.cpp
+++ b/Engine/Plugins/AudioInsights/Source/AudioInsightsEditor/Private/Views/AudioBusesDashboardViewFactory.cpp
@@ -22,15 +22,14 @@ namespace UE::Audio::Insights
 		{
 			return static_cast<const FAudioBusAssetDashboardEntry&>(InData);
 		};
-
-		FAudioBusAssetDashboardEntry& CastEntry(IDashboardDataViewEntry& InData)
-		{
-			return static_cast<FAudioBusAssetDashboardEntry&>(InData);
-		};
 	}
 
 	FAudioBusesDashboardViewFactory::FAudioBusesDashboardViewFactory()
 	{
+		FAudioBusProvider::OnAudioBusAssetAdded.AddRaw(this, &FAudioBusesDashboardViewFactory::HandleOnAudioBusAssetListUpdated);
+		FAudioBusProvider::OnAudioBusAssetRemoved.AddRaw(this, &FAudioBusesDashboardViewFactory::HandleOnAudioBusAssetListUpdated);
+		FAudioBusProvider::OnAudioBusAssetListUpdated.AddRaw(this, &FAudioBusesDashboardViewFactory::RequestListRefresh);
+
 		IAudioInsightsTraceModule& AudioInsightsTraceModule = IAudioInsightsEditorModule::GetChecked().GetTraceModule();
 
 		const TSharedPtr<FAudioBusProvider> AudioBusProvider = MakeShared<FAudioBusProvider>();
@@ -46,6 +45,13 @@ namespace UE::Audio::Insights
 		SortMode     = EColumnSortMode::Ascending;
 	}
 
+	FAudioBusesDashboardViewFactory::~FAudioBusesDashboardViewFactory()
+	{
+		FAudioBusProvider::OnAudioBusAssetAdded.RemoveAll(this);
+		FAudioBusProvider::OnAudioBusAssetRemoved.RemoveAll(this);
+		FAudioBusProvider::OnAudioBusAssetListUpdated.RemoveAll(this);
+	}
+
 	FName FAudioBusesDashboardViewFactory::GetName() const
 	{
 		return "AudioBuses";
@@ -88,11 +94,11 @@ namespace UE::Audio::Insights
 				return SNullWidget::NullWidget;
 			}
 
-			FAudioBusAssetDashboardEntry& AudioBusAssetDashboardEntry = AudioBusesPrivate::CastEntry(InRowData.Get());
+			const FAudioBusAssetDashboardEntry& AudioBusAssetDashboardEntry = AudioBusesPrivate::CastEntry(InRowData.Get());
 
 			if (AudioBusAssetDashboardEntry.EntryType == EAudioBusEntryType::AssetBased)
 			{
-				const TWeakObjectPtr<UAudioBus> AudioBus(Cast<UAudioBus>(AudioBusAssetDashboardEntry.GetObject()));
+				const TWeakObjectPtr<UAudioBus> AudioBus = AudioBusAssetDashboardEntry.AudioBus;
 
 				if (!AudioBus.IsValid())
 				{
@@ -110,8 +116,11 @@ namespace UE::Audio::Insights
 						.IsChecked_Lambda([this, AudioBus]()
 						{
 							const bool* FoundAudioBusCheckedState = AudioBusCheckboxCheckedStates.Find(AudioBus);
-							return FoundAudioBusCheckedState && *FoundAudioBusCheckedState == true ? ECheckBoxState::Checked : ECheckBoxState::Unchecked;
+							const bool bIsAudioBusChecked = FoundAudioBusCheckedState && *FoundAudioBusCheckedState == true;
+
+							return bIsAudioBusChecked ? ECheckBoxState::Checked : ECheckBoxState::Unchecked;
 						})
+
 						.OnCheckStateChanged_Lambda([this, AudioBus](ECheckBoxState NewState)
 						{
 							if (bool* FoundAudioBusCheckedState = AudioBusCheckboxCheckedStates.Find(AudioBus))
@@ -119,7 +128,7 @@ namespace UE::Audio::Insights
 								*FoundAudioBusCheckedState = NewState == ECheckBoxState::Checked;
 							}
 
-							OnAudioBusItemChecked.Broadcast(NewState == ECheckBoxState::Checked, AudioBus);
+							OnAudioBusAssetChecked.Broadcast(NewState == ECheckBoxState::Checked, AudioBus);
 						})
 					]
 					+ SHorizontalBox::Slot()
@@ -134,14 +143,23 @@ namespace UE::Audio::Insights
 						SNew(STextBlock)
 						.Text(ValueText)
 						.MinDesiredWidth(300)
+						.ColorAndOpacity_Lambda([InRowData]()
+						{
+							static const FSlateColor DefaultColor(FColor::White);
+							static const FSlateColor AudioBusAssetColor(FColor(80, 200, 255));
+							
+							const FAudioBusAssetDashboardEntry& AudioBusAssetDashboardEntry = AudioBusesPrivate::CastEntry(InRowData.Get());
+
+							return AudioBusAssetDashboardEntry.bIsRegistered ? AudioBusAssetColor : DefaultColor;
+						})
 						.OnDoubleClicked_Lambda([this, InRowData](const FGeometry& MyGeometry, const FPointerEvent& PointerEvent)
 						{
 							if (GEditor)
 							{
-								TSharedPtr<IObjectDashboardEntry> ObjectData = StaticCastSharedPtr<IObjectDashboardEntry>(InRowData.ToSharedPtr());
+								const TSharedPtr<IObjectDashboardEntry> ObjectData = StaticCastSharedPtr<IObjectDashboardEntry>(InRowData.ToSharedPtr());
 								if (ObjectData.IsValid())
 								{
-									const UObject* Object = ObjectData->GetObject();
+									const TObjectPtr<UObject> Object = ObjectData->GetObject();
 									if (Object && Object->IsAsset())
 									{
 										GEditor->GetEditorSubsystem<UAssetEditorSubsystem>()->OpenEditorForAsset(Object);
@@ -165,6 +183,7 @@ namespace UE::Audio::Insights
 						SNew(STextBlock)
 						.Text(ValueText)
 						.MinDesiredWidth(300)
+						.ColorAndOpacity(FSlateColor(FColor(80, 200, 255)))
 					];
 			}
 		}
@@ -256,8 +275,6 @@ namespace UE::Audio::Insights
 	{
 		if (!DashboardWidget.IsValid())
 		{
-			FAudioDeviceManagerDelegates::OnAudioDeviceDestroyed.AddSP(this, &FAudioBusesDashboardViewFactory::HandleOnAudioDeviceDestroyed);
-
 			SAssignNew(DashboardWidget, SVerticalBox)
 			+ SVerticalBox::Slot()
 			.AutoHeight()
@@ -412,9 +429,17 @@ namespace UE::Audio::Insights
 		}
 	}
 
-	void FAudioBusesDashboardViewFactory::HandleOnAudioDeviceDestroyed(::Audio::FDeviceId InAudioDeviceId)
+	void FAudioBusesDashboardViewFactory::RequestListRefresh()
+	{
+		if (FilteredEntriesListView.IsValid())
+		{
+			FilteredEntriesListView->RequestListRefresh();
+		}
+	}
+
+	void FAudioBusesDashboardViewFactory::HandleOnAudioBusAssetListUpdated(const TWeakObjectPtr<UObject> InAsset)
 	{
-		AudioBusCheckboxCheckedStates.Empty();
+		RequestListRefresh();
 	}
 } // namespace UE::Audio::Insights
 #undef LOCTEXT_NAMESPACE
diff --git a/Engine/Plugins/AudioInsights/Source/AudioInsightsEditor/Private/Views/AudioBusesDashboardViewFactory.h b/Engine/Plugins/AudioInsights/Source/AudioInsightsEditor/Private/Views/AudioBusesDashboardViewFactory.h
--- a/Engine/Plugins/AudioInsights/Source/AudioInsightsEditor/Private/Views/AudioBusesDashboardViewFactory.h
+++ b/Engine/Plugins/AudioInsights/Source/AudioInsightsEditor/Private/Views/AudioBusesDashboardViewFactory.h
@@ -3,7 +3,6 @@
 
 #include "Delegates/Delegate.h"
 #include "Templates/SharedPointer.h"
-#include "Messages/AudioBusMessages.h"
 #include "Views/TableDashboardViewFactory.h"
 
 class UAudioBus;
@@ -14,7 +13,7 @@ namespace UE::Audio::Insights
 	{
 	public:
 		FAudioBusesDashboardViewFactory();
-		virtual ~FAudioBusesDashboardViewFactory() = default;
+		virtual ~FAudioBusesDashboardViewFactory();
 
 		virtual FName GetName() const override;
 		virtual FText GetDisplayName() const override;
@@ -22,8 +21,8 @@ namespace UE::Audio::Insights
 		virtual EDefaultDashboardTabStack GetDefaultTabStack() const override;
 		virtual TSharedRef<SWidget> MakeWidget(TSharedRef<SDockTab> OwnerTab, const FSpawnTabArgs& SpawnTabArgs) override;
 
-		DECLARE_MULTICAST_DELEGATE_TwoParams(FOnAudioBusItemChecked, const bool /*bIsChecked*/, const TWeakObjectPtr<UAudioBus> /*AudioBus*/);
-		inline static FOnAudioBusItemChecked OnAudioBusItemChecked;
+		DECLARE_MULTICAST_DELEGATE_TwoParams(FOnAudioBusAssetChecked, const bool /*bIsChecked*/, const TWeakObjectPtr<UAudioBus> /*AudioBus*/);
+		inline static FOnAudioBusAssetChecked OnAudioBusAssetChecked;
 
 	private:
 		enum class EAudioBusTypeComboboxSelection : uint8
@@ -44,7 +43,9 @@ namespace UE::Audio::Insights
 		void FilterByAudioBusName();
 		void FilterByAudioBusType();
 
-		void HandleOnAudioDeviceDestroyed(::Audio::FDeviceId InAudioDeviceId);
+		void RequestListRefresh();
+
+		void HandleOnAudioBusAssetListUpdated(const TWeakObjectPtr<UObject> InAsset);
 
 		TMap<const TWeakObjectPtr<UAudioBus>, bool> AudioBusCheckboxCheckedStates;
 
diff --git a/Engine/Plugins/AudioInsights/Source/AudioInsightsEditor/Private/Views/AudioMetersDashboardViewFactory.cpp b/Engine/Plugins/AudioInsights/Source/AudioInsightsEditor/Private/Views/AudioMetersDashboardViewFactory.cpp
--- a/Engine/Plugins/AudioInsights/Source/AudioInsightsEditor/Private/Views/AudioMetersDashboardViewFactory.cpp
+++ b/Engine/Plugins/AudioInsights/Source/AudioInsightsEditor/Private/Views/AudioMetersDashboardViewFactory.cpp
@@ -41,9 +41,9 @@ namespace UE::Audio::Insights
 			OnSubmixAssetCheckedHandle = FSubmixesDashboardViewFactory::OnSubmixAssetChecked.AddSP(this, &FAudioMetersDashboardViewFactory::HandleOnSubmixAssetChecked);
 		}
 
-		if (!OnAudioBusItemCheckedHandle.IsValid())
+		if (!OnAudioBusAssetCheckedHandle.IsValid())
 		{
-			OnAudioBusItemCheckedHandle = FAudioBusesDashboardViewFactory::OnAudioBusItemChecked.AddSP(this, &FAudioMetersDashboardViewFactory::HandleOnAudioBusItemChecked);
+			OnAudioBusAssetCheckedHandle = FAudioBusesDashboardViewFactory::OnAudioBusAssetChecked.AddSP(this, &FAudioMetersDashboardViewFactory::HandleOnAudioBusAssetChecked);
 		}
 
 		if (!OnSubmixAssetRemovedHandle.IsValid())
@@ -51,9 +51,9 @@ namespace UE::Audio::Insights
 			OnSubmixAssetRemovedHandle = FSoundSubmixProvider::OnSubmixAssetRemoved.AddSP(this, &FAudioMetersDashboardViewFactory::HandleOnAssetRemoved);
 		}
 
-		if (!OnAudioBusDeactivatedHandle.IsValid())
+		if (!OnAudioBusAssetRemovedHandle.IsValid())
 		{
-			OnAudioBusDeactivatedHandle = FAudioBusProvider::OnAudioBusDeactivated.AddSP(this, &FAudioMetersDashboardViewFactory::HandleOnAssetRemoved);
+			OnAudioBusAssetRemovedHandle = FAudioBusProvider::OnAudioBusAssetRemoved.AddSP(this, &FAudioMetersDashboardViewFactory::HandleOnAssetRemoved);
 		}
 		
 		if (!MeterViewsScrollBox.IsValid())
@@ -104,7 +104,7 @@ namespace UE::Audio::Insights
 		}
 	}
 
-	void FAudioMetersDashboardViewFactory::HandleOnAudioBusItemChecked(const bool bInIsChecked, const TWeakObjectPtr<UAudioBus> InAudioBus)
+	void FAudioMetersDashboardViewFactory::HandleOnAudioBusAssetChecked(const bool bInIsChecked, const TWeakObjectPtr<UAudioBus> InAudioBus)
 	{
 		if (!InAudioBus.IsValid())
 		{
diff --git a/Engine/Plugins/AudioInsights/Source/AudioInsightsEditor/Private/Views/AudioMetersDashboardViewFactory.h b/Engine/Plugins/AudioInsights/Source/AudioInsightsEditor/Private/Views/AudioMetersDashboardViewFactory.h
--- a/Engine/Plugins/AudioInsights/Source/AudioInsightsEditor/Private/Views/AudioMetersDashboardViewFactory.h
+++ b/Engine/Plugins/AudioInsights/Source/AudioInsightsEditor/Private/Views/AudioMetersDashboardViewFactory.h
@@ -27,7 +27,7 @@ namespace UE::Audio::Insights
 
 	private:
 		void HandleOnSubmixAssetChecked(const bool bInIsChecked, const TWeakObjectPtr<USoundSubmix> InSoundSubmix);
-		void HandleOnAudioBusItemChecked(const bool bInIsChecked, const TWeakObjectPtr<UAudioBus> InAudioBus);
+		void HandleOnAudioBusAssetChecked(const bool bInIsChecked, const TWeakObjectPtr<UAudioBus> InAudioBus);
 
 		void HandleOnAssetRemoved(const TWeakObjectPtr<UObject> InAsset);
 
@@ -36,8 +36,8 @@ namespace UE::Audio::Insights
 		TMap<uint32, TSharedRef<FAudioMeterView>> AudioMeterViews;
 
 		FDelegateHandle OnSubmixAssetCheckedHandle;
-		FDelegateHandle OnAudioBusItemCheckedHandle;
+		FDelegateHandle OnAudioBusAssetCheckedHandle;
 		FDelegateHandle OnSubmixAssetRemovedHandle;
-		FDelegateHandle OnAudioBusDeactivatedHandle;
+		FDelegateHandle OnAudioBusAssetRemovedHandle;
 	};
 } // namespace UE::Audio::Insights
diff --git a/Engine/Source/Runtime/AudioMixer/Private/AudioMixerSourceManager.cpp b/Engine/Source/Runtime/AudioMixer/Private/AudioMixerSourceManager.cpp
--- a/Engine/Source/Runtime/AudioMixer/Private/AudioMixerSourceManager.cpp
+++ b/Engine/Source/Runtime/AudioMixer/Private/AudioMixerSourceManager.cpp
@@ -270,15 +270,17 @@ namespace Audio
 		int32 Index = 0;
 
 #if PLATFORM_ENABLE_VECTORINTRINSICS
+		constexpr VectorRegister4Float Minus60dBLinear = MakeVectorRegisterFloatConstant(0.001f, 0.001f, 0.001f, 0.001f);
+
 		const int32 SimdNum = NumSamples & 0xFFFFFFF0;
 		for (; Index < SimdNum; Index += 16)
 		{
 			const VectorRegister4x4Float Samples = VectorLoad16(&AudioBuffer[Index]);
 
-			if (VectorAnyGreaterThan(VectorAbs(Samples.val[0]), GlobalVectorConstants::SmallNumber) ||
-			    VectorAnyGreaterThan(VectorAbs(Samples.val[1]), GlobalVectorConstants::SmallNumber) ||
-			    VectorAnyGreaterThan(VectorAbs(Samples.val[2]), GlobalVectorConstants::SmallNumber) ||
-			    VectorAnyGreaterThan(VectorAbs(Samples.val[3]), GlobalVectorConstants::SmallNumber))
+			if (VectorAnyGreaterThan(VectorAbs(Samples.val[0]), Minus60dBLinear) ||
+			    VectorAnyGreaterThan(VectorAbs(Samples.val[1]), Minus60dBLinear) ||
+			    VectorAnyGreaterThan(VectorAbs(Samples.val[2]), Minus60dBLinear) ||
+			    VectorAnyGreaterThan(VectorAbs(Samples.val[3]), Minus60dBLinear))
 			{
 				bIsSilent = false;
 				Index = INT_MAX;
