[Audio Insights] Audio Insights will not populate with data if the trace was already active prior to PIE

- This happens because FTraceModule::OnAnalysisBegin doesn't get called when a trace is started outside of Audio Insights, so the current analysis session doesn't have the providers and analyzers required to function
- Solution: we need to reboot the current analysis session as we cannot add analyzers or providers after the session has been initialized. InsightsModule.StartAnalysisForLastLiveSession() does this by reseting and then reloading the current active session, allowing us to add the required analyzers in the proccess.

diff --git a/Engine/Plugins/AudioInsights/Source/AudioInsights/Private/AudioInsightsTraceModule.cpp b/Engine/Plugins/AudioInsights/Source/AudioInsights/Private/AudioInsightsTraceModule.cpp
--- a/Engine/Plugins/AudioInsights/Source/AudioInsights/Private/AudioInsightsTraceModule.cpp
+++ b/Engine/Plugins/AudioInsights/Source/AudioInsights/Private/AudioInsightsTraceModule.cpp
@@ -161,11 +161,12 @@ namespace UE::Audio::Insights
 			Options.bExcludeTail = true;
 
 			FTraceAuxiliary::Start(FTraceAuxiliary::EConnectionType::Network, TEXT("localhost"), TEXT(""), &Options);
-
-			IUnrealInsightsModule& UnrealInsightsModule = FModuleManager::LoadModuleChecked<IUnrealInsightsModule>("TraceInsights");
-			UnrealInsightsModule.StartAnalysisForLastLiveSession();
 		}
 
+		// Even if a trace was already active, we still need to notify Audio Insights that analysis has begun
+		IUnrealInsightsModule& InsightsModule = FModuleManager::GetModuleChecked<IUnrealInsightsModule>("TraceInsights");
+		InsightsModule.StartAnalysisForLastLiveSession();
+
 		bAudioTracesAreAcive = true;
 	}
 
