[Audio Insights] - Explore adding messaging when the standalone window is opened, but the audio and audio mixer channels are not connected

- When using the standalone version of Unreal/Audio insights, no data will be analyzed without the audio and audiomixer trace channels being enabled
- We now detect whether these channels are enabled, and display a button (Enable Audio Traces) if they are not. This lets the user quickly enable the audio traces without having to reconnect with new command line arguments

diff --git a/Engine/Plugins/AudioInsights/Source/AudioInsights/Private/AudioInsightsComponent.cpp b/Engine/Plugins/AudioInsights/Source/AudioInsights/Private/AudioInsightsComponent.cpp
--- a/Engine/Plugins/AudioInsights/Source/AudioInsights/Private/AudioInsightsComponent.cpp
+++ b/Engine/Plugins/AudioInsights/Source/AudioInsights/Private/AudioInsightsComponent.cpp
@@ -2,6 +2,7 @@
 #include "AudioInsightsComponent.h"
 
 #include "AudioInsightsModule.h"
+#include "AudioInsightsTraceModule.h"
 #include "AudioInsightsStyle.h"
 #include "Insights/IInsightsManager.h"
 #include "Modules/ModuleManager.h"
@@ -115,6 +116,10 @@ namespace UE::Audio::Insights
 
 	TSharedRef<SDockTab> FAudioInsightsComponent::SpawnTab(const FSpawnTabArgs& Args)
 	{
+#if !WITH_EDITOR
+		IAudioInsightsTraceModule& TraceModule = FAudioInsightsModule::GetChecked().GetTraceModule();
+		TraceModule.RequestChannelUpdate();
+#endif // !WITH_EDITOR
 		const TSharedRef<SDockTab> DockTab = FAudioInsightsModule::GetChecked().CreateDashboardTabWidget(Args);
 
 		OnTabSpawn.Broadcast();
@@ -176,6 +181,11 @@ namespace UE::Audio::Insights
 					const TraceServices::FSessionInfo& TraceServicesSessionInfo = DiagnosticsProvider->GetSessionInfo();
 
 					bIsEditorTrace = TraceServicesSessionInfo.TargetType == EBuildTargetType::Editor;
+
+#if !WITH_EDITOR
+					IAudioInsightsTraceModule& TraceModule = FAudioInsightsModule::GetChecked().GetTraceModule();
+					TraceModule.InitializeSessionInfo(TraceServicesSessionInfo);
+#endif // !WITH_EDITOR
 				}
 
 				// Allow to spawn tab if non-editor trace
diff --git a/Engine/Plugins/AudioInsights/Source/AudioInsights/Private/AudioInsightsDashboardFactory.cpp b/Engine/Plugins/AudioInsights/Source/AudioInsights/Private/AudioInsightsDashboardFactory.cpp
--- a/Engine/Plugins/AudioInsights/Source/AudioInsights/Private/AudioInsightsDashboardFactory.cpp
+++ b/Engine/Plugins/AudioInsights/Source/AudioInsights/Private/AudioInsightsDashboardFactory.cpp
@@ -7,6 +7,7 @@
 #include "Framework/MultiBox/MultiBoxBuilder.h"
 #include "Internationalization/Text.h"
 #include "Templates/SharedPointer.h"
+#include "Widgets/Input/SButton.h"
 #include "Widgets/Input/SCheckBox.h"
 #include "Widgets/Layout/SBox.h"
 
@@ -18,6 +19,11 @@ namespace UE::Audio::Insights
 	namespace DashboardFactoryPrivate
 	{
 		static const FText ToolName = LOCTEXT("AudioDashboard_ToolName", "Audio Insights");
+
+		static const FText EnableTracesButtonText = LOCTEXT("AudioDashboard_AutomaticallyEnableTracesTitle", "Enable audio traces");
+		static const FText EnableTracesDescription = LOCTEXT("AudioDashboard_AutomaticallyEnableTracesDescription", "Enables the audio and audio mixer trace channels. Audio Insights will not function without these channels enabled.");
+		static const FText NoTracesEnabledWarning = LOCTEXT("AudioDashboard_NoTracesEnabledWarning", "Audio Insights requires the audio and audio mixer trace channels to be enabled to function.");
+		static const FText EnableThemNowText = LOCTEXT("AudioDashboard_EnableNowText", "Enable them now?");
 	}
 
 	::Audio::FDeviceId FDashboardFactory::GetDeviceId() const
@@ -38,22 +44,45 @@ namespace UE::Audio::Insights
 
 		RegisterTabSpawners();
 
-		const TSharedRef<SWidget> TabContent = SNew(SVerticalBox)
-			+ SVerticalBox::Slot()
-			.AutoHeight()
-			[
-				MakeMenuBarWidget()
-			]
-			+ SVerticalBox::Slot()
-			.AutoHeight()
+		const TSharedRef<SWidget> TabContent =
+
+#if !WITH_EDITOR
+			SNew(SOverlay)
+
+			+ SOverlay::Slot()
 			[
-				SNew(SBox)
-				.HeightOverride(4.0f)
+#endif	// !WITH_EDITOR
+
+				SNew(SVerticalBox)
+				+ SVerticalBox::Slot()
+				.AutoHeight()
+				[
+					MakeMenuBarWidget()
+				]
+
+				+ SVerticalBox::Slot()
+
+				.AutoHeight()
+				[
+					SNew(SBox)
+					.HeightOverride(4.0f)
+				]
+				+ SVerticalBox::Slot()
+				[
+					DashboardTabManager->RestoreFrom(TabLayout->AsShared(), Args.GetOwnerWindow()).ToSharedRef()
+#if WITH_EDITOR
+				];
+#else
+				]
 			]
-			+ SVerticalBox::Slot()
+
+			+ SOverlay::Slot()
+			.HAlign(HAlign_Fill)
+			.VAlign(VAlign_Fill)
 			[
-				DashboardTabManager->RestoreFrom(TabLayout->AsShared(), Args.GetOwnerWindow()).ToSharedRef()
+				MakeEnableTracesOverlay()
 			];
+#endif	// WITH_EDITOR
 
 		DockTab->SetContent(TabContent);
 
@@ -264,5 +293,93 @@ namespace UE::Audio::Insights
 	{
 		DashboardViewFactories.Remove(InName);
 	}
+
+#if !WITH_EDITOR
+	TSharedRef<SWidget> FDashboardFactory::MakeEnableTracesOverlay()
+	{
+		using namespace DashboardFactoryPrivate;
+
+		return SNew(SBorder)
+			.BorderBackgroundColor(FLinearColor(0.2f, 0.2f, 0.2f, 0.8f))
+			.Visibility_Lambda([]() -> EVisibility
+			{
+				IAudioInsightsTraceModule& TraceModule = FAudioInsightsModule::GetChecked().GetTraceModule();
+				return TraceModule.AudioChannelsCanBeManuallyEnabled() ? EVisibility::Visible : EVisibility::Hidden;
+			})
+			[
+				SNew(SVerticalBox)
+					
+				+ SVerticalBox::Slot()
+				.VAlign(VAlign_Top)
+				[
+					SNullWidget::NullWidget
+				]
+
+				+ SVerticalBox::Slot()
+				.AutoHeight()
+				.HAlign(HAlign_Center)
+				.VAlign(VAlign_Center)
+				[
+					SNew(STextBlock)
+					.Text(NoTracesEnabledWarning)
+				]
+
+				+ SVerticalBox::Slot()
+				.AutoHeight()
+				.HAlign(HAlign_Center)
+				.VAlign(VAlign_Center)
+				[
+					SNew(STextBlock)
+					.Text(EnableThemNowText)
+				]
+
+				+ SVerticalBox::Slot()
+				.AutoHeight()
+				.HAlign(HAlign_Center)
+				.VAlign(VAlign_Center)
+				.Padding(0.0f, 10.0f)
+				[
+					MakeEnableTracesButton()
+				]
+
+				+ SVerticalBox::Slot()
+				.VAlign(VAlign_Bottom)
+				[
+					SNullWidget::NullWidget
+				]
+			];
+	}
+
+	TSharedRef<SWidget> FDashboardFactory::MakeEnableTracesButton()
+	{
+		using namespace DashboardFactoryPrivate;
+
+		return SNew(SButton)
+			.ButtonStyle(&FAppStyle::Get().GetWidgetStyle<FButtonStyle>("PrimaryButton"))
+			.OnClicked(this, &FDashboardFactory::ToggleAutoEnableAudioTraces)
+			.ToolTipText(EnableTracesDescription)
+			.ContentPadding(FMargin(0.0f, 0.0f, 0.0f, 0.0f))
+			.Content()
+			[
+				SNew(SBox)
+				.VAlign(VAlign_Center)
+				[
+					SNew(STextBlock)
+					.TextStyle(&FAppStyle::Get().GetWidgetStyle<FTextBlockStyle>("DialogButtonText"))
+					.Justification(ETextJustify::Center)
+					.Text(EnableTracesButtonText)
+				]
+			];
+	}
+
+	FReply FDashboardFactory::ToggleAutoEnableAudioTraces()
+	{
+		IAudioInsightsTraceModule& TraceModule = FAudioInsightsModule::GetChecked().GetTraceModule();
+		TraceModule.StartTraceAnalysis(false);
+
+		return FReply::Handled();
+	}
+#endif // !WITH_EDITOR
+
 } // namespace UE::Audio::Insights
 #undef LOCTEXT_NAMESPACE
diff --git a/Engine/Plugins/AudioInsights/Source/AudioInsights/Private/AudioInsightsTraceModule.cpp b/Engine/Plugins/AudioInsights/Source/AudioInsights/Private/AudioInsightsTraceModule.cpp
--- a/Engine/Plugins/AudioInsights/Source/AudioInsights/Private/AudioInsightsTraceModule.cpp
+++ b/Engine/Plugins/AudioInsights/Source/AudioInsights/Private/AudioInsightsTraceModule.cpp
@@ -2,21 +2,26 @@
 #include "AudioInsightsTraceModule.h"
 
 #include "CoreGlobals.h"
+#include "Insights/IInsightsManager.h"
 #include "Insights/IUnrealInsightsModule.h"
+#include "ISessionServicesModule.h"
+#include "Misc/App.h"
 #include "ProfilingDebugging/TraceAuxiliary.h"
 #include "Providers/MixerSourceTraceProvider.h"
 #include "Providers/VirtualLoopTraceProvider.h"
 #include "Templates/SharedPointer.h"
 #include "Trace/Trace.h"
+#include "TraceServices/Model/Channel.h"
+#include "TraceServices/Model/Diagnostics.h"
 
 
 namespace UE::Audio::Insights
 {
 	namespace FTraceModulePrivate
 	{
-		static const TCHAR* CPUChannelName = TEXT("CpuChannel");
-		static const TCHAR* AudioChannelName = TEXT("AudioChannel");
-		static const TCHAR* AudioMixerChannelName = TEXT("AudioMixerChannel");
+		static const TCHAR* CPUChannelName = TEXT("Cpu");
+		static const TCHAR* AudioChannelName = TEXT("Audio");
+		static const TCHAR* AudioMixerChannelName = TEXT("AudioMixer");
 	}
 
 	FTraceModule::FTraceModule()
@@ -69,13 +74,46 @@ namespace UE::Audio::Insights
 		, nullptr);
 	}
 
-	void FTraceModule::EnableAudioInsightsTraceChannels()
+	bool FTraceModule::EnableAudioInsightsTraceChannels()
 	{
 		using namespace FTraceModulePrivate;
 
-		UE::Trace::ToggleChannel(CPUChannelName, true);
-		UE::Trace::ToggleChannel(AudioChannelName, true);
-		UE::Trace::ToggleChannel(AudioMixerChannelName, true);
+#if WITH_EDITOR
+ 		const bool bCPUOn = UE::Trace::ToggleChannel(CPUChannelName, true);
+ 		const bool bAudioOn = UE::Trace::ToggleChannel(AudioChannelName, true);
+ 		const bool bAudioMixerOn = UE::Trace::ToggleChannel(AudioMixerChannelName, true);
+
+ 		return bCPUOn && bAudioOn && bAudioMixerOn;
+#else
+		if (!InstanceID.IsValid())
+		{
+			return false;
+		}
+
+		ISessionServicesModule& SessionServicesModule = FModuleManager::LoadModuleChecked<ISessionServicesModule>("SessionServices");
+		TSharedPtr<ITraceController> TraceController = SessionServicesModule.GetTraceController();
+		if (!TraceController.IsValid())
+		{
+			return false;
+		}
+
+		static TConstArrayView<FStringView> AudioChannels{ AudioChannelName, AudioMixerChannelName };
+		static TConstArrayView<FStringView> EmptyArray;
+
+		// Requires ITraceController to have received a response to discovering the session and channel info
+		if (TraceController->HasAvailableInstance(InstanceID))
+		{
+			TraceController->WithInstance(InstanceID, [&](const FTraceStatus& Status, ITraceControllerCommands& Commands)
+			{
+				if (Status.bIsTracing)
+				{
+					Commands.SetChannels(AudioChannels, EmptyArray);
+					return true;
+				}
+			});
+		}
+		return false;
+#endif // WITH_EDITOR
 	}
 
 	void FTraceModule::DisableAudioInsightsTraceChannels() const
@@ -119,6 +157,72 @@ namespace UE::Audio::Insights
 		}
 	}
 
+#if !WITH_EDITOR
+	const UE::Trace::FStoreClient::FSessionInfo* FTraceModule::GetCurrentAnalysisSessionInfo() const
+	{
+		IUnrealInsightsModule& InsightsModule = FModuleManager::GetModuleChecked<IUnrealInsightsModule>("TraceInsights");
+
+		TSharedPtr<const TraceServices::IAnalysisSession> AnalysisSession = InsightsModule.GetAnalysisSession();
+		if (!AnalysisSession.IsValid())
+		{
+			return nullptr;
+		}
+
+		UE::Trace::FStoreClient* StoreClient = InsightsModule.GetStoreClient();
+		if (StoreClient == nullptr)
+		{
+			return nullptr;
+		}
+
+		return StoreClient->GetSessionInfoByTraceId(AnalysisSession->GetTraceId());
+	}
+
+	bool FTraceModule::GetAudioTracesAreEnabled() const
+	{
+		using namespace FTraceModulePrivate;
+
+		const IUnrealInsightsModule& InsightsModule = FModuleManager::GetModuleChecked<IUnrealInsightsModule>("TraceInsights");
+		const TSharedPtr<const TraceServices::IAnalysisSession> AnalysisSession = InsightsModule.GetAnalysisSession();
+
+		if (!AnalysisSession.IsValid())
+		{
+			return false;
+		}
+
+		const TraceServices::IChannelProvider* ChannelProvider = AnalysisSession->ReadProvider<TraceServices::IChannelProvider>("ChannelProvider");
+		if (ChannelProvider == nullptr)
+		{
+			return false;
+		}
+		
+		bool bAudioChannelIsEnabled = false;
+		bool bAudioMixerChannelIsEnabled = false;
+
+		const TArray<TraceServices::FChannelEntry>& Channels = ChannelProvider->GetChannels();
+		for (int32 ChannelIndex = 0; ChannelIndex < Channels.Num(); ++ChannelIndex)
+		{
+			if (Channels[ChannelIndex].bIsEnabled)
+			{
+				if (!bAudioChannelIsEnabled && Channels[ChannelIndex].Name.Compare(AudioChannelName) == 0)
+				{
+					bAudioChannelIsEnabled = true;
+				}
+				else if (!bAudioMixerChannelIsEnabled && Channels[ChannelIndex].Name.Compare(AudioMixerChannelName) == 0)
+				{
+					bAudioMixerChannelIsEnabled = true;
+				}
+
+				if (bAudioChannelIsEnabled && bAudioMixerChannelIsEnabled)
+				{
+					return true;
+				}
+			}
+		}
+
+		return false;
+	}
+#endif // !WITH_EDITOR
+
 	void FTraceModule::OnAnalysisBegin(TraceServices::IAnalysisSession& InSession)
 	{
 		for (const auto& [ProviderName, Provider] : TraceProviders)
@@ -136,11 +240,12 @@ namespace UE::Audio::Insights
 
 	void FTraceModule::StartTraceAnalysis(const bool bOnlyTraceAudioChannels)
 	{
-		if (bAudioTracesAreAcive)
+		if (bTraceAnalysisHasStarted)
 		{
 			return;
 		}
 
+#if WITH_EDITOR
 		// We don't want to stop any traces that were started prior to running Audio Insights
 		// Check the original state here
 		bStopTracingAfterAudioInsightsIsFinished = !FTraceAuxiliary::IsConnected();
@@ -152,7 +257,7 @@ namespace UE::Audio::Insights
 			DisableAllTraceChannels();
 		}
 
-		EnableAudioInsightsTraceChannels();
+		bTraceAnalysisHasStarted = EnableAudioInsightsTraceChannels();
 
 		if (!FTraceAuxiliary::IsConnected())
 		{
@@ -166,8 +271,9 @@ namespace UE::Audio::Insights
 		// Even if a trace was already active, we still need to notify Audio Insights that analysis has begun
 		IUnrealInsightsModule& InsightsModule = FModuleManager::GetModuleChecked<IUnrealInsightsModule>("TraceInsights");
 		InsightsModule.StartAnalysisForLastLiveSession();
-
-		bAudioTracesAreAcive = true;
+#else
+		bTraceAnalysisHasStarted = EnableAudioInsightsTraceChannels();
+#endif // WITH_EDITOR
 	}
 
 	void FTraceModule::StopTraceAnalysis()
@@ -185,12 +291,12 @@ namespace UE::Audio::Insights
 		}
 
 		ChannelsToRestore.Empty();
-		bAudioTracesAreAcive = false;
+		bTraceAnalysisHasStarted = false;
 	}
 
 	void FTraceModule::OnOnlyTraceAudioChannelsStateChanged(const bool bOnlyTraceAudioChannels)
 	{
-		if (!bAudioTracesAreAcive)
+		if (!bTraceAnalysisHasStarted)
 		{
 			return;
 		}
@@ -210,9 +316,56 @@ namespace UE::Audio::Insights
 
 	bool FTraceModule::IsTraceAnalysisActive() const
 	{
-		return FTraceAuxiliary::IsConnected() && bAudioTracesAreAcive;
+		return FTraceAuxiliary::IsConnected() && bTraceAnalysisHasStarted;
+	}
+
+	bool FTraceModule::AudioChannelsCanBeManuallyEnabled() const
+	{
+#if WITH_EDITOR
+		return false;
+#else
+		// When attaching to a packaged build, we need to check whether we're connected and the audio channels are enabled
+		return !bTraceAnalysisHasStarted && GetCurrentAnalysisSessionInfo() != nullptr && !GetAudioTracesAreEnabled();
+#endif // WITH_EDITOR
 	}
 
+#if !WITH_EDITOR
+	void FTraceModule::InitializeSessionInfo(const TraceServices::FSessionInfo& SessionInfo)
+	{
+		InstanceID = SessionInfo.InstanceId;
+
+		// We need to send discovery requests to ITraceController so we can update the active channels later
+		ISessionServicesModule& SessionServicesModule = FModuleManager::LoadModuleChecked<ISessionServicesModule>("SessionServices");
+
+		TSharedPtr<ISessionManager> SessionManager = SessionServicesModule.GetSessionManager();
+		TSharedPtr<ITraceController> TraceController = SessionServicesModule.GetTraceController();
+		if (SessionManager.IsValid() && TraceController.IsValid())
+		{
+			const FGuid SessionId = FApp::GetSessionId();
+
+			if (TraceController.IsValid())
+			{
+				TraceController->SendDiscoveryRequest(SessionId, InstanceID);
+				TraceController->SendStatusUpdateRequest();
+			}
+		}
+
+	}
+
+	void FTraceModule::RequestChannelUpdate()
+	{
+		// Request for ITraceController to get the latest active channels list
+		ISessionServicesModule& SessionServicesModule = FModuleManager::LoadModuleChecked<ISessionServicesModule>("SessionServices");
+		TSharedPtr<ITraceController> TraceController = SessionServicesModule.GetTraceController();
+		if (!TraceController.IsValid())
+		{
+			return;
+		}
+
+		TraceController->SendChannelUpdateRequest();
+	}
+#endif // !WITH_EDITOR
+
 	void FTraceModule::GetLoggers(TArray<const TCHAR *>& OutLoggers)
 	{
 		OutLoggers.Add(TEXT("Audio"));
diff --git a/Engine/Plugins/AudioInsights/Source/AudioInsights/Public/AudioInsightsDashboardFactory.h b/Engine/Plugins/AudioInsights/Source/AudioInsights/Public/AudioInsightsDashboardFactory.h
--- a/Engine/Plugins/AudioInsights/Source/AudioInsights/Public/AudioInsightsDashboardFactory.h
+++ b/Engine/Plugins/AudioInsights/Source/AudioInsights/Public/AudioInsightsDashboardFactory.h
@@ -31,6 +31,12 @@ namespace UE::Audio::Insights
 		void RegisterTabSpawners();
 		void UnregisterTabSpawners();
 
+#if !WITH_EDITOR
+		TSharedRef<SWidget> MakeEnableTracesOverlay();
+		TSharedRef<SWidget> MakeEnableTracesButton();
+		FReply ToggleAutoEnableAudioTraces();
+#endif
+
 		TSharedPtr<FTabManager> DashboardTabManager;
 		TSharedPtr<FWorkspaceItem> DashboardWorkspace;
 		TSharedPtr<FTabManager::FLayout> TabLayout;
diff --git a/Engine/Plugins/AudioInsights/Source/AudioInsights/Public/AudioInsightsTraceModule.h b/Engine/Plugins/AudioInsights/Source/AudioInsights/Public/AudioInsightsTraceModule.h
--- a/Engine/Plugins/AudioInsights/Source/AudioInsights/Public/AudioInsightsTraceModule.h
+++ b/Engine/Plugins/AudioInsights/Source/AudioInsights/Public/AudioInsightsTraceModule.h
@@ -3,7 +3,9 @@
 
 #include "Containers/Map.h"
 #include "IAudioInsightsTraceModule.h"
+#include "ITraceController.h"
 #include "Templates/SharedPointer.h"
+#include "Trace/StoreClient.h"
 #include "TraceServices/Model/AnalysisSession.h"
 #include "UObject/NameTypes.h"
 
@@ -38,6 +40,13 @@ namespace UE::Audio::Insights
 		virtual void StopTraceAnalysis() override;
 		virtual void OnOnlyTraceAudioChannelsStateChanged(const bool bOnlyTraceAudioChannels) override;
 
+		virtual bool AudioChannelsCanBeManuallyEnabled() const override;
+
+#if !WITH_EDITOR
+		virtual void InitializeSessionInfo(const TraceServices::FSessionInfo& SessionInfo) override;
+		virtual void RequestChannelUpdate() override;
+#endif
+
 		DECLARE_MULTICAST_DELEGATE_OneParam(FOnAnalysisStarting, const double /*Timestamp*/);
 		FOnAnalysisStarting OnAnalysisStarting;
 
@@ -45,18 +54,27 @@ namespace UE::Audio::Insights
 		static const FName GetName();
 
 		static void DisableAllTraceChannels();
-		static void EnableAudioInsightsTraceChannels();
+		bool EnableAudioInsightsTraceChannels();
 
 		void DisableAudioInsightsTraceChannels() const;
 
 		void CacheCurrentlyEnabledTraceChannels();
 		void RestoreCachedChannels() const;
 
+#if !WITH_EDITOR
+		const UE::Trace::FStoreClient::FSessionInfo* GetCurrentAnalysisSessionInfo() const;
+		bool GetAudioTracesAreEnabled() const;
+#endif
+
 		TSharedRef<FTraceChannelManager> ChannelManager;
 		TMap<FName, TSharedPtr<FTraceProviderBase>> TraceProviders;
 		TArray<FString> ChannelsToRestore;
 
-		bool bAudioTracesAreAcive = false;
+		bool bTraceAnalysisHasStarted = false;
 		bool bStopTracingAfterAudioInsightsIsFinished = false;
+
+#if !WITH_EDITOR
+		FGuid InstanceID;
+#endif
 	};
 } // namespace UE::Audio::Insights
diff --git a/Engine/Plugins/AudioInsights/Source/AudioInsights/Public/IAudioInsightsTraceModule.h b/Engine/Plugins/AudioInsights/Source/AudioInsights/Public/IAudioInsightsTraceModule.h
--- a/Engine/Plugins/AudioInsights/Source/AudioInsights/Public/IAudioInsightsTraceModule.h
+++ b/Engine/Plugins/AudioInsights/Source/AudioInsights/Public/IAudioInsightsTraceModule.h
@@ -10,6 +10,11 @@ namespace UE::Audio::Insights
 	class FTraceProviderBase;
 } // namespace UE::Audio::Insights
 
+namespace TraceServices
+{
+	struct FSessionInfo;
+} // namespace TraceServices
+
 class AUDIOINSIGHTS_API IAudioInsightsTraceModule : public TraceServices::IModule
 {
 public:
@@ -21,5 +26,12 @@ public:
 	
 	virtual void OnOnlyTraceAudioChannelsStateChanged(const bool bOnlyTraceAudioChannels) = 0;
 
+	virtual bool AudioChannelsCanBeManuallyEnabled() const = 0;
+
 	virtual TSharedRef<UE::Audio::Insights::FTraceChannelManager> GetChannelManager() = 0;
+
+#if !WITH_EDITOR
+	virtual void InitializeSessionInfo(const TraceServices::FSessionInfo& SessionInfo) = 0;
+	virtual void RequestChannelUpdate() = 0;
+#endif // !WITH_EDITOR
 };
