[Audio Insights] Fix for the build - ensure all control paths return a value from lamda

diff --git a/Engine/Plugins/AudioInsights/Source/AudioInsights/Private/AudioInsightsTraceModule.cpp b/Engine/Plugins/AudioInsights/Source/AudioInsights/Private/AudioInsightsTraceModule.cpp
--- a/Engine/Plugins/AudioInsights/Source/AudioInsights/Private/AudioInsightsTraceModule.cpp
+++ b/Engine/Plugins/AudioInsights/Source/AudioInsights/Private/AudioInsightsTraceModule.cpp
@@ -108,7 +108,6 @@ namespace UE::Audio::Insights
 				if (Status.bIsTracing)
 				{
 					Commands.SetChannels(AudioChannels, EmptyArray);
-					return true;
 				}
 			});
 		}
