[AudioInsights] SoundAsset entries with empty name are displayed in the Sources tab (includes bonus fix for the same issue in the Virtual Loops tab)

- Removed original temp fix which would filter out any empty named sources, but wasn't a fix for the source of the issue (TableDashboardFactory)
- Sources Tab: Issue was due to the order in which we receive messages. In rare cases when a lot of activity is happening, we can process a parameter message (such as volume) after we have processed the stop message, possibly due to message queues getting filled up. Instead of allowing these parameter messages to create a device entry inside of MixerSourceTraceProvider, we now only let them query if one already exists. If it doesn't, this mixer source has already stopped, and there is no need to create a new entry for it.
- Virtual Loop Tab: The same fix has been applied to Virtual Loops with some extra work - The message FVirtualLoopRealizeMessage and FVirtualLoopStopMessage essentially perform the same function (remove the entry from the list). Realize however was being spammed every frame in situations where sounds are virtualized due to concurrency rules. This was causing entries to appear for one frame before being replaced with an unnamed version. Whether a virtual loop is realized or stopped, the stop message appears to fire in either case. The realize message has been dropped as a result.

diff --git a/Engine/Plugins/AudioInsights/Source/AudioInsights/Private/Providers/MixerSourceTraceProvider.cpp b/Engine/Plugins/AudioInsights/Source/AudioInsights/Private/Providers/MixerSourceTraceProvider.cpp
--- a/Engine/Plugins/AudioInsights/Source/AudioInsights/Private/Providers/MixerSourceTraceProvider.cpp
+++ b/Engine/Plugins/AudioInsights/Source/AudioInsights/Private/Providers/MixerSourceTraceProvider.cpp
@@ -253,7 +253,7 @@ namespace UE::Audio::Insights
 
 	bool FMixerSourceTraceProvider::ProcessMessages()
 	{
-		auto BumpEntryFunc = [this](const FMixerSourceMessageBase& Msg)
+		auto GetOrCreateEntry = [this](const FMixerSourceMessageBase& Msg)
 		{
 			TSharedPtr<FMixerSourceDashboardEntry>* ToReturn = nullptr;
 			UpdateDeviceEntry(Msg.DeviceId, Msg.PlayOrder, [&ToReturn, &Msg](TSharedPtr<FMixerSourceDashboardEntry>& Entry)
@@ -271,7 +271,12 @@ namespace UE::Audio::Insights
 			return ToReturn;
 		};
 
-		ProcessMessageQueue<FMixerSourceStartMessage>(TraceMessages.StartMessages, BumpEntryFunc,
+		auto GetEntry = [this](const FMixerSourceMessageBase& Msg)
+		{
+			return FindDeviceEntry(Msg.DeviceId, Msg.PlayOrder);
+		};
+
+		ProcessMessageQueue<FMixerSourceStartMessage>(TraceMessages.StartMessages, GetOrCreateEntry,
 		[this](const FMixerSourceStartMessage& Msg, TSharedPtr<FMixerSourceDashboardEntry>* OutEntry)
 		{
 #if !WITH_EDITOR
@@ -289,9 +294,14 @@ namespace UE::Audio::Insights
 
 		TArray<int32, TInlineAllocator<64>> EntriesWithPoppedDataPoints;
 
-		ProcessMessageQueue<FMixerSourceVolumeMessage>(TraceMessages.VolumeMessages, BumpEntryFunc,
+		ProcessMessageQueue<FMixerSourceVolumeMessage>(TraceMessages.VolumeMessages, GetEntry,
 		[this, &EntriesWithPoppedDataPoints](const FMixerSourceVolumeMessage& Msg, TSharedPtr<FMixerSourceDashboardEntry>* OutEntry)
 		{
+			if (OutEntry == nullptr)
+			{
+				return;
+			}
+
 #if !WITH_EDITOR
 			if (SessionCachedMessages.IsValid())
 			{
@@ -309,9 +319,14 @@ namespace UE::Audio::Insights
 		});
 
 		EntriesWithPoppedDataPoints.Reset();
-		ProcessMessageQueue<FMixerSourcePitchMessage>(TraceMessages.PitchMessages, BumpEntryFunc,
+		ProcessMessageQueue<FMixerSourcePitchMessage>(TraceMessages.PitchMessages, GetEntry,
 		[this, &EntriesWithPoppedDataPoints](const FMixerSourcePitchMessage& Msg, TSharedPtr<FMixerSourceDashboardEntry>* OutEntry)
 		{
+			if (OutEntry == nullptr)
+			{
+				return;
+			}
+
 #if !WITH_EDITOR
 			if (SessionCachedMessages.IsValid())
 			{
@@ -329,9 +344,14 @@ namespace UE::Audio::Insights
 		});
 
 		EntriesWithPoppedDataPoints.Reset();
-		ProcessMessageQueue<FMixerSourceLPFFreqMessage>(TraceMessages.LPFFreqMessages, BumpEntryFunc,
+		ProcessMessageQueue<FMixerSourceLPFFreqMessage>(TraceMessages.LPFFreqMessages, GetEntry,
 		[this, &EntriesWithPoppedDataPoints](const FMixerSourceLPFFreqMessage& Msg, TSharedPtr<FMixerSourceDashboardEntry>* OutEntry)
 		{
+			if (OutEntry == nullptr)
+			{
+				return;
+			}
+
 #if !WITH_EDITOR
 			if (SessionCachedMessages.IsValid())
 			{
@@ -349,9 +369,14 @@ namespace UE::Audio::Insights
 		});
 
 		EntriesWithPoppedDataPoints.Reset();
-		ProcessMessageQueue<FMixerSourceHPFFreqMessage>(TraceMessages.HPFFreqMessages, BumpEntryFunc,
+		ProcessMessageQueue<FMixerSourceHPFFreqMessage>(TraceMessages.HPFFreqMessages, GetEntry,
 		[this, &EntriesWithPoppedDataPoints](const FMixerSourceHPFFreqMessage& Msg, TSharedPtr<FMixerSourceDashboardEntry>* OutEntry)
 		{
+			if (OutEntry == nullptr)
+			{
+				return;
+			}
+
 #if !WITH_EDITOR
 			if (SessionCachedMessages.IsValid())
 			{
@@ -369,9 +394,14 @@ namespace UE::Audio::Insights
 		});
 
 		EntriesWithPoppedDataPoints.Reset();
-		ProcessMessageQueue<FMixerSourceEnvelopeMessage>(TraceMessages.EnvelopeMessages, BumpEntryFunc,
+		ProcessMessageQueue<FMixerSourceEnvelopeMessage>(TraceMessages.EnvelopeMessages, GetEntry,
 		[this, &EntriesWithPoppedDataPoints](const FMixerSourceEnvelopeMessage& Msg, TSharedPtr<FMixerSourceDashboardEntry>* OutEntry)
 		{
+			if (OutEntry == nullptr)
+			{
+				return;
+			}
+
 #if !WITH_EDITOR
 			if (SessionCachedMessages.IsValid())
 			{
@@ -389,9 +419,14 @@ namespace UE::Audio::Insights
 		});
 
 		EntriesWithPoppedDataPoints.Reset();
-		ProcessMessageQueue<FMixerSourceDistanceAttenuationMessage>(TraceMessages.DistanceAttenuationMessages, BumpEntryFunc,
+		ProcessMessageQueue<FMixerSourceDistanceAttenuationMessage>(TraceMessages.DistanceAttenuationMessages, GetEntry,
 		[this, &EntriesWithPoppedDataPoints](const FMixerSourceDistanceAttenuationMessage& Msg, TSharedPtr<FMixerSourceDashboardEntry>* OutEntry)
 		{
+			if (OutEntry == nullptr)
+			{
+				return;
+			}
+
 #if !WITH_EDITOR
 			if (SessionCachedMessages.IsValid())
 			{
@@ -408,11 +443,6 @@ namespace UE::Audio::Insights
 			(*OutEntry)->DistanceAttenuationDataPoints.Push({ Msg.Timestamp, Msg.DistanceAttenuation });
 		});
 
-		auto GetEntry = [this](const FMixerSourceMessageBase& Msg)
-		{
-			return FindDeviceEntry(Msg.DeviceId, Msg.PlayOrder);
-		};
-
 		ProcessMessageQueue<FMixerSourceStopMessage>(TraceMessages.StopMessages, GetEntry,
 		[this](const FMixerSourceStopMessage& Msg, TSharedPtr<FMixerSourceDashboardEntry>* OutEntry)
 		{
diff --git a/Engine/Plugins/AudioInsights/Source/AudioInsights/Private/Providers/VirtualLoopTraceProvider.cpp b/Engine/Plugins/AudioInsights/Source/AudioInsights/Private/Providers/VirtualLoopTraceProvider.cpp
--- a/Engine/Plugins/AudioInsights/Source/AudioInsights/Private/Providers/VirtualLoopTraceProvider.cpp
+++ b/Engine/Plugins/AudioInsights/Source/AudioInsights/Private/Providers/VirtualLoopTraceProvider.cpp
@@ -150,7 +150,7 @@ namespace UE::Audio::Insights
 
 	bool FVirtualLoopTraceProvider::ProcessMessages()
 	{
-		auto RemoveEntryFunc = [this](const FVirtualLoopRealizeMessage& Msg, TSharedPtr<FVirtualLoopDashboardEntry>* OutEntry)
+		auto RemoveEntryFunc = [this](const FVirtualLoopStopOrRealizeMessage& Msg, TSharedPtr<FVirtualLoopDashboardEntry>* OutEntry)
 		{
 #if !WITH_EDITOR
 			if (SessionCachedMessages.IsValid())
@@ -205,10 +205,8 @@ namespace UE::Audio::Insights
 			EntryRef.ComponentId = Msg.ComponentId;
 		});
 
-		ProcessMessageQueue<FVirtualLoopStopMessage>(TraceMessages.StopMessages, GetEntryFunc, RemoveEntryFunc);
-
-		ProcessMessageQueue<FVirtualLoopUpdateMessage>(TraceMessages.UpdateMessages, 
-		BumpEntryFunc,
+		ProcessMessageQueue<FVirtualLoopUpdateMessage>(TraceMessages.UpdateMessages,
+		GetEntryFunc,
 		[this](const FVirtualLoopUpdateMessage& Msg, TSharedPtr<FVirtualLoopDashboardEntry>* OutEntry)
 		{
 #if !WITH_EDITOR
@@ -229,7 +227,7 @@ namespace UE::Audio::Insights
 			}
 		});
 
-		ProcessMessageQueue<FVirtualLoopRealizeMessage>(TraceMessages.RealizeMessages, BumpEntryFunc, RemoveEntryFunc);
+		ProcessMessageQueue<FVirtualLoopStopOrRealizeMessage>(TraceMessages.StopOrRealizeMessages, GetEntryFunc, RemoveEntryFunc);
 
 		return true;
 	}
@@ -257,8 +255,7 @@ namespace UE::Audio::Insights
 				FTraceAnalyzerBase::OnAnalysisBegin(Context);
 
 				UE::Trace::IAnalyzer::FInterfaceBuilder& Builder = Context.InterfaceBuilder;
-				Builder.RouteEvent(RouteId_Realize, "Audio", "VirtualLoopRealize");
-				Builder.RouteEvent(RouteId_Stop, "Audio", "VirtualLoopStop");
+				Builder.RouteEvent(RouteId_Stop, "Audio", "VirtualLoopStopOrRealize");
 				Builder.RouteEvent(RouteId_Update, "Audio", "VirtualLoopUpdate");
 				Builder.RouteEvent(RouteId_Virtualize, "Audio", "VirtualLoopVirtualize");
 			}
@@ -270,15 +267,9 @@ namespace UE::Audio::Insights
 				FVirtualLoopMessages& Messages = GetProvider<FVirtualLoopTraceProvider>().TraceMessages;
 				switch (RouteId)
 				{
-					case RouteId_Realize:
-					{
-						Messages.RealizeMessages.Enqueue(FVirtualLoopRealizeMessage { Context });
-						break;
-					}
-
 					case RouteId_Stop:
 					{
-						Messages.StopMessages.Enqueue(FVirtualLoopStopMessage { Context });
+						Messages.StopOrRealizeMessages.Enqueue(FVirtualLoopStopOrRealizeMessage { Context });
 						break;
 					}
 
@@ -315,7 +306,6 @@ namespace UE::Audio::Insights
 			{
 				RouteId_Virtualize,
 				RouteId_Update,
-				RouteId_Realize,
 				RouteId_Stop
 			};
 
diff --git a/Engine/Plugins/AudioInsights/Source/AudioInsights/Public/Messages/VirtualLoopTraceMessages.h b/Engine/Plugins/AudioInsights/Source/AudioInsights/Public/Messages/VirtualLoopTraceMessages.h
--- a/Engine/Plugins/AudioInsights/Source/AudioInsights/Public/Messages/VirtualLoopTraceMessages.h
+++ b/Engine/Plugins/AudioInsights/Source/AudioInsights/Public/Messages/VirtualLoopTraceMessages.h
@@ -30,8 +30,7 @@ namespace UE::Audio::Insights
 		double Timestamp = 0.0;
 	};
 
-	using FVirtualLoopRealizeMessage = FVirtualLoopMessageBase;
-	using FVirtualLoopStopMessage = FVirtualLoopMessageBase;
+	using FVirtualLoopStopOrRealizeMessage = FVirtualLoopMessageBase;
 
 	struct FVirtualLoopVirtualizeMessage : public FVirtualLoopMessageBase
 	{
@@ -106,8 +105,7 @@ namespace UE::Audio::Insights
 	class FVirtualLoopMessages
 	{
 		TAnalyzerMessageQueue<FVirtualLoopVirtualizeMessage> VirtualizeMessages { 5.0 };
-		TAnalyzerMessageQueue<FVirtualLoopRealizeMessage> RealizeMessages { 5.0 };
-		TAnalyzerMessageQueue<FVirtualLoopStopMessage> StopMessages { 5.0 };
+		TAnalyzerMessageQueue<FVirtualLoopStopOrRealizeMessage> StopOrRealizeMessages { 5.0 };
 		TAnalyzerMessageQueue<FVirtualLoopUpdateMessage> UpdateMessages { 0.1 };
 
 		friend class FVirtualLoopTraceProvider;
@@ -125,7 +123,7 @@ namespace UE::Audio::Insights
 		}
 
 		TraceServices::TPagedArray<FVirtualLoopVirtualizeMessage> VirtualizeCachedMessages;
-		TraceServices::TPagedArray<FVirtualLoopStopMessage> StopOrRealizeCachedMessages;
+		TraceServices::TPagedArray<FVirtualLoopStopOrRealizeMessage> StopOrRealizeCachedMessages;
 		TraceServices::TPagedArray<FVirtualLoopUpdateMessage> UpdateCachedMessages;
 	};
 #endif // !WITH_EDITOR
diff --git a/Engine/Plugins/AudioInsights/Source/AudioInsights/Public/Views/TableDashboardViewFactory.h b/Engine/Plugins/AudioInsights/Source/AudioInsights/Public/Views/TableDashboardViewFactory.h
--- a/Engine/Plugins/AudioInsights/Source/AudioInsights/Public/Views/TableDashboardViewFactory.h
+++ b/Engine/Plugins/AudioInsights/Source/AudioInsights/Public/Views/TableDashboardViewFactory.h
@@ -100,14 +100,6 @@ namespace UE::Audio::Insights
 						};
 						Algo::TransformIf(*DeviceData, DataViewEntries, FilterEntry, TransformEntry);
 					}
-
-					// Temp fix for UE-201431 (SoundAsset entries with empty name are displayed in the Sources tab)
-					DataViewEntries.RemoveAll([](const TSharedPtr<IDashboardDataViewEntry>& Item)
-					{
-						const TSharedPtr<FSoundAssetDashboardEntry> SoundAssetDashboardEntry = StaticCastSharedPtr<FSoundAssetDashboardEntry>(Item);
-
-						return SoundAssetDashboardEntry.IsValid() && SoundAssetDashboardEntry->Name.IsEmpty();
-					});
 					
 					// Sort list
 					RequestSort();
diff --git a/Engine/Source/Runtime/Engine/Private/AudioDevice.cpp b/Engine/Source/Runtime/Engine/Private/AudioDevice.cpp
--- a/Engine/Source/Runtime/Engine/Private/AudioDevice.cpp
+++ b/Engine/Source/Runtime/Engine/Private/AudioDevice.cpp
@@ -240,7 +240,7 @@ FAutoConsoleVariableRef CVarEnableRelativeRenderCostVoiceLimit(
 	ECVF_Default);
 
 #if UE_AUDIO_PROFILERTRACE_ENABLED
-UE_TRACE_EVENT_BEGIN(Audio, VirtualLoopStop)
+UE_TRACE_EVENT_BEGIN(Audio, VirtualLoopStopOrRealize)
 	UE_TRACE_EVENT_FIELD(uint32, DeviceId)
 	UE_TRACE_EVENT_FIELD(uint64, Timestamp)
 	UE_TRACE_EVENT_FIELD(uint32, PlayOrder)
@@ -5511,10 +5511,10 @@ bool FAudioDevice::RemoveVirtualLoop(FActiveSound& InActiveSound)
 		const bool bChannelEnabled = UE_TRACE_CHANNELEXPR_IS_ENABLED(AudioChannel);
 		if (bChannelEnabled)
 		{
-			UE_TRACE_LOG(Audio, VirtualLoopStop, AudioChannel)
-				<< VirtualLoopStop.DeviceId(static_cast<uint32>(DeviceID))
-				<< VirtualLoopStop.Timestamp(FPlatformTime::Cycles64())
-				<< VirtualLoopStop.PlayOrder(InActiveSound.GetPlayOrder());
+			UE_TRACE_LOG(Audio, VirtualLoopStopOrRealize, AudioChannel)
+				<< VirtualLoopStopOrRealize.DeviceId(static_cast<uint32>(DeviceID))
+				<< VirtualLoopStopOrRealize.Timestamp(FPlatformTime::Cycles64())
+				<< VirtualLoopStopOrRealize.PlayOrder(InActiveSound.GetPlayOrder());
 		}
 #endif // UE_AUDIO_PROFILERTRACE_ENABLED
 		VirtualLoops.Remove(&InActiveSound);
diff --git a/Engine/Source/Runtime/Engine/Private/AudioVirtualLoop.cpp b/Engine/Source/Runtime/Engine/Private/AudioVirtualLoop.cpp
--- a/Engine/Source/Runtime/Engine/Private/AudioVirtualLoop.cpp
+++ b/Engine/Source/Runtime/Engine/Private/AudioVirtualLoop.cpp
@@ -51,12 +51,6 @@ UE_TRACE_EVENT_BEGIN(Audio, VirtualLoopVirtualize)
 	UE_TRACE_EVENT_FIELD(UE::Trace::WideString, Name)
 UE_TRACE_EVENT_END()
 
-UE_TRACE_EVENT_BEGIN(Audio, VirtualLoopRealize)
-	UE_TRACE_EVENT_FIELD(uint32, DeviceId)
-	UE_TRACE_EVENT_FIELD(uint64, Timestamp)
-	UE_TRACE_EVENT_FIELD(uint32, PlayOrder)
-UE_TRACE_EVENT_END()
-
 UE_TRACE_EVENT_BEGIN(Audio, VirtualLoopUpdate)
 	UE_TRACE_EVENT_FIELD(uint32, DeviceId)
 	UE_TRACE_EVENT_FIELD(double, Timestamp)
@@ -340,20 +334,6 @@ bool FAudioVirtualLoop::Update(float DeltaTime, bool bForceUpdate)
 		return false;
 	}
 
-
-
-#if UE_AUDIO_PROFILERTRACE_ENABLED
-	if (bChannelEnabled)
-	{
-		if (const FAudioDevice* AudioDevice = ActiveSound->AudioDevice)
-		{
-			UE_TRACE_LOG(Audio, VirtualLoopRealize, AudioChannel)
-				<< VirtualLoopRealize.DeviceId(static_cast<uint32>(AudioDevice->DeviceID))
-				<< VirtualLoopRealize.Timestamp(FPlatformTime::Cycles64())
-				<< VirtualLoopRealize.PlayOrder(ActiveSound->GetPlayOrder());
-		}
-	}
-#endif // UE_AUDIO_PROFILERTRACE_ENABLED
 	return true;
 }
 
