Fix for compile errors in Module.AudioInsights.cpp and AudioInsightsTraceModule.cpp

diff --git a/Engine/Plugins/AudioInsights/Source/AudioInsights/Private/AudioInsightsTraceModule.cpp b/Engine/Plugins/AudioInsights/Source/AudioInsights/Private/AudioInsightsTraceModule.cpp
--- a/Engine/Plugins/AudioInsights/Source/AudioInsights/Private/AudioInsightsTraceModule.cpp
+++ b/Engine/Plugins/AudioInsights/Source/AudioInsights/Private/AudioInsightsTraceModule.cpp
@@ -26,6 +26,10 @@ namespace UE::Audio::Insights
 
 	FTraceModule::FTraceModule()
 		: ChannelManager(MakeShared<FTraceChannelManager>())
+#if !WITH_EDITOR
+		, AudioChannels({ FTraceModulePrivate::AudioChannelName, FTraceModulePrivate::AudioMixerChannelName } )
+		, EmptyArray()
+#endif
 	{
 		// Don't run providers in any commandlet to avoid additional, unnecessary overhead as audio insights is dormant.
 		if (!IsRunningCommandlet())
@@ -97,9 +101,6 @@ namespace UE::Audio::Insights
 			return false;
 		}
 
-		static TConstArrayView<FStringView> AudioChannels{ AudioChannelName, AudioMixerChannelName };
-		static TConstArrayView<FStringView> EmptyArray;
-
 		// Requires ITraceController to have received a response to discovering the session and channel info
 		if (TraceController->HasAvailableInstance(InstanceID))
 		{
diff --git a/Engine/Plugins/AudioInsights/Source/AudioInsights/Public/AudioInsightsTraceModule.h b/Engine/Plugins/AudioInsights/Source/AudioInsights/Public/AudioInsightsTraceModule.h
--- a/Engine/Plugins/AudioInsights/Source/AudioInsights/Public/AudioInsightsTraceModule.h
+++ b/Engine/Plugins/AudioInsights/Source/AudioInsights/Public/AudioInsightsTraceModule.h
@@ -75,6 +75,9 @@ namespace UE::Audio::Insights
 
 #if !WITH_EDITOR
 		FGuid InstanceID;
+
+		const TArray<FString> AudioChannels;
+		const TArray<FString> EmptyArray;
 #endif
 	};
 } // namespace UE::Audio::Insights
