[Audio Insights] Sound Dashboard part 1: added sound dashboard icons + register icons in Audio Insights style.

diff --git a/Engine/Plugins/AudioInsights/Source/AudioInsights/Public/AudioInsightsStyle.h b/Engine/Plugins/AudioInsights/Source/AudioInsights/Public/AudioInsightsStyle.h
--- a/Engine/Plugins/AudioInsights/Source/AudioInsights/Public/AudioInsightsStyle.h
+++ b/Engine/Plugins/AudioInsights/Source/AudioInsights/Public/AudioInsightsStyle.h
@@ -147,6 +147,21 @@ namespace UE::Audio::Insights
 
 			Set("AudioInsights.Thumbnail", new IMAGE_BRUSH_SVG(TEXT("Icons/audio_insights"), Icon16));
 
+			// Sound Dashboard icons
+			Set("AudioInsights.Icon.SoundDashboard.Browse",           new IMAGE_BRUSH_SVG(TEXT("Icons/SoundDashboard/browse"),           Icon16));
+			Set("AudioInsights.Icon.SoundDashboard.Edit",             new IMAGE_BRUSH_SVG(TEXT("Icons/SoundDashboard/edit"),             Icon16));
+			Set("AudioInsights.Icon.SoundDashboard.Filter",           new IMAGE_BRUSH_SVG(TEXT("Icons/SoundDashboard/filter"),           Icon16));
+			Set("AudioInsights.Icon.SoundDashboard.Hide",             new IMAGE_BRUSH_SVG(TEXT("Icons/SoundDashboard/hide"),             Icon16));
+			Set("AudioInsights.Icon.SoundDashboard.Info",             new IMAGE_BRUSH_SVG(TEXT("Icons/SoundDashboard/info"),             Icon16));
+			Set("AudioInsights.Icon.SoundDashboard.MetaSound",        new IMAGE_BRUSH_SVG(TEXT("Icons/SoundDashboard/metasound"),        Icon16));
+			Set("AudioInsights.Icon.SoundDashboard.Mute",             new IMAGE_BRUSH_SVG(TEXT("Icons/SoundDashboard/mute"),             Icon16));
+			Set("AudioInsights.Icon.SoundDashboard.Pin",              new IMAGE_BRUSH_SVG(TEXT("Icons/SoundDashboard/pin"),              Icon16));
+			Set("AudioInsights.Icon.SoundDashboard.ProceduralSource", new IMAGE_BRUSH_SVG(TEXT("Icons/SoundDashboard/proceduralsource"), Icon16));
+			Set("AudioInsights.Icon.SoundDashboard.Solo",             new IMAGE_BRUSH_SVG(TEXT("Icons/SoundDashboard/solo"),             Icon16));
+			Set("AudioInsights.Icon.SoundDashboard.SoundCue",         new IMAGE_BRUSH_SVG(TEXT("Icons/SoundDashboard/soundcue"),         Icon16));
+			Set("AudioInsights.Icon.SoundDashboard.SoundWave",        new IMAGE_BRUSH_SVG(TEXT("Icons/SoundDashboard/soundwave"),        Icon16));
+			Set("AudioInsights.Icon.SoundDashboard.Tab",              new IMAGE_BRUSH_SVG(TEXT("Icons/SoundDashboard/tab"),              Icon16));
+
 			FSlateStyleRegistry::RegisterSlateStyle(*this);
 		}
 	};
