[Audio Insights] Sound Dashboard part 2: added SoundDashboardCommands (for right click actions)

diff --git a/Engine/Plugins/AudioInsights/Source/AudioInsights/Private/SoundDashboardCommands.cpp b/Engine/Plugins/AudioInsights/Source/AudioInsights/Private/SoundDashboardCommands.cpp
new file mode 100644
--- /dev/null
+++ b/Engine/Plugins/AudioInsights/Source/AudioInsights/Private/SoundDashboardCommands.cpp
@@ -0,0 +1,28 @@
+// Copyright Epic Games, Inc. All Rights Reserved.
+#include "SoundDashboardCommands.h"
+
+#include "AudioInsightsStyle.h"
+#include "Framework/Commands/InputChord.h"
+#include "Framework/Commands/UICommandInfo.h"
+
+#define LOCTEXT_NAMESPACE "AudioInsights"
+
+namespace UE::Audio::Insights
+{
+	FSoundDashboardCommands::FSoundDashboardCommands()
+		: TCommands<FSoundDashboardCommands>("SoundDashboardCommands", LOCTEXT("SoundDashboardCommands_ContextDescText", "Sound Dashboard Commands"), NAME_None, FSlateStyle::GetStyleName())
+	{
+
+	}
+	
+	void FSoundDashboardCommands::RegisterCommands()
+	{
+		UI_COMMAND(Pin, "Pin", "Pins the selected sound in the Pinned category.", EUserInterfaceActionType::Button, FInputChord(EKeys::P, EModifierKey::Control));
+		UI_COMMAND(Browse, "Browse To Asset", "Browses to the selected sound asset in the content browser.", EUserInterfaceActionType::Button, FInputChord(EKeys::B, EModifierKey::Control));
+		UI_COMMAND(Edit, "Edit", "Opens the selected sound for edit.", EUserInterfaceActionType::Button, FInputChord(EKeys::E, EModifierKey::Control));
+		UI_COMMAND(Hide, "Hide", "Hides the selected sound.", EUserInterfaceActionType::Button, FInputChord(EKeys::H, EModifierKey::Control));
+		UI_COMMAND(Info, "More Information", "Provides more information about this sound.", EUserInterfaceActionType::Button, FInputChord(EKeys::I, EModifierKey::Control));
+	}
+} // namespace UE::Audio::Insights
+
+#undef LOCTEXT_NAMESPACE
diff --git a/Engine/Plugins/AudioInsights/Source/AudioInsights/Private/SoundDashboardCommands.h b/Engine/Plugins/AudioInsights/Source/AudioInsights/Private/SoundDashboardCommands.h
new file mode 100644
--- /dev/null
+++ b/Engine/Plugins/AudioInsights/Source/AudioInsights/Private/SoundDashboardCommands.h
@@ -0,0 +1,31 @@
+// Copyright Epic Games, Inc. All Rights Reserved.
+#pragma once
+
+#include "Framework/Commands/Commands.h"
+#include "Templates/SharedPointer.h"
+
+class FUICommandInfo;
+
+namespace UE::Audio::Insights
+{
+	class FSoundDashboardCommands : public TCommands<FSoundDashboardCommands>
+	{
+	public:
+		FSoundDashboardCommands();
+
+		virtual void RegisterCommands() override;
+
+		TSharedPtr<const FUICommandInfo> GetPinCommand() const    { return Pin; }
+		TSharedPtr<const FUICommandInfo> GetBrowseCommand() const { return Browse; }
+		TSharedPtr<const FUICommandInfo> GetEditCommand() const   { return Edit; }
+		TSharedPtr<const FUICommandInfo> GetHideCommand() const   { return Hide; }
+		TSharedPtr<const FUICommandInfo> GetInfoCommand() const   { return Info; }
+		
+	private:
+		TSharedPtr<FUICommandInfo> Pin;
+		TSharedPtr<FUICommandInfo> Browse;
+		TSharedPtr<FUICommandInfo> Edit;
+		TSharedPtr<FUICommandInfo> Hide;
+		TSharedPtr<FUICommandInfo> Info;
+	};
+} // namespace UE::Audio::Insights
