[Audio Insights] Sound Dashboard part 3: added SAudioFilterBar (derived from SBasicFilterBar in order to be able to add custom icons to the filter list, reusable for other Audio Insights dashboards)

diff --git a/Engine/Plugins/AudioInsights/Source/AudioInsights/Public/Views/SAudioFilterBar.h b/Engine/Plugins/AudioInsights/Source/AudioInsights/Public/Views/SAudioFilterBar.h
new file mode 100644
--- /dev/null
+++ b/Engine/Plugins/AudioInsights/Source/AudioInsights/Public/Views/SAudioFilterBar.h
@@ -0,0 +1,111 @@
+// Copyright Epic Games, Inc. All Rights Reserved.
+#pragma once
+
+#include "AudioInsightsStyle.h"
+#include "Filters/FilterBase.h"
+#include "Filters/SBasicFilterBar.h"
+
+namespace UE::Audio::Insights
+{
+	template<typename EnumType>
+	class SAudioFilterBar : public SBasicFilterBar<EnumType>
+	{
+	public:
+		SLATE_BEGIN_ARGS(SAudioFilterBar){}
+			SLATE_ARGUMENT(TArray<TSharedRef<FFilterBase<EnumType>>>, CustomFilters)
+			SLATE_EVENT(SBasicFilterBar<EnumType>::FOnFilterChanged, OnFilterChanged)
+		SLATE_END_ARGS()
+
+		void Construct(const FArguments& InArgs)
+		{
+			SBasicFilterBar<EnumType>::Construct(
+				SBasicFilterBar<EnumType>::FArguments()
+				.CustomFilters(InArgs._CustomFilters)
+				.CanChangeOrientation(false)
+				.FilterBarLayout(EFilterBarLayout::Horizontal)
+				.UseSectionsForCategories(true)
+			);
+		}
+
+	protected:
+		virtual TSharedRef<SWidget> MakeAddFilterMenu() override
+		{
+			TObjectPtr<UToolMenus> ToolsMenus = UToolMenus::Get();
+
+			if (ToolsMenus == nullptr)
+			{
+				return SNullWidget::NullWidget;
+			}
+
+			const FName FilterMenuName = "FilterBar.FilterMenu";
+
+			if (!ToolsMenus->IsMenuRegistered(FilterMenuName))
+			{
+				const TObjectPtr<UToolMenu> ToolMenu = ToolsMenus->RegisterMenu(FilterMenuName);
+
+				ToolMenu->bShouldCloseWindowAfterMenuSelection = true;
+				ToolMenu->bCloseSelfOnly = true;
+
+				ToolMenu->AddDynamicSection(NAME_None, FNewToolMenuDelegate::CreateLambda([](UToolMenu* InMenu)
+				{
+					if (const TObjectPtr<UFilterBarContext> FilterBarContext = InMenu->FindContext<UFilterBarContext>())
+					{
+						FilterBarContext->PopulateFilterMenu.ExecuteIfBound(InMenu);
+						FilterBarContext->OnExtendAddFilterMenu.ExecuteIfBound(InMenu);
+					}
+				}));
+			}
+
+			TObjectPtr<UFilterBarContext> FilterBarContext = NewObject<UFilterBarContext>();
+
+			FilterBarContext->PopulateFilterMenu = FOnPopulateAddFilterMenu::CreateSP(this, &SAudioFilterBar::PopulateAddFilterMenu);
+			FilterBarContext->OnExtendAddFilterMenu = SBasicFilterBar<EnumType>::OnExtendAddFilterMenu;
+
+			FToolMenuContext ToolMenuContext(FilterBarContext);
+
+			return ToolsMenus->GenerateWidget(FilterMenuName, ToolMenuContext);
+		}
+
+		void PopulateAddFilterMenu(UToolMenu* InMenu)
+		{
+			SBasicFilterBar<EnumType>::PopulateCommonFilterSections(InMenu);
+			PopulateCustomFilters(InMenu);
+		}
+
+		void PopulateCustomFilters(UToolMenu* InMenu)
+		{
+			if (!SBasicFilterBar<EnumType>::bUseSectionsForCategories)
+			{
+				return;
+			}
+
+			for (const TSharedPtr<FFilterCategory>& Category : SBasicFilterBar<EnumType>::AllFilterCategories)
+			{
+				if (!Category.IsValid())
+				{
+					continue;
+				}
+
+				FToolMenuSection& Section = InMenu->AddSection(*Category->Title.ToString(), Category->Title);
+
+				for (const TSharedRef<FFilterBase<EnumType>>& FrontendFilter : SBasicFilterBar<EnumType>::AllFrontendFilters)
+				{
+					if (FrontendFilter->GetCategory() == Category)
+					{
+						Section.AddMenuEntry(
+							NAME_None,
+							FrontendFilter->GetDisplayName(),
+							FrontendFilter->GetToolTipText(),
+							FSlateStyle::Get().CreateIcon(FrontendFilter->GetIconName()),
+							FUIAction(
+								FExecuteAction::CreateSP(const_cast<SAudioFilterBar*>(this), &SAudioFilterBar::FrontendFilterClicked, FrontendFilter),
+								FCanExecuteAction(),
+								FIsActionChecked::CreateSP(this, &SAudioFilterBar::IsFrontendFilterInUse, FrontendFilter)),
+							EUserInterfaceActionType::ToggleButton
+						);
+					}
+				}
+			}
+		}
+	};
+} // namespace UE::Audio::Insights
