[Audio Insights] Sound Dashboard part 4: added IDashboardDataTreeViewEntry + deleted old unused FAudioDataTableView

diff --git a/Engine/Plugins/AudioInsights/Source/AudioInsights/Public/AudioInsightsDataSource.h b/Engine/Plugins/AudioInsights/Source/AudioInsights/Public/AudioInsightsDataSource.h
--- a/Engine/Plugins/AudioInsights/Source/AudioInsights/Public/AudioInsightsDataSource.h
+++ b/Engine/Plugins/AudioInsights/Source/AudioInsights/Public/AudioInsightsDataSource.h
@@ -1,12 +1,8 @@
 // Copyright Epic Games, Inc. All Rights Reserved.
 #pragma once
 
-#include "Delegates/Delegate.h"
+#include "Containers/Array.h"
 #include "Templates/SharedPointer.h"
-#include "Trace/Analyzer.h"
-#include "TraceServices/Model/AnalysisSession.h"
-#include "UObject/NameTypes.h"
-
 
 namespace UE::Audio::Insights
 {
@@ -18,38 +14,13 @@ namespace UE::Audio::Insights
 		virtual bool IsValid() const = 0;
 	};
 
-	class AUDIOINSIGHTS_API FAudioDataTableView
-	{ 
+	class AUDIOINSIGHTS_API IDashboardDataTreeViewEntry : public TSharedFromThis<IDashboardDataTreeViewEntry>
+	{
 	public:
-		FAudioDataTableView()
-			: Entries(MakeShared<TArray<TSharedPtr<IDashboardDataViewEntry>>>())
-		{
-		}
-
-		FAudioDataTableView(TSharedRef<TArray<TSharedPtr<IDashboardDataViewEntry>>> InEntries)
-			: Entries(InEntries)
-		{
-		}
-
-		virtual ~FAudioDataTableView() = default;
+		virtual ~IDashboardDataTreeViewEntry() = default;
 
-		TSharedRef<const TArray<TSharedPtr<IDashboardDataViewEntry>>> GetEntries() const
-		{
-			return Entries;
-		}
-
-		virtual bool FilterViewEntries() = 0;
-
-		virtual void ResetViewEntries()
-		{
-			Entries->Empty();
-			OnViewEntriesUpdatedDelegate.Broadcast();
-		}
-
-		DECLARE_MULTICAST_DELEGATE(FOnViewEntriesUpdatedDelegate);
-		FOnViewEntriesUpdatedDelegate OnViewEntriesUpdatedDelegate;
+		virtual bool IsValid() const = 0;
 
-	protected:
-		TSharedRef<TArray<TSharedPtr<IDashboardDataViewEntry>>> Entries;
+		TArray<TSharedPtr<IDashboardDataTreeViewEntry>> Children;
 	};
 } // namespace UE::Audio::Insights
