[Audio Insights] Sound Dashboard part 5: added TreeDashboardViewFactory

diff --git a/Engine/Plugins/AudioInsights/Source/AudioInsights/Private/Views/TreeDashboardViewFactory.cpp b/Engine/Plugins/AudioInsights/Source/AudioInsights/Private/Views/TreeDashboardViewFactory.cpp
new file mode 100644
--- /dev/null
+++ b/Engine/Plugins/AudioInsights/Source/AudioInsights/Private/Views/TreeDashboardViewFactory.cpp
@@ -0,0 +1,314 @@
+// Copyright Epic Games, Inc. All Rights Reserved.
+#include "Views/TreeDashboardViewFactory.h"
+
+#include "Algo/Transform.h"
+#include "AudioInsightsStyle.h"
+#include "AudioInsightsTraceProviderBase.h"
+#include "Containers/Array.h"
+#include "Widgets/Input/SSearchBox.h"
+#include "Widgets/Views/SHeaderRow.h"
+#include "Widgets/Layout/SScrollBox.h"
+
+#if WITH_EDITOR
+#include "AudioDeviceManager.h"
+#endif // WITH_EDITOR
+
+#define LOCTEXT_NAMESPACE "AudioInsights"
+
+namespace UE::Audio::Insights
+{
+	// SRowWidget
+	void FTraceTreeDashboardViewFactory::SRowWidget::Construct(const FArguments& InArgs, const TSharedRef<STableViewBase>& InOwnerTable, TSharedPtr<IDashboardDataTreeViewEntry> InData, TSharedRef<FTraceTreeDashboardViewFactory> InFactory)
+	{
+		Data = InData;
+		Factory = InFactory;
+		SMultiColumnTableRow<TSharedPtr<IDashboardDataTreeViewEntry>>::Construct(FSuperRowType::FArguments(), InOwnerTable);
+	}
+
+	TSharedRef<SWidget> FTraceTreeDashboardViewFactory::SRowWidget::GenerateWidgetForColumn(const FName& Column)
+	{
+		return Factory->GenerateWidgetForColumn(StaticCastSharedRef<SRowWidget>(AsShared()), Data->AsShared(), Column);
+	}
+
+
+	// FTraceTreeDashboardViewFactory
+	FTraceTreeDashboardViewFactory::FTraceTreeDashboardViewFactory()
+	{
+		TickerHandle = FTSTicker::GetCoreTicker().AddTicker(TEXT("TraceTreeDashboardViewFactory"), 0.0f, [this](float DeltaTime)
+		{
+			Tick(DeltaTime);
+			return true;
+		});
+	}
+
+	FTraceTreeDashboardViewFactory::~FTraceTreeDashboardViewFactory()
+	{
+		FTSTicker::GetCoreTicker().RemoveTicker(TickerHandle);
+	}
+
+	TSharedPtr<SWidget> FTraceTreeDashboardViewFactory::GetFilterBarButtonWidget()
+	{
+		return nullptr;
+	}
+
+	TSharedPtr<SWidget> FTraceTreeDashboardViewFactory::GetFilterBarWidget()
+	{
+		return nullptr;
+	}
+
+	TSharedRef<SWidget> FTraceTreeDashboardViewFactory::MakeWidget(TSharedRef<SDockTab> OwnerTab, const FSpawnTabArgs& SpawnTabArgs)
+	{
+		if (!DashboardWidget.IsValid())
+		{
+			const TSharedPtr<SWidget> FilterBar = GetFilterBarWidget();
+			const TSharedPtr<SWidget> FilterBarButton = GetFilterBarButtonWidget();
+
+			SAssignNew(DashboardWidget, SVerticalBox)
+			.Clipping(EWidgetClipping::ClipToBounds)
+			// Filter / Search area
+			+ SVerticalBox::Slot()
+			.AutoHeight()
+			.HAlign(HAlign_Fill)
+			.Padding(0.0f, 0.0f, 0.0f, 3.0f)
+			[
+				SNew(SHorizontalBox)
+				// Filter selector
+				+ SHorizontalBox::Slot()
+				.AutoWidth()
+				.Padding(FilterBarButton ? 3.0f : 0.0f, 0.0f, 0.0f, 0.0f)
+				[
+					FilterBarButton ? FilterBarButton.ToSharedRef() : SNullWidget::NullWidget
+				]
+				// Search box
+				+ SHorizontalBox::Slot()
+				.AutoWidth()
+				[
+					SNew(SVerticalBox)
+					+ SVerticalBox::Slot()
+					.Padding(0.0f, 4.0f, 0.0f, 6.0f)
+					[
+						SAssignNew(SearchBoxWidget, SSearchBox)
+						.SelectAllTextWhenFocused(true)
+						.HintText(LOCTEXT("TreeDashboardView_SearchBoxHintText", "Search"))
+						.MinDesiredWidth(200)
+						.OnTextChanged(this, &FTraceTreeDashboardViewFactory::SetSearchBoxFilterText)
+					]
+				]
+				// Active filters area
+				+ SHorizontalBox::Slot()
+				[
+					FilterBar ? FilterBar.ToSharedRef() : SNullWidget::NullWidget
+				]
+			]
+			// TreeView
+			+ SVerticalBox::Slot()
+			.FillHeight(1.0f)
+			[
+				SNew(SScrollBox)
+				.Orientation(EOrientation::Orient_Horizontal)
+				+ SScrollBox::Slot()
+				.FillSize(1.0f)
+				.HAlign(HAlign_Fill)
+				[
+					SAssignNew(FilteredEntriesListView, STreeView<TSharedPtr<IDashboardDataTreeViewEntry>>)
+					.HeaderRow(MakeHeaderRowWidget())
+					.TreeItemsSource(&DataViewEntries)
+					.OnGenerateRow_Lambda([this](TSharedPtr<IDashboardDataTreeViewEntry> Item, const TSharedRef<STableViewBase>& OwnerTable)
+					{
+						return SNew(SRowWidget, OwnerTable, Item, AsShared());
+					})
+					.OnContextMenuOpening(this, &FTraceTreeDashboardViewFactory::OnConstructContextMenu)
+					.SelectionMode(ESelectionMode::Multi)
+					.OnSelectionChanged(this, &FTraceTreeDashboardViewFactory::OnSelectionChanged)
+					.OnKeyDownHandler(this, &FTraceTreeDashboardViewFactory::OnDataRowKeyInput)
+					.OnGetChildren_Lambda([this](TSharedPtr<IDashboardDataTreeViewEntry> InParent, TArray<TSharedPtr<IDashboardDataTreeViewEntry>>& OutChildren)
+					{
+						if (InParent.IsValid() && !InParent->Children.IsEmpty())
+						{
+							OutChildren = InParent->Children;
+						}
+					})
+				]
+			];
+		}
+
+		return DashboardWidget->AsShared();
+	}
+	
+	TSharedRef<SHeaderRow> FTraceTreeDashboardViewFactory::MakeHeaderRowWidget()
+	{
+		TArray<FName> DefaultHiddenColumns;
+
+		Algo::TransformIf(GetHeaderRowColumns(), DefaultHiddenColumns,
+			[](const TPair<FName, FHeaderRowColumnData>& ColumnInfo) { return ColumnInfo.Value.bDefaultHidden; },
+			[](const TPair<FName, FHeaderRowColumnData>& ColumnInfo) { return ColumnInfo.Key; });
+
+		SAssignNew(HeaderRowWidget, SHeaderRow)
+		.CanSelectGeneratedColumn(true); // Allows for showing/hiding columns
+
+		// This only works if header row columns are added with slots and not programmatically
+		// check in SHeaderRow::Construct: for ( FColumn* const Column : InArgs.Slots ) for more info
+		// A potential alternative would be to delegate to the derived classes the SHeaderRow creation with slots
+		//.HiddenColumnsList(DefaultHiddenColumns);
+
+		for (const auto& [ColumnName, ColumnData] : GetHeaderRowColumns())
+		{
+			const SHeaderRow::FColumn::FArguments ColumnArgs = SHeaderRow::Column(ColumnName)
+				.DefaultLabel(ColumnData.DisplayName)
+				.HAlignCell(ColumnData.Alignment)
+				.FillWidth(ColumnData.FillWidth)
+				.SortMode(this, &FTraceTreeDashboardViewFactory::GetColumnSortMode, ColumnName)
+				.OnSort(this, &FTraceTreeDashboardViewFactory::OnColumnSortModeChanged)
+				.HeaderContent()
+				[
+					SNew(SHorizontalBox)
+					// Icon (optional)
+					+ SHorizontalBox::Slot()
+					.AutoWidth()
+					.Padding(ColumnData.IconName != NAME_None ? 2.0f : 0.0f, 3.0f)
+					[
+						ColumnData.IconName != NAME_None
+							? SNew(SImage).Image(FSlateStyle::Get().GetBrush(ColumnData.IconName)) 
+							: SNullWidget::NullWidget
+					]
+					// Text (optional)
+					+ SHorizontalBox::Slot()
+					.AutoWidth()
+					.Padding(0.0f, 3.0f, 0.0f, 3.0f)
+					[
+						ColumnData.bShowDisplayName 
+							? SNew(STextBlock).Text(ColumnData.DisplayName) 
+							: SNullWidget::NullWidget
+					]
+				];
+
+			// .HiddenColumnsList workaround:
+			// simulate what SHeaderRow::AddColumn( const FColumn::FArguments& NewColumnArgs ) does but allowing us to modify the bIsVisible property
+			// Memory handling (delete) is done by TIndirectArray<FColumn> Columns; defined in SHeaderRow
+			SHeaderRow::FColumn* NewColumn = new SHeaderRow::FColumn(ColumnArgs);
+			NewColumn->bIsVisible = !DefaultHiddenColumns.Contains(ColumnName);
+			HeaderRowWidget->AddColumn(*NewColumn);
+		}
+
+		return HeaderRowWidget.ToSharedRef();
+	}
+
+	TSharedPtr<SWidget> FTraceTreeDashboardViewFactory::OnConstructContextMenu()
+	{
+		// To be optionally implemented by derived classes
+		return SNullWidget::NullWidget;
+	}
+
+	void FTraceTreeDashboardViewFactory::OnSelectionChanged(TSharedPtr<IDashboardDataTreeViewEntry> SelectedItem, ESelectInfo::Type SelectInfo)
+	{
+		// To be optionally implemented by derived classes
+	}
+
+	FReply FTraceTreeDashboardViewFactory::OnDataRowKeyInput(const FGeometry& Geometry, const FKeyEvent& KeyEvent) const
+	{
+		// To be optionally implemented by derived classes
+		return FReply::Unhandled();
+	}
+
+	const FText& FTraceTreeDashboardViewFactory::GetSearchFilterText() const
+	{
+		return SearchBoxFilterText;
+	}
+
+	void FTraceTreeDashboardViewFactory::SetSearchBoxFilterText(const FText& NewText)
+	{
+		SearchBoxFilterText = NewText;
+		UpdateFilterReason = EProcessReason::FilterUpdated;
+	}
+
+	void FTraceTreeDashboardViewFactory::RefreshFilteredEntriesListView()
+	{
+		if (FilteredEntriesListView.IsValid())
+		{
+			FilteredEntriesListView->RequestTreeRefresh();
+		}
+	}
+
+	EColumnSortMode::Type FTraceTreeDashboardViewFactory::GetColumnSortMode(const FName InColumnId) const
+	{
+		return SortByColumn == InColumnId ? SortMode : EColumnSortMode::None;
+	}
+
+	void FTraceTreeDashboardViewFactory::RequestSort()
+	{
+		SortTable();
+
+		if (FilteredEntriesListView.IsValid())
+		{
+			FilteredEntriesListView->RequestTreeRefresh();
+		}
+	}
+
+	void FTraceTreeDashboardViewFactory::OnColumnSortModeChanged(const EColumnSortPriority::Type InSortPriority, const FName& InColumnId, const EColumnSortMode::Type InSortMode)
+	{
+		SortByColumn = InColumnId;
+		SortMode = InSortMode;
+
+		RequestSort();
+	}
+
+	void FTraceTreeDashboardViewFactory::Tick(float InElapsed)
+	{
+		for (const TSharedPtr<FTraceProviderBase>& Provider : Providers)
+		{
+			if (!Provider.IsValid())
+			{
+				continue;
+			}
+
+			if (const uint64* FoundCurrentUpdateId = UpdateIds.Find(Provider->GetName()))
+			{
+				if (*FoundCurrentUpdateId != Provider->GetLastUpdateId())
+				{
+					UpdateFilterReason = EProcessReason::EntriesUpdated;
+				}
+			}
+			else
+			{
+				UpdateFilterReason = EProcessReason::EntriesUpdated;
+			}
+		}
+
+		if (UpdateFilterReason != EProcessReason::None)
+		{
+			ProcessEntries(UpdateFilterReason);
+
+			if (UpdateFilterReason == EProcessReason::EntriesUpdated)
+			{
+				for (const TSharedPtr<FTraceProviderBase>& Provider : Providers)
+				{
+					UpdateIds.FindOrAdd(Provider->GetName()) = Provider->GetLastUpdateId();
+				}
+			}
+
+			RefreshFilteredEntriesListView();
+
+			UpdateFilterReason = EProcessReason::None;
+		}
+
+#if WITH_EDITOR
+		if (IsDebugDrawEnabled())
+		{
+			if (FilteredEntriesListView.IsValid())
+			{
+				TArray<TSharedPtr<IDashboardDataTreeViewEntry>> SelectedItems = FilteredEntriesListView->GetSelectedItems();
+
+				if (FAudioDeviceManager* AudioDeviceManager = FAudioDeviceManager::Get())
+				{
+					AudioDeviceManager->IterateOverAllDevices([this, &SelectedItems, InElapsed](::Audio::FDeviceId DeviceId, FAudioDevice* Device)
+					{
+						DebugDraw(InElapsed, SelectedItems, DeviceId);
+					});
+				}
+			}
+		}
+#endif // WITH_EDITOR
+	}
+} // namespace UE::Audio::Insights
+
+#undef LOCTEXT_NAMESPACE
diff --git a/Engine/Plugins/AudioInsights/Source/AudioInsights/Public/Views/TreeDashboardViewFactory.h b/Engine/Plugins/AudioInsights/Source/AudioInsights/Public/Views/TreeDashboardViewFactory.h
new file mode 100644
--- /dev/null
+++ b/Engine/Plugins/AudioInsights/Source/AudioInsights/Public/Views/TreeDashboardViewFactory.h
@@ -0,0 +1,179 @@
+// Copyright Epic Games, Inc. All Rights Reserved.
+#pragma once
+
+#include "AudioDefines.h"
+#include "AudioInsightsDataSource.h"
+#include "Templates/SharedPointer.h"
+#include "Views/DashboardViewFactory.h"
+#include "Widgets/Views/STableRow.h"
+#include "Widgets/Views/STreeView.h"
+
+class SHeaderRow;
+class SSearchBox;
+class STableViewBase;
+
+namespace UE::Audio::Insights
+{
+	/**
+	 * Tree view entries can inherit from this class to implement extra UObject functionality (ex: open, browse, edit, etc.) 
+	 */
+	class AUDIOINSIGHTS_API IObjectTreeDashboardEntry : public IDashboardDataTreeViewEntry
+	{
+	public:
+		virtual ~IObjectTreeDashboardEntry() = default;
+
+		virtual TObjectPtr<UObject> GetObject() = 0;
+		virtual const TObjectPtr<UObject> GetObject() const = 0;
+
+		virtual FText GetDisplayName() const = 0;
+	};
+
+	/** 
+	 * Inherit from this class to create a tree view dashboard for Audio Insights.
+	 * It contains a search textbox, filters can be optionally be implemented via GetFilterBarWidget and GetFilterBarButtonWidget.
+	 * Item actions can be done via OnSelectionChanged, OnDataRowKeyInput, OnConstructContextMenu (for right click)
+	 */
+	class AUDIOINSIGHTS_API FTraceTreeDashboardViewFactory : public FTraceDashboardViewFactoryBase, public TSharedFromThis<FTraceTreeDashboardViewFactory>
+	{
+	public:
+		FTraceTreeDashboardViewFactory();
+		virtual ~FTraceTreeDashboardViewFactory();
+
+		enum class EProcessReason : uint8
+		{
+			None,
+			FilterUpdated,
+			EntriesUpdated
+		};
+
+	protected:
+		struct SRowWidget : public SMultiColumnTableRow<TSharedPtr<IDashboardDataTreeViewEntry>>
+		{
+			SLATE_BEGIN_ARGS(SRowWidget){}
+			SLATE_END_ARGS()
+
+			void Construct(const FArguments& InArgs, const TSharedRef<STableViewBase>& InOwnerTable, TSharedPtr<IDashboardDataTreeViewEntry> InData, TSharedRef<FTraceTreeDashboardViewFactory> InFactory);
+			virtual TSharedRef<SWidget> GenerateWidgetForColumn(const FName& Column) override;
+
+			TSharedPtr<IDashboardDataTreeViewEntry> Data;
+			TSharedPtr<FTraceTreeDashboardViewFactory> Factory;
+		};
+
+		struct FHeaderRowColumnData
+		{
+			const FText DisplayName;
+			const FName IconName;
+			const bool bShowDisplayName = true;
+			const bool bDefaultHidden = false;
+			const float FillWidth = 1.0f;
+			const EHorizontalAlignment Alignment = HAlign_Left;
+		};
+
+		struct FColumnData
+		{
+			const TFunction<FText(const IDashboardDataTreeViewEntry&)> GetDisplayValue;
+			const TFunction<FName(const IDashboardDataTreeViewEntry&)> GetIconName;
+		};
+
+		virtual TSharedPtr<SWidget> GetFilterBarWidget();
+		virtual TSharedPtr<SWidget> GetFilterBarButtonWidget();
+
+		virtual TSharedRef<SWidget> MakeWidget(TSharedRef<SDockTab> OwnerTab, const FSpawnTabArgs& SpawnTabArgs);
+
+		virtual TSharedRef<SWidget> GenerateWidgetForColumn(const TSharedRef<FTraceTreeDashboardViewFactory::SRowWidget>& InRowWidget, const TSharedRef<IDashboardDataTreeViewEntry>& InRowData, const FName& Column) = 0;
+		
+		virtual TSharedPtr<SWidget> OnConstructContextMenu();
+
+		virtual void OnSelectionChanged(TSharedPtr<IDashboardDataTreeViewEntry> SelectedItem, ESelectInfo::Type SelectInfo);
+		virtual FReply OnDataRowKeyInput(const FGeometry& Geometry, const FKeyEvent& KeyEvent) const;
+
+		const FText& GetSearchFilterText() const;
+
+		virtual void RefreshFilteredEntriesListView();
+
+		virtual const TMap<FName, FHeaderRowColumnData>& GetHeaderRowColumns() const = 0;
+		virtual const TMap<FName, FColumnData>& GetColumns() const = 0;
+
+		virtual void ProcessEntries(EProcessReason Reason) = 0;
+		virtual void SortTable() = 0;
+
+		void Tick(float InElapsed);
+
+#if WITH_EDITOR
+		virtual bool IsDebugDrawEnabled() const { return false; }
+		virtual void DebugDraw(float InElapsed, const TArray<TSharedPtr<IDashboardDataTreeViewEntry>>& InSelectedItems, ::Audio::FDeviceId InAudioDeviceId) const {};
+#endif // WITH_EDITOR
+
+		template<typename TableProviderType>
+		bool FilterEntries(TFunctionRef<bool(const IDashboardDataTreeViewEntry&)> IsFiltered)
+		{
+			const TSharedPtr<const TableProviderType> Provider = FindProvider<const TableProviderType>();
+			if (Provider.IsValid())
+			{
+				if (const TableProviderType::FDeviceData* DeviceData = Provider->FindFilteredDeviceData())
+				{
+					DataViewEntries.Reset();
+
+					// Filter Entries
+					const auto TransformEntry = [](const TableProviderType::FEntryPair& Pair)
+					{
+						return StaticCastSharedPtr<IDashboardDataTreeViewEntry>(Pair.Value);
+					};
+
+					const auto FilterEntry = [this, &IsFiltered](const TableProviderType::FEntryPair& Pair)
+					{
+						return !IsFiltered((const IDashboardDataTreeViewEntry&)(*Pair.Value));
+					};
+
+					if (SearchBoxFilterText.IsEmpty())
+					{
+						Algo::Transform(*DeviceData, DataViewEntries, TransformEntry);
+					}
+					else
+					{
+						Algo::TransformIf(*DeviceData, DataViewEntries, FilterEntry, TransformEntry);
+					}
+					
+					// Sort list
+					RequestSort();
+
+					return true;
+				}
+				else
+				{
+					if (!DataViewEntries.IsEmpty())
+					{
+						DataViewEntries.Empty();
+						return true;
+					}
+				}
+			}
+
+			return false;
+		}
+
+		EProcessReason UpdateFilterReason = EProcessReason::None;
+		FTSTicker::FDelegateHandle TickerHandle;
+
+		TArray<TSharedPtr<IDashboardDataTreeViewEntry>> DataViewEntries;
+		TMap<FName, uint64> UpdateIds;
+
+		TSharedPtr<SWidget> DashboardWidget;
+		TSharedPtr<SHeaderRow> HeaderRowWidget;
+		TSharedPtr<STreeView<TSharedPtr<IDashboardDataTreeViewEntry>>> FilteredEntriesListView;
+
+		FName SortByColumn;
+		EColumnSortMode::Type SortMode = EColumnSortMode::None;
+
+	private:
+		TSharedRef<SHeaderRow> MakeHeaderRowWidget();
+		void SetSearchBoxFilterText(const FText& NewText);
+
+		EColumnSortMode::Type GetColumnSortMode(const FName InColumnId) const;
+		void RequestSort();
+		void OnColumnSortModeChanged(const EColumnSortPriority::Type InSortPriority, const FName& InColumnId, const EColumnSortMode::Type InSortMode);
+
+		TSharedPtr<SSearchBox> SearchBoxWidget;
+		FText SearchBoxFilterText;
+	};
+} // namespace UE::Audio::Insights
