[Audio Insights] Added missing typename from CL 41016745 (Sound Dashboard part 5)

diff --git a/Engine/Plugins/AudioInsights/Source/AudioInsights/Public/Views/TreeDashboardViewFactory.h b/Engine/Plugins/AudioInsights/Source/AudioInsights/Public/Views/TreeDashboardViewFactory.h
--- a/Engine/Plugins/AudioInsights/Source/AudioInsights/Public/Views/TreeDashboardViewFactory.h
+++ b/Engine/Plugins/AudioInsights/Source/AudioInsights/Public/Views/TreeDashboardViewFactory.h
@@ -110,17 +110,17 @@ namespace UE::Audio::Insights
 			const TSharedPtr<const TableProviderType> Provider = FindProvider<const TableProviderType>();
 			if (Provider.IsValid())
 			{
-				if (const TableProviderType::FDeviceData* DeviceData = Provider->FindFilteredDeviceData())
+				if (const typename TableProviderType::FDeviceData* DeviceData = Provider->FindFilteredDeviceData())
 				{
 					DataViewEntries.Reset();
 
 					// Filter Entries
-					const auto TransformEntry = [](const TableProviderType::FEntryPair& Pair)
+					const auto TransformEntry = [](const typename TableProviderType::FEntryPair& Pair)
 					{
 						return StaticCastSharedPtr<IDashboardDataTreeViewEntry>(Pair.Value);
 					};
 
-					const auto FilterEntry = [this, &IsFiltered](const TableProviderType::FEntryPair& Pair)
+					const auto FilterEntry = [this, &IsFiltered](const typename TableProviderType::FEntryPair& Pair)
 					{
 						return !IsFiltered((const IDashboardDataTreeViewEntry&)(*Pair.Value));
 					};
