[Audio Insights] Sound Dashboard part 6: added SoundTraceMessages and SoundTraceProvider

diff --git a/Engine/Plugins/AudioInsights/Source/AudioInsights/Private/Providers/SoundTraceProvider.cpp b/Engine/Plugins/AudioInsights/Source/AudioInsights/Private/Providers/SoundTraceProvider.cpp
new file mode 100644
--- /dev/null
+++ b/Engine/Plugins/AudioInsights/Source/AudioInsights/Private/Providers/SoundTraceProvider.cpp
@@ -0,0 +1,544 @@
+// Copyright Epic Games, Inc. All Rights Reserved.
+#include "Providers/SoundTraceProvider.h"
+
+#include "AudioInsightsModule.h"
+#include "Trace/Analyzer.h"
+#include "TraceServices/Model/AnalysisSession.h"
+
+#if !WITH_EDITOR
+#include "Common/PagedArray.h"
+#include "Async/ParallelFor.h"
+#endif // !WITH_EDITOR
+
+namespace UE::Audio::Insights
+{
+	namespace FSoundTraceProviderPrivate
+	{
+#if !WITH_EDITOR
+		template<typename T>
+		const T* FindClosestMessageToTimestamp(const TraceServices::TPagedArray<T>& InCachedMessages, const double InTimeMarker, const uint32 InPlayOrder)
+		{
+			const int32 ClosestMessageToTimeStampIndex = TraceServices::PagedArrayAlgo::BinarySearchClosestBy(InCachedMessages, InTimeMarker, [](const T& Msg) { return Msg.Timestamp; });
+
+			// Iterate backwards from TimeMarker until we find the matching PlayOrder
+			for (auto It = InCachedMessages.GetIteratorFromItem(ClosestMessageToTimeStampIndex); It; --It)
+			{
+				if (It->PlayOrder == InPlayOrder)
+				{
+					return &(*It);
+				}
+			}
+
+			return nullptr;
+		}
+#endif // !WITH_EDITOR
+	}
+
+	FSoundTraceProvider::FSoundTraceProvider(TSharedRef<FTraceChannelManager> InManager)
+		: TDeviceDataMapTraceProvider<uint32, TSharedPtr<FSoundDashboardEntry>>(GetName_Static())
+	{
+		Channels.Add(InManager->CreateHandle({ TEXT("AudioMixerChannel") }));
+	}
+
+	FName FSoundTraceProvider::GetName_Static()
+	{
+		static const FLazyName SoundTraceProviderName = "SoundProvider";
+		return SoundTraceProviderName;
+	}
+
+	UE::Trace::IAnalyzer* FSoundTraceProvider::ConstructAnalyzer(TraceServices::IAnalysisSession& InSession)
+	{
+		class FSoundTraceAnalyzer : public FTraceAnalyzerBase
+		{
+		public:
+			FSoundTraceAnalyzer(TSharedRef<FSoundTraceProvider> InProvider, TraceServices::IAnalysisSession& InSession)
+				: FTraceAnalyzerBase(InProvider)
+				, Session(InSession)
+			{
+			}
+
+			virtual void OnAnalysisBegin(const FOnAnalysisContext& Context) override
+			{
+				FTraceAnalyzerBase::OnAnalysisBegin(Context);
+
+				UE::Trace::IAnalyzer::FInterfaceBuilder& Builder = Context.InterfaceBuilder;
+
+				Builder.RouteEvent(RouteId_Start,     "Audio", "SoundStart");
+				Builder.RouteEvent(RouteId_Stop,      "Audio", "SoundStop");
+				Builder.RouteEvent(RouteId_Priority,  "Audio", "SoundPriority");
+				Builder.RouteEvent(RouteId_Distance,  "Audio", "SoundDistance");
+				Builder.RouteEvent(RouteId_Amplitude, "Audio", "SoundAmplitude");
+				Builder.RouteEvent(RouteId_Volume,    "Audio", "SoundVolume");
+				Builder.RouteEvent(RouteId_Pitch,     "Audio", "SoundPitch");
+			}
+
+			virtual bool OnEvent(uint16 RouteId, EStyle Style, const FOnEventContext& Context) override
+			{
+				LLM_SCOPE_BYNAME(TEXT("Insights/FSoundTraceAnalyzer"));
+
+				FSoundMessages& Messages = GetProvider<FSoundTraceProvider>().TraceMessages;
+
+				switch (RouteId)
+				{
+					case RouteId_Start:
+					{
+						Messages.StartMessages.Enqueue(FSoundStartMessage{ Context });
+						break;
+					}
+
+					case RouteId_Stop:
+					{
+						Messages.StopMessages.Enqueue(FSoundStopMessage{ Context });
+						break;
+					}
+
+					case RouteId_Priority:
+					{
+						Messages.PriorityMessages.Enqueue(FSoundPriorityMessage{ Context });
+						break;
+					}
+
+					case RouteId_Distance:
+					{
+						Messages.DistanceMessages.Enqueue(FSoundDistanceMessage{ Context });
+						break;
+					}
+
+					case RouteId_Amplitude:
+					{
+						Messages.AmplitudeMessages.Enqueue(FSoundAmplitudeMessage{ Context });
+						break;
+					}
+
+					case RouteId_Volume:
+					{
+						Messages.VolumeMessages.Enqueue(FSoundVolumeMessage{ Context });
+						break;
+					}
+
+					case RouteId_Pitch:
+					{
+						Messages.PitchMessages.Enqueue(FSoundPitchMessage{ Context });
+						break;
+					}
+
+					default:
+					{
+						return OnEventFailure(RouteId, Style, Context);
+					}
+				}
+
+				const double Timestamp = Context.EventTime.AsSeconds(Context.EventData.GetValue<uint64>("Timestamp"));
+
+				{
+					TraceServices::FAnalysisSessionEditScope SessionEditScope(Session);
+					Session.UpdateDurationSeconds(Timestamp);
+				}
+
+				return OnEventSuccess(RouteId, Style, Context);
+			}
+
+		private:
+			enum : uint16
+			{
+				RouteId_Start,
+				RouteId_Stop,
+				RouteId_Priority,
+				RouteId_Distance,
+				RouteId_Amplitude,
+				RouteId_Volume,
+				RouteId_Pitch
+			};
+
+			TraceServices::IAnalysisSession& Session;
+		};
+
+		return new FSoundTraceAnalyzer(AsShared(), InSession);
+	}
+
+#if !WITH_EDITOR
+	void FSoundTraceProvider::InitSessionCachedMessages(TraceServices::IAnalysisSession& InSession)
+	{
+		SessionCachedMessages = MakeUnique<FSoundSessionCachedMessages>(InSession);
+	}
+#endif // !WITH_EDITOR
+
+	bool FSoundTraceProvider::ProcessMessages()
+	{
+		// Helper lambdas
+		auto GetOrCreateEntry = [this](const FSoundMessageBase& Msg)
+		{
+			TSharedPtr<FSoundDashboardEntry>* ToReturn = nullptr;
+
+			UpdateDeviceEntry(Msg.DeviceId, Msg.PlayOrder, [&ToReturn, &Msg](TSharedPtr<FSoundDashboardEntry>& Entry)
+			{
+				if (!Entry.IsValid())
+				{
+					Entry = MakeShared<FSoundDashboardEntry>();
+					Entry->DeviceId = Msg.DeviceId;
+					Entry->PlayOrder = Msg.PlayOrder;
+				}
+
+				Entry->Timestamp = Msg.Timestamp;
+				ToReturn = &Entry;
+			});
+
+			return ToReturn;
+		};
+
+		auto GetEntry = [this](const FSoundMessageBase& Msg)
+		{
+			return FindDeviceEntry(Msg.DeviceId, Msg.PlayOrder);
+		};
+
+
+		// Process messages
+		ProcessMessageQueue<FSoundStartMessage>(TraceMessages.StartMessages, GetOrCreateEntry,
+		[this](const FSoundStartMessage& Msg, TSharedPtr<FSoundDashboardEntry>* OutEntry)
+		{
+#if !WITH_EDITOR
+			if (SessionCachedMessages.IsValid())
+			{
+				SessionCachedMessages->StartCachedMessages.EmplaceBack(Msg);
+			}
+#endif // !WITH_EDITOR
+
+			FSoundDashboardEntry& EntryRef = *OutEntry->Get();
+			EntryRef.Name = *Msg.Name;
+			EntryRef.EntryType = Msg.EntryType;
+		});
+
+		TArray<int32, TInlineAllocator<64>> EntriesWithPoppedDataPoints;
+
+		ProcessMessageQueue<FSoundPriorityMessage>(TraceMessages.PriorityMessages, GetEntry,
+		[this, &EntriesWithPoppedDataPoints](const FSoundPriorityMessage& Msg, TSharedPtr<FSoundDashboardEntry>* OutEntry)
+		{
+			if (OutEntry == nullptr)
+			{
+				return;
+			}
+
+#if !WITH_EDITOR
+			if (SessionCachedMessages.IsValid())
+			{
+				SessionCachedMessages->PriorityCachedMessages.EmplaceBack(Msg);
+			}
+#endif // !WITH_EDITOR
+
+			if (!EntriesWithPoppedDataPoints.Contains(Msg.PlayOrder))
+			{
+				(*OutEntry)->PriorityDataPoints.Pop((*OutEntry)->PriorityDataPoints.Num());
+				EntriesWithPoppedDataPoints.Add(Msg.PlayOrder);
+			}
+
+			(*OutEntry)->PriorityDataPoints.Push({ Msg.Timestamp, Msg.Priority });
+		});
+		
+		EntriesWithPoppedDataPoints.Reset();
+		ProcessMessageQueue<FSoundDistanceMessage>(TraceMessages.DistanceMessages, GetEntry,
+		[this, &EntriesWithPoppedDataPoints](const FSoundDistanceMessage& Msg, TSharedPtr<FSoundDashboardEntry>* OutEntry)
+		{
+			if (OutEntry == nullptr)
+			{
+				return;
+			}
+
+#if !WITH_EDITOR
+			if (SessionCachedMessages.IsValid())
+			{
+				SessionCachedMessages->DistanceCachedMessages.EmplaceBack(Msg);
+			}
+#endif // !WITH_EDITOR
+
+			if (!EntriesWithPoppedDataPoints.Contains(Msg.PlayOrder))
+			{
+				(*OutEntry)->DistanceDataPoints.Pop((*OutEntry)->DistanceDataPoints.Num());
+				EntriesWithPoppedDataPoints.Add(Msg.PlayOrder);
+			}
+
+			(*OutEntry)->DistanceDataPoints.Push({ Msg.Timestamp, Msg.Distance });
+		});
+		
+		EntriesWithPoppedDataPoints.Reset();
+		ProcessMessageQueue<FSoundAmplitudeMessage>(TraceMessages.AmplitudeMessages, GetEntry,
+		[this, &EntriesWithPoppedDataPoints](const FSoundAmplitudeMessage& Msg, TSharedPtr<FSoundDashboardEntry>* OutEntry)
+		{
+			if (OutEntry == nullptr)
+			{
+				return;
+			}
+
+#if !WITH_EDITOR
+			if (SessionCachedMessages.IsValid())
+			{
+				SessionCachedMessages->AmplitudeCachedMessages.EmplaceBack(Msg);
+			}
+#endif // !WITH_EDITOR
+
+			if (!EntriesWithPoppedDataPoints.Contains(Msg.PlayOrder))
+			{
+				(*OutEntry)->AmplitudeDataPoints.Pop((*OutEntry)->AmplitudeDataPoints.Num());
+				EntriesWithPoppedDataPoints.Add(Msg.PlayOrder);
+			}
+
+			(*OutEntry)->AmplitudeDataPoints.Push({ Msg.Timestamp, Msg.Amplitude });
+		});
+
+		ProcessMessageQueue<FSoundVolumeMessage>(TraceMessages.VolumeMessages, GetEntry,
+		[this, &EntriesWithPoppedDataPoints](const FSoundVolumeMessage& Msg, TSharedPtr<FSoundDashboardEntry>* OutEntry)
+		{
+			if (OutEntry == nullptr)
+			{
+				return;
+			}
+
+#if !WITH_EDITOR
+			if (SessionCachedMessages.IsValid())
+			{
+				SessionCachedMessages->VolumeCachedMessages.EmplaceBack(Msg);
+			}
+#endif // !WITH_EDITOR
+
+			if (!EntriesWithPoppedDataPoints.Contains(Msg.PlayOrder))
+			{
+				(*OutEntry)->VolumeDataPoints.Pop((*OutEntry)->VolumeDataPoints.Num());
+				EntriesWithPoppedDataPoints.Add(Msg.PlayOrder);
+			}
+
+			(*OutEntry)->VolumeDataPoints.Push({ Msg.Timestamp, Msg.Volume });
+		});
+
+		EntriesWithPoppedDataPoints.Reset();
+		ProcessMessageQueue<FSoundPitchMessage>(TraceMessages.PitchMessages, GetEntry,
+		[this, &EntriesWithPoppedDataPoints](const FSoundPitchMessage& Msg, TSharedPtr<FSoundDashboardEntry>* OutEntry)
+		{
+			if (OutEntry == nullptr)
+			{
+				return;
+			}
+
+#if !WITH_EDITOR
+			if (SessionCachedMessages.IsValid())
+			{
+				SessionCachedMessages->PitchCachedMessages.EmplaceBack(Msg);
+			}
+#endif // !WITH_EDITOR
+
+			if (!EntriesWithPoppedDataPoints.Contains(Msg.PlayOrder))
+			{
+				(*OutEntry)->PitchDataPoints.Pop((*OutEntry)->PitchDataPoints.Num());
+				EntriesWithPoppedDataPoints.Add(Msg.PlayOrder);
+			}
+
+			(*OutEntry)->PitchDataPoints.Push({ Msg.Timestamp, Msg.Pitch });
+		});
+
+		ProcessMessageQueue<FSoundStopMessage>(TraceMessages.StopMessages, GetEntry,
+		[this](const FSoundStopMessage& Msg, TSharedPtr<FSoundDashboardEntry>* OutEntry)
+		{
+#if !WITH_EDITOR
+			if (SessionCachedMessages.IsValid())
+			{
+				SessionCachedMessages->StopCachedMessages.EmplaceBack(Msg);
+			}
+#endif // !WITH_EDITOR
+
+			if (OutEntry && (*OutEntry)->Timestamp < Msg.Timestamp)
+			{
+				RemoveDeviceEntry(Msg.DeviceId, Msg.PlayOrder);
+			}
+		});
+
+		return true;
+	}
+
+#if !WITH_EDITOR
+	void FSoundTraceProvider::OnTimingViewTimeMarkerChanged(double TimeMarker)
+	{
+		using namespace FSoundTraceProviderPrivate;
+
+		if (!SessionCachedMessages.IsValid())
+		{
+			return;
+		}
+
+		DeviceDataMap.Empty();
+
+		// Collect all the start messages registered until this point in time 
+		for (const FSoundStartMessage& StartCachedMessage : SessionCachedMessages->StartCachedMessages)
+		{
+			if (StartCachedMessage.Timestamp > TimeMarker)
+			{
+				break;
+			}
+
+			UpdateDeviceEntry(StartCachedMessage.DeviceId, StartCachedMessage.PlayOrder, [&StartCachedMessage](TSharedPtr<FSoundDashboardEntry>& Entry)
+			{
+				if (!Entry.IsValid())
+				{
+					Entry = MakeShared<FSoundDashboardEntry>();
+					Entry->DeviceId  = StartCachedMessage.DeviceId;
+					Entry->PlayOrder = StartCachedMessage.PlayOrder;
+				}
+				Entry->Timestamp = StartCachedMessage.Timestamp;
+
+				Entry->Name = *StartCachedMessage.Name;
+			});
+		}
+
+		// Selectively remove start messages collected in the step above by knowing which sounds were stopped.
+		// With this we will know what are the active sounds at this point in time.
+		for (const FSoundStopMessage& StopCachedMessage : SessionCachedMessages->StopCachedMessages)
+		{
+			if (StopCachedMessage.Timestamp > TimeMarker)
+			{
+				break;
+			}
+
+			const TSharedPtr<FSoundDashboardEntry>* OutEntry = FindDeviceEntry(StopCachedMessage.DeviceId, StopCachedMessage.PlayOrder);
+
+			if (OutEntry && (*OutEntry)->Timestamp < StopCachedMessage.Timestamp)
+			{
+				RemoveDeviceEntry(StopCachedMessage.DeviceId, StopCachedMessage.PlayOrder);
+			}
+		}
+
+		// For now we only retrieve information from AudioDeviceId 1 (main device in standalone games)
+		const FDeviceData* DeviceData = DeviceDataMap.Find(1);
+		if (DeviceData)
+		{
+			// Collect messages (volume, pitch, etc.) from active sounds (based on active sounds's PlayOrder)
+			struct CachedEntryInfo
+			{
+				FSoundPriorityMessage PriorityMessage;
+				FSoundDistanceMessage DistanceMessage;
+				FSoundAmplitudeMessage AmplitudeMessage;
+				FSoundVolumeMessage VolumeMessage;
+				FSoundPitchMessage PitchMessage;
+			};
+
+			TArray<uint32> PlayOrderArray;
+			DeviceData->GenerateKeyArray(PlayOrderArray);
+
+			TArray<CachedEntryInfo> CachedEntryInfos;
+			CachedEntryInfos.SetNumUninitialized(PlayOrderArray.Num());
+
+			// Using ParallelFor to speed-up the cached messages retrieval, using a traditional for loop is unacceptably slower, specially in large traces.
+			ParallelFor(PlayOrderArray.Num(), 
+			[&PlayOrderArray, &CachedEntryInfos, TimeMarker, this](const int32 Index)
+			{
+				const uint32 PlayOrder = PlayOrderArray[Index];
+
+				// Priority
+				const FSoundPriorityMessage* FoundPriorityCachedMessage = FindClosestMessageToTimestamp(SessionCachedMessages->PriorityCachedMessages, TimeMarker, PlayOrder);
+				if (FoundPriorityCachedMessage)
+				{
+					CachedEntryInfos[Index].PriorityMessage = *FoundPriorityCachedMessage;
+				}
+
+				// Distance
+				const FSoundDistanceMessage* FoundDistanceCachedMessage = FindClosestMessageToTimestamp(SessionCachedMessages->DistanceCachedMessages, TimeMarker, PlayOrder);
+				if (FoundDistanceCachedMessage)
+				{
+					CachedEntryInfos[Index].DistanceMessage = *FoundDistanceCachedMessage;
+				}
+
+				// Amplitude
+				const FSoundAmplitudeMessage* FoundAmplitudeCachedMessage = FindClosestMessageToTimestamp(SessionCachedMessages->AmplitudeCachedMessages, TimeMarker, PlayOrder);
+				if (FoundAmplitudeCachedMessage)
+				{
+					CachedEntryInfos[Index].AmplitudeMessage = *FoundAmplitudeCachedMessage;
+				}
+
+				// Volume
+				const FSoundVolumeMessage* FoundVolumeCachedMessage = FindClosestMessageToTimestamp(SessionCachedMessages->VolumeCachedMessages, TimeMarker, PlayOrder);
+				if (FoundVolumeCachedMessage)
+				{
+					CachedEntryInfos[Index].VolumeMessage = *FoundVolumeCachedMessage;
+				}
+
+				// Pitch
+				const FSoundPitchMessage* FoundPitchCachedMessage = FindClosestMessageToTimestamp(SessionCachedMessages->PitchCachedMessages, TimeMarker, PlayOrder);
+				if (FoundPitchCachedMessage)
+				{
+					CachedEntryInfos[Index].PitchMessage = *FoundPitchCachedMessage;
+				}
+			});
+
+			// Update the device entries with the collected info
+			for (const CachedEntryInfo& CachedEntryInfo : CachedEntryInfos)
+			{
+				UpdateDeviceEntry(CachedEntryInfo.PriorityMessage.DeviceId, CachedEntryInfo.PriorityMessage.PlayOrder, [&CachedEntryInfo](TSharedPtr<FSoundDashboardEntry>& Entry)
+				{
+					if (!Entry.IsValid())
+					{
+						Entry = MakeShared<FSoundDashboardEntry>();
+						Entry->DeviceId  = CachedEntryInfo.PriorityMessage.DeviceId;
+						Entry->PlayOrder = CachedEntryInfo.PriorityMessage.PlayOrder;
+					}
+
+					Entry->Timestamp = CachedEntryInfo.PriorityMessage.Timestamp;
+					Entry->PriorityDataPoints.Push({ CachedEntryInfo.PriorityMessage.Timestamp, CachedEntryInfo.PriorityMessage.Priority });
+				});
+
+				UpdateDeviceEntry(CachedEntryInfo.DistanceMessage.DeviceId, CachedEntryInfo.DistanceMessage.PlayOrder, [&CachedEntryInfo](TSharedPtr<FSoundDashboardEntry>& Entry)
+				{
+					if (!Entry.IsValid())
+					{
+						Entry = MakeShared<FSoundDashboardEntry>();
+						Entry->DeviceId  = CachedEntryInfo.DistanceMessage.DeviceId;
+						Entry->PlayOrder = CachedEntryInfo.DistanceMessage.PlayOrder;
+					}
+
+					Entry->Timestamp = CachedEntryInfo.DistanceMessage.Timestamp;
+					Entry->DistanceDataPoints.Push({ CachedEntryInfo.DistanceMessage.Timestamp, CachedEntryInfo.DistanceMessage.Distance });
+				});
+
+				UpdateDeviceEntry(CachedEntryInfo.AmplitudeMessage.DeviceId, CachedEntryInfo.AmplitudeMessage.PlayOrder, [&CachedEntryInfo](TSharedPtr<FSoundDashboardEntry>& Entry)
+				{
+					if (!Entry.IsValid())
+					{
+						Entry = MakeShared<FSoundDashboardEntry>();
+						Entry->DeviceId  = CachedEntryInfo.AmplitudeMessage.DeviceId;
+						Entry->PlayOrder = CachedEntryInfo.AmplitudeMessage.PlayOrder;
+					}
+
+					Entry->Timestamp = CachedEntryInfo.AmplitudeMessage.Timestamp;
+					Entry->AmplitudeDataPoints.Push({ CachedEntryInfo.AmplitudeMessage.Timestamp, CachedEntryInfo.AmplitudeMessage.Amplitude });
+				});
+
+				UpdateDeviceEntry(CachedEntryInfo.VolumeMessage.DeviceId, CachedEntryInfo.VolumeMessage.PlayOrder, [&CachedEntryInfo](TSharedPtr<FSoundDashboardEntry>& Entry)
+				{
+					if (!Entry.IsValid())
+					{
+						Entry = MakeShared<FSoundDashboardEntry>();
+						Entry->DeviceId  = CachedEntryInfo.VolumeMessage.DeviceId;
+						Entry->PlayOrder = CachedEntryInfo.VolumeMessage.PlayOrder;
+					}
+
+					Entry->Timestamp = CachedEntryInfo.VolumeMessage.Timestamp;
+					Entry->VolumeDataPoints.Push({ CachedEntryInfo.VolumeMessage.Timestamp, CachedEntryInfo.VolumeMessage.Volume });
+				});
+
+				UpdateDeviceEntry(CachedEntryInfo.PitchMessage.DeviceId, CachedEntryInfo.PitchMessage.PlayOrder, [&CachedEntryInfo](TSharedPtr<FSoundDashboardEntry>& Entry)
+				{
+					if (!Entry.IsValid())
+					{
+						Entry = MakeShared<FSoundDashboardEntry>();
+						Entry->DeviceId  = CachedEntryInfo.PitchMessage.DeviceId;
+						Entry->PlayOrder = CachedEntryInfo.PitchMessage.PlayOrder;
+					}
+
+					Entry->Timestamp = CachedEntryInfo.PitchMessage.Timestamp;
+					Entry->PitchDataPoints.Push({ CachedEntryInfo.PitchMessage.Timestamp, CachedEntryInfo.PitchMessage.Pitch });
+				});
+			}
+		}
+
+		// Call parent method to update LastMessageId
+		FTraceProviderBase::OnTimingViewTimeMarkerChanged(TimeMarker);
+	}
+
+#endif // !WITH_EDITOR
+
+} // namespace UE::Audio::Insights
diff --git a/Engine/Plugins/AudioInsights/Source/AudioInsights/Private/Providers/SoundTraceProvider.h b/Engine/Plugins/AudioInsights/Source/AudioInsights/Private/Providers/SoundTraceProvider.h
new file mode 100644
--- /dev/null
+++ b/Engine/Plugins/AudioInsights/Source/AudioInsights/Private/Providers/SoundTraceProvider.h
@@ -0,0 +1,40 @@
+// Copyright Epic Games, Inc. All Rights Reserved.
+#pragma once
+
+#include "AudioInsightsTraceChannelHandle.h"
+#include "AudioInsightsTraceProviderBase.h"
+#include "Messages/SoundTraceMessages.h"
+#include "UObject/NameTypes.h"
+
+namespace UE::Audio::Insights
+{
+	class FSoundTraceProvider
+		: public TDeviceDataMapTraceProvider<uint32, TSharedPtr<FSoundDashboardEntry>>
+		, public TSharedFromThis<FSoundTraceProvider>
+	{
+	public:
+		FSoundTraceProvider(TSharedRef<FTraceChannelManager> InManager);
+		virtual ~FSoundTraceProvider() = default;
+
+		static FName GetName_Static();
+
+		virtual UE::Trace::IAnalyzer* ConstructAnalyzer(TraceServices::IAnalysisSession& InSession) override;
+
+#if !WITH_EDITOR
+		virtual void InitSessionCachedMessages(TraceServices::IAnalysisSession& InSession) override;
+#endif // !WITH_EDITOR
+
+		virtual bool ProcessMessages() override;
+
+	private:
+#if !WITH_EDITOR
+		virtual void OnTimingViewTimeMarkerChanged(double TimeMarker) override;
+
+		TUniquePtr<FSoundSessionCachedMessages> SessionCachedMessages;
+#endif // !WITH_EDITOR
+
+		FSoundMessages TraceMessages;
+
+		TSet<FTraceChannelHandle> Channels;
+	};
+} // namespace UE::Audio::Insights
diff --git a/Engine/Plugins/AudioInsights/Source/AudioInsights/Public/Messages/SoundTraceMessages.h b/Engine/Plugins/AudioInsights/Source/AudioInsights/Public/Messages/SoundTraceMessages.h
new file mode 100644
--- /dev/null
+++ b/Engine/Plugins/AudioInsights/Source/AudioInsights/Public/Messages/SoundTraceMessages.h
@@ -0,0 +1,236 @@
+// Copyright Epic Games, Inc. All Rights Reserved.
+#pragma once
+
+#include "AudioDefines.h"
+#include "DSP/Dsp.h"
+#include "Messages/AnalyzerMessageQueue.h"
+#include "Trace/Analyzer.h"
+#include "UObject/SoftObjectPath.h"
+#include "Views/TreeDashboardViewFactory.h"
+
+#if !WITH_EDITOR
+#include "Common/PagedArray.h"
+#include "TraceServices/Model/AnalysisSession.h"
+#endif // !WITH_EDITOR
+
+namespace UE::Audio::Insights
+{
+	enum class ESoundDashboardEntryType : uint8
+	{
+		None,
+		MetaSound,
+		SoundCue,
+		ProceduralSource,
+	};
+
+	// Trace messages
+	struct FSoundMessageBase
+	{
+		FSoundMessageBase() = default;
+		FSoundMessageBase(const Trace::IAnalyzer::FOnEventContext& InContext)
+		{
+			const Trace::IAnalyzer::FEventData& EventData = InContext.EventData;
+
+			DeviceId  = static_cast<::Audio::FDeviceId>(EventData.GetValue<uint32>("DeviceId"));
+			Timestamp = InContext.EventTime.AsSeconds(EventData.GetValue<uint64>("Timestamp"));
+			PlayOrder = EventData.GetValue<uint32>("PlayOrder");
+		}
+
+		::Audio::FDeviceId DeviceId = INDEX_NONE;
+		uint32 PlayOrder = INDEX_NONE;
+		double Timestamp = 0.0;
+	};
+
+	struct FSoundStartMessage : public FSoundMessageBase
+	{
+		FSoundStartMessage() = default;
+
+		FSoundStartMessage(const Trace::IAnalyzer::FOnEventContext& InContext)
+			: FSoundMessageBase(InContext)
+		{
+			const Trace::IAnalyzer::FEventData& EventData = InContext.EventData;
+			EventData.GetString("Name", Name);
+		}
+
+		FString Name;
+		ESoundDashboardEntryType EntryType = ESoundDashboardEntryType::None;
+	};
+
+	using FSoundStopMessage = FSoundMessageBase;
+
+	struct FMetaSoundStartMessage : public FSoundStartMessage
+	{
+		FMetaSoundStartMessage() = default;
+
+		FMetaSoundStartMessage(const UE::Trace::IAnalyzer::FOnEventContext& InContext)
+			: FSoundStartMessage(InContext)
+		{
+			EntryType = ESoundDashboardEntryType::MetaSound;
+		}
+	};
+
+	struct FSoundCueStartMessage : public FSoundStartMessage
+	{
+		FSoundCueStartMessage() = default;
+
+		FSoundCueStartMessage(const UE::Trace::IAnalyzer::FOnEventContext& InContext)
+			: FSoundStartMessage(InContext)
+		{
+			EntryType = ESoundDashboardEntryType::SoundCue;
+		}
+	};
+
+	struct FProceduralSourceStartMessage : public FSoundStartMessage
+	{
+		FProceduralSourceStartMessage() = default;
+
+		FProceduralSourceStartMessage(const UE::Trace::IAnalyzer::FOnEventContext& InContext)
+			: FSoundStartMessage(InContext)
+		{
+			EntryType = ESoundDashboardEntryType::ProceduralSource;
+		}
+	};
+
+	struct FSoundPriorityMessage : public FSoundMessageBase
+	{
+		FSoundPriorityMessage() = default;
+		FSoundPriorityMessage(const UE::Trace::IAnalyzer::FOnEventContext& InContext)
+			: FSoundMessageBase(InContext)
+		{
+			Priority = InContext.EventData.GetValue<float>("Priority");
+		}
+
+		float Priority = 0.0f;
+	};
+
+	struct FSoundDistanceMessage : public FSoundMessageBase
+	{
+		FSoundDistanceMessage() = default;
+		FSoundDistanceMessage(const UE::Trace::IAnalyzer::FOnEventContext& InContext)
+			: FSoundMessageBase(InContext)
+		{
+			Distance = InContext.EventData.GetValue<float>("Distance");
+		}
+
+		float Distance = 0.0f;
+	};
+
+	struct FSoundAmplitudeMessage : public FSoundMessageBase
+	{
+		FSoundAmplitudeMessage() = default;
+		FSoundAmplitudeMessage(const UE::Trace::IAnalyzer::FOnEventContext& InContext)
+			: FSoundMessageBase(InContext)
+		{
+			Amplitude = InContext.EventData.GetValue<float>("Amplitude");
+		}
+
+		float Amplitude = 0.0f;
+	};
+
+	struct FSoundVolumeMessage : public FSoundMessageBase
+	{
+		FSoundVolumeMessage() = default;
+		FSoundVolumeMessage(const UE::Trace::IAnalyzer::FOnEventContext& InContext)
+			: FSoundMessageBase(InContext)
+		{
+			Volume = InContext.EventData.GetValue<float>("Volume");
+		}
+
+		float Volume = 1.0f;
+	};
+
+	struct FSoundPitchMessage : public FSoundMessageBase
+	{
+		FSoundPitchMessage() = default;
+		FSoundPitchMessage(const UE::Trace::IAnalyzer::FOnEventContext& InContext)
+			: FSoundMessageBase(InContext)
+		{
+			Pitch = InContext.EventData.GetValue<float>("Pitch");
+		}
+
+		float Pitch = 1.0f;
+	};
+
+
+	// Trace message queues
+	class FSoundMessages
+	{
+	public:
+		TAnalyzerMessageQueue<FSoundStartMessage> StartMessages { 0.1 };
+		TAnalyzerMessageQueue<FSoundPriorityMessage> PriorityMessages{ 0.1 };
+		TAnalyzerMessageQueue<FSoundDistanceMessage> DistanceMessages { 0.1 };
+		TAnalyzerMessageQueue<FSoundAmplitudeMessage> AmplitudeMessages { 0.1 };
+		TAnalyzerMessageQueue<FSoundVolumeMessage> VolumeMessages { 0.1 };
+		TAnalyzerMessageQueue<FSoundPitchMessage> PitchMessages { 0.1 };
+		TAnalyzerMessageQueue<FSoundStopMessage> StopMessages { 0.1 };
+	};
+
+#if !WITH_EDITOR
+	struct FSoundSessionCachedMessages
+	{
+		FSoundSessionCachedMessages(TraceServices::IAnalysisSession& InSession)
+			: StartCachedMessages(InSession.GetLinearAllocator(), 16384)
+			, PriorityCachedMessages(InSession.GetLinearAllocator(), 16384)
+			, DistanceCachedMessages(InSession.GetLinearAllocator(), 16384)
+			, AmplitudeCachedMessages(InSession.GetLinearAllocator(), 16384)
+			, VolumeCachedMessages(InSession.GetLinearAllocator(), 16384)
+			, PitchCachedMessages(InSession.GetLinearAllocator(), 16384)
+			, StopCachedMessages(InSession.GetLinearAllocator(), 4096)
+		{
+
+		}
+
+		TraceServices::TPagedArray<FSoundStartMessage> StartCachedMessages;
+		TraceServices::TPagedArray<FSoundPriorityMessage> PriorityCachedMessages;
+		TraceServices::TPagedArray<FSoundDistanceMessage> DistanceCachedMessages;
+		TraceServices::TPagedArray<FSoundAmplitudeMessage> AmplitudeCachedMessages;
+		TraceServices::TPagedArray<FSoundVolumeMessage> VolumeCachedMessages;
+		TraceServices::TPagedArray<FSoundPitchMessage> PitchCachedMessages;
+		TraceServices::TPagedArray<FSoundStopMessage> StopCachedMessages;
+	};
+#endif // !WITH_EDITOR
+
+
+	// Dashboard entry
+	using FDataPoint = TPair<double, float>; // (Timestamp, Value)
+
+	class FSoundDashboardEntry : public IObjectTreeDashboardEntry
+	{
+	public:
+		FSoundDashboardEntry()
+		{
+			constexpr uint32 DataPointsCapacity = 256;
+
+			PriorityDataPoints.SetCapacity(DataPointsCapacity);
+			DistanceDataPoints.SetCapacity(DataPointsCapacity);
+			AmplitudeDataPoints.SetCapacity(DataPointsCapacity);
+			VolumeDataPoints.SetCapacity(DataPointsCapacity);
+			PitchDataPoints.SetCapacity(DataPointsCapacity);
+		}
+
+		virtual ~FSoundDashboardEntry() = default;
+
+		virtual TObjectPtr<UObject> GetObject() override { return FSoftObjectPath(Name).ResolveObject(); }
+		virtual const TObjectPtr<UObject> GetObject() const override { return GetObject(); }
+
+		virtual FText GetDisplayName() const override
+		{
+			const FSoftObjectPath AssetPath(Name);
+			return FText::FromString(AssetPath.IsValid() ? AssetPath.GetAssetName() : Name);
+		}
+
+		virtual bool IsValid() const override { return PlayOrder != INDEX_NONE; }
+
+		::Audio::FDeviceId DeviceId = INDEX_NONE;
+		uint32 PlayOrder = INDEX_NONE;
+		double Timestamp = 0.0;
+		FString Name;
+		ESoundDashboardEntryType EntryType = ESoundDashboardEntryType::None;
+
+		::Audio::TCircularAudioBuffer<FDataPoint> PriorityDataPoints;
+		::Audio::TCircularAudioBuffer<FDataPoint> DistanceDataPoints;
+		::Audio::TCircularAudioBuffer<FDataPoint> AmplitudeDataPoints;
+		::Audio::TCircularAudioBuffer<FDataPoint> VolumeDataPoints;
+		::Audio::TCircularAudioBuffer<FDataPoint> PitchDataPoints;
+	};
+} // namespace UE::Audio::Insights
