[Audio Insights] Fix compiler complain from CL 41020215 (Sound Dashboard part 6)

diff --git a/Engine/Plugins/AudioInsights/Source/AudioInsights/Public/Messages/SoundTraceMessages.h b/Engine/Plugins/AudioInsights/Source/AudioInsights/Public/Messages/SoundTraceMessages.h
--- a/Engine/Plugins/AudioInsights/Source/AudioInsights/Public/Messages/SoundTraceMessages.h
+++ b/Engine/Plugins/AudioInsights/Source/AudioInsights/Public/Messages/SoundTraceMessages.h
@@ -211,7 +211,7 @@ namespace UE::Audio::Insights
 		virtual ~FSoundDashboardEntry() = default;
 
 		virtual TObjectPtr<UObject> GetObject() override { return FSoftObjectPath(Name).ResolveObject(); }
-		virtual const TObjectPtr<UObject> GetObject() const override { return GetObject(); }
+		virtual const TObjectPtr<UObject> GetObject() const override { return FSoftObjectPath(Name).ResolveObject(); }
 
 		virtual FText GetDisplayName() const override
 		{
