[Audio Insights] Sound Dashboard part 7: added SoundDashabord + registration in AudioInsights and AudioInsightsEditor modules

diff --git a/Engine/Plugins/AudioInsights/Source/AudioInsights/AudioInsights.Build.cs b/Engine/Plugins/AudioInsights/Source/AudioInsights/AudioInsights.Build.cs
--- a/Engine/Plugins/AudioInsights/Source/AudioInsights/AudioInsights.Build.cs
+++ b/Engine/Plugins/AudioInsights/Source/AudioInsights/AudioInsights.Build.cs
@@ -32,6 +32,7 @@ public class AudioInsights : ModuleRules
 				"SlateCore",
 				"TraceAnalysis",
 				"TraceLog",
+				"ToolMenus",
 			}
 		);
 
diff --git a/Engine/Plugins/AudioInsights/Source/AudioInsights/Private/AudioInsightsModule.cpp b/Engine/Plugins/AudioInsights/Source/AudioInsights/Private/AudioInsightsModule.cpp
--- a/Engine/Plugins/AudioInsights/Source/AudioInsights/Private/AudioInsightsModule.cpp
+++ b/Engine/Plugins/AudioInsights/Source/AudioInsights/Private/AudioInsightsModule.cpp
@@ -15,6 +15,7 @@
 #if !WITH_EDITOR
 #include "AudioInsightsComponent.h"
 #include "Views/MixerSourceDashboardViewFactory.h"
+#include "Views/SoundDashboardViewFactory.h"
 #include "Views/VirtualLoopDashboardViewFactory.h"
 #endif // !WITH_EDITOR
 
@@ -41,6 +42,7 @@ namespace UE::Audio::Insights
 
 			DashboardFactory->RegisterViewFactory(MakeShared<FMixerSourceDashboardViewFactory>());
 			DashboardFactory->RegisterViewFactory(MakeShared<FVirtualLoopDashboardViewFactory>());
+			DashboardFactory->RegisterViewFactory(MakeShared<FSoundDashboardViewFactory>());
 			
 			AudioInsightsComponent = FAudioInsightsComponent::CreateInstance();
 
diff --git a/Engine/Plugins/AudioInsights/Source/AudioInsights/Private/Views/SoundDashboardViewFactory.cpp b/Engine/Plugins/AudioInsights/Source/AudioInsights/Private/Views/SoundDashboardViewFactory.cpp
new file mode 100644
--- /dev/null
+++ b/Engine/Plugins/AudioInsights/Source/AudioInsights/Private/Views/SoundDashboardViewFactory.cpp
@@ -0,0 +1,1006 @@
+// Copyright Epic Games, Inc. All Rights Reserved.
+#include "Views/SoundDashboardViewFactory.h"
+
+#include "AudioInsightsModule.h"
+#include "AudioInsightsStyle.h"
+#include "Messages/SoundTraceMessages.h"
+#include "Providers/SoundTraceProvider.h"
+#include "SoundDashboardCommands.h"
+#include "Views/SAudioFilterBar.h"
+
+#if WITH_EDITOR
+#include "Audio/AudioDebug.h"
+#include "AudioDeviceManager.h"
+#include "Editor.h"
+#include "Subsystems/AssetEditorSubsystem.h"
+#endif // WITH_EDITOR
+
+#define LOCTEXT_NAMESPACE "AudioInsights"
+
+namespace UE::Audio::Insights
+{
+	namespace FSoundDashboardViewFactoryPrivate
+	{
+		const FSoundDashboardEntry& CastEntry(const IDashboardDataTreeViewEntry& InData)
+		{
+			return static_cast<const FSoundDashboardEntry&>(InData);
+		};
+
+		float GetLastEntryArrayValue(const ::Audio::TCircularAudioBuffer<FDataPoint>& InDataPoints)
+		{
+			if (InDataPoints.Num() > 0)
+			{
+				const ::Audio::DisjointedArrayView<const FDataPoint> DataPointsDisjointedArrayView = InDataPoints.PeekInPlace(InDataPoints.Num());
+				return DataPointsDisjointedArrayView.FirstBuffer.Last().Value;
+			}
+
+			return 0.0f;
+		};
+	} // namespace FSoundDashboardViewFactoryPrivate
+
+
+	FSoundDashboardViewFactory::FSoundDashboardViewFactory()
+	{
+		FTraceModule& AudioInsightsTraceModule = static_cast<FTraceModule&>(FAudioInsightsModule::GetChecked().GetTraceModule());
+
+		const TSharedPtr<FSoundTraceProvider> SoundsTraceProvider = MakeShared<FSoundTraceProvider>(AudioInsightsTraceModule.GetChannelManager());
+
+		AudioInsightsTraceModule.AddTraceProvider(SoundsTraceProvider);
+
+		Providers = TArray<TSharedPtr<FTraceProviderBase>>
+		{
+			SoundsTraceProvider
+		};
+		
+		FSoundDashboardCommands::Register();
+
+		BindCommands();
+	}
+
+	FSoundDashboardViewFactory::~FSoundDashboardViewFactory()
+	{
+		FSoundDashboardCommands::Unregister();
+	}
+
+	FName FSoundDashboardViewFactory::GetName() const
+	{
+		return "Sounds";
+	}
+
+	FText FSoundDashboardViewFactory::GetDisplayName() const
+	{
+		return LOCTEXT("AudioDashboard_Sounds_DisplayName", "Sounds");
+	}
+
+	FSlateIcon FSoundDashboardViewFactory::GetIcon() const
+	{
+		return FSlateStyle::Get().CreateIcon("AudioInsights.Icon.SoundDashboard.Tab");
+	}
+
+	EDefaultDashboardTabStack FSoundDashboardViewFactory::GetDefaultTabStack() const
+	{
+		return EDefaultDashboardTabStack::Analysis;
+	}
+
+	void FSoundDashboardViewFactory::BindCommands()
+	{
+		CommandList = MakeShared<FUICommandList>();
+
+		const FSoundDashboardCommands& Commands = FSoundDashboardCommands::Get();
+
+		CommandList->MapAction(Commands.GetPinCommand(), FExecuteAction::CreateRaw(this, &FSoundDashboardViewFactory::PinSound));
+#if WITH_EDITOR
+		CommandList->MapAction(Commands.GetBrowseCommand(), FExecuteAction::CreateRaw(this, &FSoundDashboardViewFactory::BrowseSoundAsset));
+		CommandList->MapAction(Commands.GetEditCommand(),   FExecuteAction::CreateRaw(this, &FSoundDashboardViewFactory::OpenSoundAsset));
+#endif // WITH_EDITOR
+		CommandList->MapAction(Commands.GetHideCommand(), FExecuteAction::CreateRaw(this, &FSoundDashboardViewFactory::HideSound));
+		CommandList->MapAction(Commands.GetInfoCommand(), FExecuteAction::CreateRaw(this, &FSoundDashboardViewFactory::SoundInfo));
+	}
+
+#if WITH_EDITOR
+	TSharedRef<SWidget> FSoundDashboardViewFactory::MakeMuteSoloWidget()
+	{
+		const auto CreateButtonContentWidget = [](const FName& InIconName = FName(), const FText& InLabel = FText::GetEmpty(), const FName& InTextStyle = TEXT("ButtonText"))
+        {
+        	TSharedRef<SHorizontalBox> ButtonContainerWidget = SNew(SHorizontalBox);
+
+			// Button icon (optional)
+        	if (!InIconName.IsNone())
+        	{
+				ButtonContainerWidget->AddSlot()
+	        	.AutoWidth()
+	            .VAlign(VAlign_Center)
+	            [
+	                SNew(SImage)
+	                .ColorAndOpacity(FSlateColor::UseForeground())
+	                .Image(FSlateStyle::Get().GetBrush(InIconName))
+	            ];
+        	}
+
+			// Button text (optional)
+        	if (!InLabel.IsEmpty())
+        	{
+				const float LeftPadding = InIconName.IsNone() ? 0.0f : 4.0f;
+
+				ButtonContainerWidget->AddSlot()
+        		.VAlign(VAlign_Center)
+	            .Padding(LeftPadding, 0.0f, 0.0f, 0.0f)
+	            .AutoWidth()
+	            [
+	                SNew(STextBlock)
+	                .TextStyle(&FSlateStyle::Get().GetWidgetStyle<FTextBlockStyle>(InTextStyle))
+	                .Justification(ETextJustify::Center)
+	                .Text(InLabel)
+	            ];
+        	}
+
+        	return SNew(SBox)
+				.HeightOverride(16.0f)
+				[
+					ButtonContainerWidget
+				];
+        };
+
+		return SNew(SHorizontalBox)
+			// Mute button
+			+ SHorizontalBox::Slot()
+			.AutoWidth()
+			.Padding(8.0f, 6.0f)
+			[
+				SNew(SButton)
+				.ButtonStyle(&FAppStyle::Get().GetWidgetStyle<FButtonStyle>("SimpleButton"))
+				.OnClicked(this, &FSoundDashboardViewFactory::MuteSelected)
+				[
+					CreateButtonContentWidget("AudioInsights.Icon.SoundDashboard.Mute", LOCTEXT("SoundDashboard_MuteButtonText", "Mute Selected"), "SmallButtonText")
+				]
+			]
+			// Solo button
+			+ SHorizontalBox::Slot()
+			.AutoWidth()
+			.Padding(8.0f, 6.0f)
+			[
+				SNew(SButton)
+				.ButtonStyle(&FAppStyle::Get().GetWidgetStyle<FButtonStyle>("SimpleButton"))
+				.OnClicked(this, &FSoundDashboardViewFactory::SoloSelected)
+				[
+					CreateButtonContentWidget("AudioInsights.Icon.SoundDashboard.Solo", LOCTEXT("SoundDashboard_SoloButtonText", "Solo Selected"), "SmallButtonText")
+				]
+			];
+	}
+#endif // WITH_EDITOR
+
+	TSharedPtr<SWidget> FSoundDashboardViewFactory::GetFilterBarWidget()
+	{
+		if (!SoundsFilterBar.IsValid())
+		{
+			const TSharedPtr<FFilterCategory> FilterCategory = MakeShared<FFilterCategory>(LOCTEXT("SoundDashboard_Filter_CategoryText", "Filters"), LOCTEXT("CurveFiltersToolTip", "Filters what kind of sounds types can be displayed."));
+			
+			TArray<TSharedRef<FFilterBase<ESoundDashboardFilterFlags>>> Filters{
+				MakeShared<FSoundDashboardFilter>(
+				ESoundDashboardFilterFlags::MetaSound,
+				"MetaSound",
+				LOCTEXT("SoundDashboard_Filter_MetaSoundNameText", "MetaSound"),
+				"AudioInsights.Icon.SoundDashboard.MetaSound",
+				FText::GetEmpty(),
+				FLinearColor(0.008f, 0.76f, 0.078f),
+				FilterCategory),
+
+				MakeShared<FSoundDashboardFilter>(
+				ESoundDashboardFilterFlags::SoundCue,
+				"SoundCue",
+				LOCTEXT("SoundDashboard_Filter_SoundCueNameText", "Sound Cue"),
+				"AudioInsights.Icon.SoundDashboard.SoundCue",
+				FText::GetEmpty(),
+				FLinearColor(0.022f, 0.49f, 0.98f),
+				FilterCategory),
+
+				MakeShared<FSoundDashboardFilter>(
+				ESoundDashboardFilterFlags::ProceduralSource,
+				"ProceduralSource",
+				LOCTEXT("SoundDashboard_Filter_ProceduralSourceNameText", "Procedural Source"),
+				"AudioInsights.Icon.SoundDashboard.ProceduralSource",
+				FText::GetEmpty(),
+				FLinearColor(0.98f, 0.32f, 0.006f),
+				FilterCategory)
+			};
+
+			SAssignNew(SoundsFilterBar, SAudioFilterBar<ESoundDashboardFilterFlags>)
+			.CustomFilters(Filters)
+			// @TODO UE-220744: OnFilterChanged_Lambda not getting invoked, pending to investigate
+			.OnFilterChanged_Lambda([this, Filters]()
+			{
+				SelectedFilterFlags = ESoundDashboardFilterFlags::None;
+
+				for(const TSharedRef<FFilterBase<ESoundDashboardFilterFlags>>& Filter : Filters)
+				{
+					TSharedRef<FSoundDashboardFilter> SoundDashboardFilter = StaticCastSharedRef<FSoundDashboardFilter>(Filter);
+
+					if(SoundDashboardFilter->IsActive())
+					{
+						SelectedFilterFlags |= SoundDashboardFilter->GetFlags();
+					}
+				}
+
+				UpdateFilterReason = EProcessReason::FilterUpdated;
+			});
+		}
+
+		return SoundsFilterBar;
+	}
+
+	TSharedPtr<SWidget> FSoundDashboardViewFactory::GetFilterBarButtonWidget()
+	{
+		if (!SoundsFilterBarButton.IsValid())
+		{
+			if (!SoundsFilterBar.IsValid())
+			{
+				GetFilterBarWidget();
+			}
+
+			SoundsFilterBarButton = SBasicFilterBar<ESoundDashboardFilterFlags>::MakeAddFilterButton(StaticCastSharedPtr<SAudioFilterBar<ESoundDashboardFilterFlags>>(SoundsFilterBar).ToSharedRef()).ToSharedPtr();
+		}
+
+		return SoundsFilterBarButton;
+	}
+
+	TSharedRef<SWidget> FSoundDashboardViewFactory::MakeWidget(TSharedRef<SDockTab> OwnerTab, const FSpawnTabArgs& SpawnTabArgs)
+	{
+		return SNew(SVerticalBox)
+#if WITH_EDITOR
+			+ SVerticalBox::Slot()
+			.AutoHeight()
+			.HAlign(HAlign_Fill)
+			.Padding(0.0f, 2.0f, 0.0f, 0.0f)
+			[
+				MakeMuteSoloWidget()
+			]
+#endif // WITH_EDITOR
+			+ SVerticalBox::Slot()
+			.HAlign(HAlign_Fill)
+			[
+				FTraceTreeDashboardViewFactory::MakeWidget(OwnerTab, SpawnTabArgs)
+			];
+	}
+
+	TSharedRef<SWidget> FSoundDashboardViewFactory::GenerateWidgetForRootColumn(const TSharedRef<FTraceTreeDashboardViewFactory::SRowWidget>& InRowWidget, const TSharedRef<IDashboardDataTreeViewEntry>& InRowData, const FName& InColumn, const FText& InValueText)
+	{
+		const FColumnData& ColumnData = GetColumns()[InColumn];
+
+		const FName IconName = ColumnData.GetIconName.IsSet() ? ColumnData.GetIconName(InRowData.Get()) : NAME_None;
+
+		if (InColumn == "Name")
+		{
+			// Custom color for root item icons
+			const auto GetIconColor = [](const TSharedPtr<IDashboardDataTreeViewEntry>& InEntry)
+			{
+				const FSoundDashboardEntry& SoundDashboardEntry = FSoundDashboardViewFactoryPrivate::CastEntry(*InEntry.Get());
+
+				switch (SoundDashboardEntry.EntryType)
+				{
+				case ESoundDashboardEntryType::MetaSound:
+					return FSlateColor(FColor::Green);
+
+				case ESoundDashboardEntryType::SoundCue:
+					return FSlateColor(FColor::Blue);
+
+				case ESoundDashboardEntryType::ProceduralSource:
+					return FSlateColor(FColor::Orange);
+
+				case ESoundDashboardEntryType::None:
+				default:
+					return FSlateColor(FColor::White);
+				}
+			};
+
+			return SNew(SHorizontalBox)
+				// Tree expander arrow
+				+ SHorizontalBox::Slot()
+				.AutoWidth()
+				[
+					SNew(SExpanderArrow, InRowWidget)
+				]
+				// Icon
+				+ SHorizontalBox::Slot()
+				.Padding(IconName != NAME_None ? 2.0f : 0.0f, 2.0f)
+				.AutoWidth()
+				[
+					IconName != NAME_None 
+						? SNew(SImage).ColorAndOpacity(GetIconColor(InRowData)).Image(FSlateStyle::Get().GetBrush(IconName)) 
+						: SNullWidget::NullWidget
+				]
+				// Text
+				+ SHorizontalBox::Slot()
+				.Padding(IconName != NAME_None ? 10.0f : 0.0f, 2.0f, 0.0f, 2.0f)
+				.AutoWidth()
+				[
+					SNew(STextBlock).Text(InValueText)
+				]
+				// Number of children text
+				+ SHorizontalBox::Slot()
+				.Padding(6.0f, 2.0f, 0.0f, 2.0f)
+				.AutoWidth()
+				[
+					SNew(STextBlock)
+					.ColorAndOpacity(FSlateColor(FColor::White.WithAlpha(128)))
+					.Text_Lambda([InRowData]()
+					{
+						return FText::FromString("(" + FString::FromInt(InRowData->Children.Num()) + ")");
+					})
+				];
+		}
+		else if (InColumn == "Mute" || InColumn == "Solo")
+		{
+			return SNew(SVerticalBox)
+				// Icon
+				+ SVerticalBox::Slot()
+				.Padding(2.0f, 2.0f, 2.0f, 2.0f)
+				.AutoHeight()
+				[
+					SNew(SImage)
+					.Image_Lambda([this, InRowData, InColumn]()
+					{
+						const FColumnData& ColumnData = GetColumns()[InColumn];
+						const FName IconName = ColumnData.GetIconName.IsSet() ? ColumnData.GetIconName(InRowData.Get()) : NAME_None;
+
+						return IconName != NAME_None ? FSlateStyle::Get().GetBrush(IconName) : nullptr;
+					})
+				];
+		}
+
+		return SNullWidget::NullWidget;
+	}
+
+	bool FSoundDashboardViewFactory::IsRootItem(const TSharedPtr<IDashboardDataTreeViewEntry>& InEntry) const
+	{
+		return FilteredEntriesListView.IsValid() && FilteredEntriesListView->GetRootItems().Contains(InEntry);
+	}
+
+	TSharedRef<SWidget> FSoundDashboardViewFactory::GenerateWidgetForColumn(const TSharedRef<FTraceTreeDashboardViewFactory::SRowWidget>& InRowWidget, const TSharedRef<IDashboardDataTreeViewEntry>& InRowData, const FName& InColumn)
+	{
+		const FColumnData& ColumnData = GetColumns()[InColumn];
+
+		const FText ValueText  = ColumnData.GetDisplayValue.IsSet() ? ColumnData.GetDisplayValue(InRowData.Get()) : FText::GetEmpty();
+		const FName& ValueIcon = ColumnData.GetIconName.IsSet() ? ColumnData.GetIconName(InRowData.Get()) : NAME_None;
+
+		if (ValueText.IsEmpty() && ValueIcon.IsNone())
+		{
+			return SNullWidget::NullWidget;
+		}
+
+		if (IsRootItem(InRowData))
+		{
+			return GenerateWidgetForRootColumn(InRowWidget, InRowData, InColumn, ValueText);
+		}
+
+		const FName IconName = ColumnData.GetIconName.IsSet() ? ColumnData.GetIconName(InRowData.Get()) : NAME_None;
+
+		return SNew(SHorizontalBox)
+			// Tree expander arrow (only for Name column)
+			+ SHorizontalBox::Slot()
+			.AutoWidth()
+			[
+				InColumn == "Name" ? SNew(SExpanderArrow, InRowWidget) : SNullWidget::NullWidget
+			]
+			// Icon (optional)
+			+ SHorizontalBox::Slot()
+			.Padding(2.0f, 2.0f, 2.0f, 2.0f)
+			.AutoWidth()
+			[
+				SNew(SImage)
+				.Image_Lambda([this, InRowData, InColumn]()
+				{
+					const FColumnData& ColumnData = GetColumns()[InColumn];
+					const FName IconName = ColumnData.GetIconName.IsSet() ? ColumnData.GetIconName(InRowData.Get()) : NAME_None;
+
+					return IconName != NAME_None ? FSlateStyle::Get().GetBrush(IconName) : nullptr;
+				})
+			]
+			// Text (optional)
+			+ SHorizontalBox::Slot()
+			.Padding(10.0f, 2.0f, 2.0f, 2.0f)
+			.AutoWidth()
+			[
+				SNew(STextBlock)
+				.Text_Lambda([this, InRowData, InColumn]()
+				{
+					const FColumnData& ColumnData = GetColumns()[InColumn];
+					const FText ValueText = ColumnData.GetDisplayValue.IsSet() ? ColumnData.GetDisplayValue(InRowData.Get()) : FText::GetEmpty();
+
+					return ValueText;
+				})
+			];
+	}
+
+	void FSoundDashboardViewFactory::ProcessEntries(FTraceTreeDashboardViewFactory::EProcessReason InReason)
+	{
+		FTraceTreeDashboardViewFactory::FilterEntries<FSoundTraceProvider>([&FilterString = GetSearchFilterText().ToString()](const IDashboardDataTreeViewEntry& Entry)
+		{
+			const FSoundDashboardEntry& SoundsEntry = FSoundDashboardViewFactoryPrivate::CastEntry(Entry);
+
+			return !SoundsEntry.GetDisplayName().ToString().Contains(FilterString);
+		});
+	}
+
+	TSharedPtr<SWidget> FSoundDashboardViewFactory::OnConstructContextMenu()
+	{
+		// Skip right click for root items
+		if (FilteredEntriesListView.IsValid() && FilteredEntriesListView->GetSelectedItems().Num() == 1 && IsRootItem(FilteredEntriesListView->GetSelectedItems()[0]))
+		{
+			return SNullWidget::NullWidget;
+		}
+
+		const FSoundDashboardCommands& Commands = FSoundDashboardCommands::Get();
+
+		constexpr bool bShouldCloseWindowAfterMenuSelection = true;
+		FMenuBuilder MenuBuilder(bShouldCloseWindowAfterMenuSelection, CommandList);
+
+		MenuBuilder.BeginSection("SoundDashboardActions", LOCTEXT("SoundDashboard_Actions_HeaderText", "Sound Options"));
+
+		{
+			MenuBuilder.AddMenuEntry(Commands.GetPinCommand(), NAME_None, TAttribute<FText>(), TAttribute<FText>(), FSlateStyle::Get().CreateIcon("AudioInsights.Icon.SoundDashboard.Pin"));
+#if WITH_EDITOR
+			MenuBuilder.AddMenuEntry(Commands.GetBrowseCommand(), NAME_None, TAttribute<FText>(), TAttribute<FText>(), FSlateStyle::Get().CreateIcon("AudioInsights.Icon.SoundDashboard.Browse"));
+			MenuBuilder.AddMenuEntry(Commands.GetEditCommand(),   NAME_None, TAttribute<FText>(), TAttribute<FText>(), FSlateStyle::Get().CreateIcon("AudioInsights.Icon.SoundDashboard.Edit"));
+#endif // WITH_EDITOR
+			MenuBuilder.AddMenuEntry(Commands.GetHideCommand(), NAME_None, TAttribute<FText>(), TAttribute<FText>(), FSlateStyle::Get().CreateIcon("AudioInsights.Icon.SoundDashboard.Hide"));
+			MenuBuilder.AddMenuEntry(Commands.GetInfoCommand(), NAME_None, TAttribute<FText>(), TAttribute<FText>(), FSlateStyle::Get().CreateIcon("AudioInsights.Icon.SoundDashboard.Info"));
+		}
+
+		MenuBuilder.EndSection();
+
+		return MenuBuilder.MakeWidget();
+	}
+
+	const TMap<FName, FTraceTreeDashboardViewFactory::FHeaderRowColumnData>& FSoundDashboardViewFactory::GetHeaderRowColumns() const
+	{
+		static const TMap<FName, FTraceTreeDashboardViewFactory::FHeaderRowColumnData> HeaderRowColumnData =
+		{
+#if WITH_EDITOR
+			{
+				"Mute",
+				{
+					.DisplayName      = LOCTEXT("SoundDashboard_MuteColumnDisplayName", "Mute"),
+					.IconName         = "AudioInsights.Icon.SoundDashboard.Mute",
+					.bShowDisplayName = false,
+					.bDefaultHidden   = false,
+					.FillWidth        = 0.04f,
+					.Alignment        = EHorizontalAlignment::HAlign_Left
+				}
+			},
+			{
+				"Solo",
+				{
+					.DisplayName      = LOCTEXT("SoundDashboard_SoloColumnDisplayName", "Solo"),
+					.IconName         = "AudioInsights.Icon.SoundDashboard.Solo",
+					.bShowDisplayName = false,
+					.bDefaultHidden   = false,
+					.FillWidth        = 0.04f,
+					.Alignment        = EHorizontalAlignment::HAlign_Left
+				}
+			},
+#endif
+			{
+				"Name",
+				{
+					.DisplayName      = LOCTEXT("SoundDashboard_NameColumnDisplayName", "Name"),
+					.IconName         = NAME_None,
+					.bShowDisplayName = true,
+					.bDefaultHidden   = false,
+					.FillWidth        = 0.52f,
+					.Alignment        = EHorizontalAlignment::HAlign_Left
+				}
+			},
+			{
+				"Priority",
+				{
+					.DisplayName = LOCTEXT("SoundDashboard_PriorityColumnDisplayName", "Priority"),
+					.IconName = NAME_None,
+					.bShowDisplayName = true,
+					.bDefaultHidden = false,
+					.FillWidth = 0.08f,
+					.Alignment = EHorizontalAlignment::HAlign_Left
+				}
+			},
+			{
+				"Distance",
+				{
+					.DisplayName      = LOCTEXT("SoundDashboard_DistanceColumnDisplayName", "Distance"),
+					.IconName         = NAME_None,
+					.bShowDisplayName = true,
+					.bDefaultHidden   = false,
+					.FillWidth        = 0.1f,
+					.Alignment        = EHorizontalAlignment::HAlign_Left
+				}
+			},
+			{
+				"Amplitude",
+				{
+					.DisplayName      = LOCTEXT("SoundDashboard_AmplitudeColumnDisplayName", "Amp (Peak)"),
+					.IconName         = NAME_None,
+					.bShowDisplayName = true,
+					.bDefaultHidden   = false,
+					.FillWidth        = 0.12f,
+					.Alignment        = EHorizontalAlignment::HAlign_Left
+				}
+			},
+			{
+				"Volume",
+				{
+					.DisplayName      = LOCTEXT("SoundDashboard_VolumeColumnDisplayName", "Volume"),
+					.IconName         = NAME_None,
+					.bShowDisplayName = true,
+					.bDefaultHidden   = false,
+					.FillWidth        = 0.1f,
+					.Alignment        = EHorizontalAlignment::HAlign_Left
+				}
+			},
+			{
+				"Pitch",
+				{
+					.DisplayName      = LOCTEXT("SoundDashboard_PitchColumnDisplayName", "Pitch"),
+					.IconName         = NAME_None,
+					.bShowDisplayName = true,
+					.bDefaultHidden   = false,
+					.FillWidth        = 0.1f,
+					.Alignment        = EHorizontalAlignment::HAlign_Left
+				}
+			}
+		};
+
+		return HeaderRowColumnData;
+	}
+
+	const TMap<FName, FTraceTreeDashboardViewFactory::FColumnData>& FSoundDashboardViewFactory::GetColumns() const
+	{
+		using namespace FSoundDashboardViewFactoryPrivate;
+
+		static const TMap<FName, FTraceTreeDashboardViewFactory::FColumnData> ColumnData =
+		{
+#if WITH_EDITOR
+			{
+				"Mute",
+				{
+					.GetIconName = [this](const IDashboardDataTreeViewEntry& InData)
+					{
+						return IsMuted(InData) ? FName("AudioInsights.Icon.SoundDashboard.Mute") : NAME_None;
+					}
+				}
+			},
+			{
+				"Solo",
+				{
+					.GetIconName = [this](const IDashboardDataTreeViewEntry& InData)
+					{
+						return IsSoloed(InData) ? FName("AudioInsights.Icon.SoundDashboard.Solo") : NAME_None;
+					}
+				}
+			},
+#endif
+			{
+				"Name",
+				{
+					.GetDisplayValue = [](const IDashboardDataTreeViewEntry& InData)
+					{
+						const FSoundDashboardEntry& SoundDashboardEntry = CastEntry(InData);
+
+						return SoundDashboardEntry.GetDisplayName();
+					},
+					.GetIconName = [this](const IDashboardDataTreeViewEntry& InData) -> FName
+					{
+						const FSoundDashboardEntry& SoundDashboardEntry = CastEntry(InData);
+
+						switch (SoundDashboardEntry.EntryType)
+						{
+							case ESoundDashboardEntryType::MetaSound:
+								return FName("AudioInsights.Icon.SoundDashboard.MetaSound");
+							case ESoundDashboardEntryType::SoundCue:
+								return FName("AudioInsights.Icon.SoundDashboard.SoundCue");
+							case ESoundDashboardEntryType::ProceduralSource:
+								return FName("AudioInsights.Icon.SoundDashboard.ProceduralSource");
+							case ESoundDashboardEntryType::None:
+							default:
+								break;
+						}
+
+						return NAME_None;
+					}
+				}
+			},
+			{
+				"Priority",
+				{
+					.GetDisplayValue = [](const IDashboardDataTreeViewEntry& InData)
+					{
+						const FSoundDashboardEntry& SoundDashboardEntry = CastEntry(InData);
+
+						const ::Audio::TCircularAudioBuffer<FDataPoint>& PriorityDataPoints = SoundDashboardEntry.PriorityDataPoints;
+						return FText::AsNumber(GetLastEntryArrayValue(PriorityDataPoints), FSlateStyle::Get().GetAmpFloatFormat());
+					}
+				}
+			},
+			{
+				"Distance",
+				{
+					.GetDisplayValue = [](const IDashboardDataTreeViewEntry& InData)
+					{
+						const FSoundDashboardEntry& SoundDashboardEntry = CastEntry(InData);
+
+						const ::Audio::TCircularAudioBuffer<FDataPoint>& DistanceDataPoints = SoundDashboardEntry.DistanceDataPoints;
+						return FText::AsNumber(GetLastEntryArrayValue(DistanceDataPoints), FSlateStyle::Get().GetAmpFloatFormat());
+					}
+				}
+			},
+			{
+				"Amplitude",
+				{
+					.GetDisplayValue = [](const IDashboardDataTreeViewEntry& InData)
+					{
+						const FSoundDashboardEntry& SoundDashboardEntry = CastEntry(InData);
+
+						const ::Audio::TCircularAudioBuffer<FDataPoint>& AmplitudeDataPoints = SoundDashboardEntry.AmplitudeDataPoints;
+						return FText::AsNumber(GetLastEntryArrayValue(AmplitudeDataPoints), FSlateStyle::Get().GetAmpFloatFormat());
+					}
+				}
+			},
+			{
+				"Volume",
+				{
+					.GetDisplayValue = [](const IDashboardDataTreeViewEntry& InData)
+					{
+						const FSoundDashboardEntry& SoundDashboardEntry = CastEntry(InData);
+
+						const ::Audio::TCircularAudioBuffer<FDataPoint>& VolumeDataPoints = SoundDashboardEntry.VolumeDataPoints;
+						return FText::AsNumber(GetLastEntryArrayValue(VolumeDataPoints), FSlateStyle::Get().GetAmpFloatFormat());
+					}
+				}
+			},
+			{
+				"Pitch",
+				{
+					.GetDisplayValue = [](const IDashboardDataTreeViewEntry& InData)
+					{
+						const FSoundDashboardEntry& SoundDashboardEntry = CastEntry(InData);
+
+						const ::Audio::TCircularAudioBuffer<FDataPoint>& PitchDataPoints = SoundDashboardEntry.PitchDataPoints;
+						return FText::AsNumber(GetLastEntryArrayValue(PitchDataPoints), FSlateStyle::Get().GetPitchFloatFormat());
+					}
+				}
+			}
+		};
+
+		return ColumnData;
+	}
+
+	void FSoundDashboardViewFactory::SortTable()
+	{
+		using namespace FSoundDashboardViewFactoryPrivate;
+
+		if (SortByColumn == "Name")
+		{
+			if (SortMode == EColumnSortMode::Ascending)
+			{
+				DataViewEntries.Sort([](const TSharedPtr<IDashboardDataTreeViewEntry>& A, const TSharedPtr<IDashboardDataTreeViewEntry>& B)
+				{
+					const FSoundDashboardEntry& AData = CastEntry(*A.Get());
+					const FSoundDashboardEntry& BData = CastEntry(*B.Get());
+
+					return AData.GetDisplayName().CompareToCaseIgnored(BData.GetDisplayName()) < 0;
+				});
+			}
+			else if (SortMode == EColumnSortMode::Descending)
+			{
+				DataViewEntries.Sort([](const TSharedPtr<IDashboardDataTreeViewEntry>& A, const TSharedPtr<IDashboardDataTreeViewEntry>& B)
+				{
+					const FSoundDashboardEntry& AData = CastEntry(*A.Get());
+					const FSoundDashboardEntry& BData = CastEntry(*B.Get());
+
+					return BData.GetDisplayName().CompareToCaseIgnored(AData.GetDisplayName()) < 0;
+				});
+			}
+		}
+		else if (SortByColumn == "Priority")
+		{
+			if (SortMode == EColumnSortMode::Ascending)
+			{
+				DataViewEntries.Sort([](const TSharedPtr<IDashboardDataTreeViewEntry>& A, const TSharedPtr<IDashboardDataTreeViewEntry>& B)
+				{
+					const FSoundDashboardEntry& AData = CastEntry(*A.Get());
+					const FSoundDashboardEntry& BData = CastEntry(*B.Get());
+
+					return GetLastEntryArrayValue(AData.PriorityDataPoints) < GetLastEntryArrayValue(BData.PriorityDataPoints);
+				});
+			}
+			else if (SortMode == EColumnSortMode::Descending)
+			{
+				DataViewEntries.Sort([](const TSharedPtr<IDashboardDataTreeViewEntry>& A, const TSharedPtr<IDashboardDataTreeViewEntry>& B)
+				{
+					const FSoundDashboardEntry& AData = CastEntry(*A.Get());
+					const FSoundDashboardEntry& BData = CastEntry(*B.Get());
+
+					return GetLastEntryArrayValue(BData.PriorityDataPoints) < GetLastEntryArrayValue(AData.PriorityDataPoints);
+				});
+			}
+		}
+		else if (SortByColumn == "Distance")
+		{
+			if (SortMode == EColumnSortMode::Ascending)
+			{
+				DataViewEntries.Sort([](const TSharedPtr<IDashboardDataTreeViewEntry>& A, const TSharedPtr<IDashboardDataTreeViewEntry>& B)
+				{
+					const FSoundDashboardEntry& AData = CastEntry(*A.Get());
+					const FSoundDashboardEntry& BData = CastEntry(*B.Get());
+
+					return GetLastEntryArrayValue(AData.DistanceDataPoints) < GetLastEntryArrayValue(BData.DistanceDataPoints);
+				});
+			}
+			else if (SortMode == EColumnSortMode::Descending)
+			{
+				DataViewEntries.Sort([](const TSharedPtr<IDashboardDataTreeViewEntry>& A, const TSharedPtr<IDashboardDataTreeViewEntry>& B)
+				{
+					const FSoundDashboardEntry& AData = CastEntry(*A.Get());
+					const FSoundDashboardEntry& BData = CastEntry(*B.Get());
+
+					return GetLastEntryArrayValue(BData.DistanceDataPoints) < GetLastEntryArrayValue(AData.DistanceDataPoints);
+				});
+			}
+		}
+		else if (SortByColumn == "Amplitude")
+		{
+			if (SortMode == EColumnSortMode::Ascending)
+			{
+				DataViewEntries.Sort([](const TSharedPtr<IDashboardDataTreeViewEntry>& A, const TSharedPtr<IDashboardDataTreeViewEntry>& B)
+				{
+					const FSoundDashboardEntry& AData = CastEntry(*A.Get());
+					const FSoundDashboardEntry& BData = CastEntry(*B.Get());
+
+					return GetLastEntryArrayValue(AData.AmplitudeDataPoints) < GetLastEntryArrayValue(BData.AmplitudeDataPoints);
+				});
+			}
+			else if (SortMode == EColumnSortMode::Descending)
+			{
+				DataViewEntries.Sort([](const TSharedPtr<IDashboardDataTreeViewEntry>& A, const TSharedPtr<IDashboardDataTreeViewEntry>& B)
+				{
+					const FSoundDashboardEntry& AData = CastEntry(*A.Get());
+					const FSoundDashboardEntry& BData = CastEntry(*B.Get());
+
+					return GetLastEntryArrayValue(BData.AmplitudeDataPoints) < GetLastEntryArrayValue(AData.AmplitudeDataPoints);
+				});
+			}
+		}
+		else if (SortByColumn == "Volume")
+		{
+			if (SortMode == EColumnSortMode::Ascending)
+			{
+				DataViewEntries.Sort([](const TSharedPtr<IDashboardDataTreeViewEntry>& A, const TSharedPtr<IDashboardDataTreeViewEntry>& B)
+				{
+					const FSoundDashboardEntry& AData = CastEntry(*A.Get());
+					const FSoundDashboardEntry& BData = CastEntry(*B.Get());
+
+					return GetLastEntryArrayValue(AData.VolumeDataPoints) < GetLastEntryArrayValue(BData.VolumeDataPoints);
+				});
+			}
+			else if (SortMode == EColumnSortMode::Descending)
+			{
+				DataViewEntries.Sort([](const TSharedPtr<IDashboardDataTreeViewEntry>& A, const TSharedPtr<IDashboardDataTreeViewEntry>& B)
+				{
+					const FSoundDashboardEntry& AData = CastEntry(*A.Get());
+					const FSoundDashboardEntry& BData = CastEntry(*B.Get());
+
+					return GetLastEntryArrayValue(BData.VolumeDataPoints) < GetLastEntryArrayValue(AData.VolumeDataPoints);
+				});
+			}
+		}
+		else if (SortByColumn == "Pitch")
+		{
+			if (SortMode == EColumnSortMode::Ascending)
+			{
+				DataViewEntries.Sort([](const TSharedPtr<IDashboardDataTreeViewEntry>& A, const TSharedPtr<IDashboardDataTreeViewEntry>& B)
+				{
+					const FSoundDashboardEntry& AData = CastEntry(*A.Get());
+					const FSoundDashboardEntry& BData = CastEntry(*B.Get());
+
+					return GetLastEntryArrayValue(AData.PitchDataPoints) < GetLastEntryArrayValue(BData.PitchDataPoints);
+				});
+			}
+			else if (SortMode == EColumnSortMode::Descending)
+			{
+				DataViewEntries.Sort([](const TSharedPtr<IDashboardDataTreeViewEntry>& A, const TSharedPtr<IDashboardDataTreeViewEntry>& B)
+				{
+					const FSoundDashboardEntry& AData = CastEntry(*A.Get());
+					const FSoundDashboardEntry& BData = CastEntry(*B.Get());
+
+					return GetLastEntryArrayValue(BData.PitchDataPoints) < GetLastEntryArrayValue(AData.PitchDataPoints);
+				});
+			}
+		}
+	}
+
+#if WITH_EDITOR
+	FReply FSoundDashboardViewFactory::MuteSelected()
+	{
+#if ENABLE_AUDIO_DEBUG
+		FAudioDeviceManager* AudioDeviceManager = FAudioDeviceManager::Get();
+
+		if (AudioDeviceManager && FilteredEntriesListView.IsValid())
+		{
+			const auto MuteSound = [&AudioDeviceManager](const TSharedPtr<IDashboardDataTreeViewEntry>& SelectedItem)
+			{
+				::Audio::FAudioDebugger AudioDebugger = AudioDeviceManager->GetDebugger();
+
+				const FSoundDashboardEntry& SoundAssetDashboardEntry = *StaticCastSharedPtr<FSoundDashboardEntry>(SelectedItem).Get();
+
+				const FName SoundAssetDisplayName = FName(SoundAssetDashboardEntry.GetDisplayName().ToString());
+
+				if (SoundAssetDashboardEntry.EntryType == ESoundDashboardEntryType::SoundCue)
+				{
+					AudioDebugger.SetMuteSoundCue(SoundAssetDisplayName, true /* bInOnOff */);
+				}
+				else
+				{
+					AudioDebugger.SetMuteSoundWave(SoundAssetDisplayName, true /* bInOnOff */);
+				}
+			};
+
+			for (const TSharedPtr<IDashboardDataTreeViewEntry>& SelectedItem : FilteredEntriesListView->GetSelectedItems())
+			{
+				if (SelectedItem.IsValid())
+				{
+					MuteSound(SelectedItem);
+				}
+			}
+		}
+#endif // ENABLE_AUDIO_DEBUG
+
+		return FReply::Handled();
+	}
+
+	FReply FSoundDashboardViewFactory::SoloSelected()
+	{
+#if ENABLE_AUDIO_DEBUG
+		FAudioDeviceManager* AudioDeviceManager = FAudioDeviceManager::Get();
+
+		if (AudioDeviceManager && FilteredEntriesListView.IsValid())
+		{
+			const auto SoloSound = [&AudioDeviceManager](const TSharedPtr<IDashboardDataTreeViewEntry>& SelectedItem)
+			{
+				::Audio::FAudioDebugger AudioDebugger = AudioDeviceManager->GetDebugger();
+
+				const FSoundDashboardEntry& SoundAssetDashboardEntry = *StaticCastSharedPtr<FSoundDashboardEntry>(SelectedItem).Get();
+
+				const FName SoundAssetDisplayName = FName(SoundAssetDashboardEntry.GetDisplayName().ToString());
+
+				if (SoundAssetDashboardEntry.EntryType == ESoundDashboardEntryType::SoundCue)
+				{
+					AudioDebugger.SetSoloSoundCue(SoundAssetDisplayName, true /* bInOnOff */);
+				}
+				else
+				{
+					AudioDebugger.SetSoloSoundWave(SoundAssetDisplayName, true /* bInOnOff */);
+				}
+			};
+
+			for (const TSharedPtr<IDashboardDataTreeViewEntry>& SelectedItem : FilteredEntriesListView->GetSelectedItems())
+			{
+				if (SelectedItem.IsValid())
+				{
+					SoloSound(SelectedItem);
+				}
+			}
+		}
+#endif // ENABLE_AUDIO_DEBUG
+
+		return FReply::Handled();
+	}
+
+	bool FSoundDashboardViewFactory::IsSoloed(const IDashboardDataTreeViewEntry& InData) const
+	{
+#if ENABLE_AUDIO_DEBUG
+		if (FAudioDeviceManager* AudioDeviceManager = FAudioDeviceManager::Get())
+		{
+			const FSoundDashboardEntry& SoundAssetDashboardEntry = FSoundDashboardViewFactoryPrivate::CastEntry(InData);
+			const FName SoundAssetName(SoundAssetDashboardEntry.GetDisplayName().ToString());
+
+			return AudioDeviceManager->GetDebugger().IsSoloSoundWave(SoundAssetName);
+		}
+#endif // ENABLE_AUDIO_DEBUG
+
+		return false;
+	}
+
+	bool FSoundDashboardViewFactory::IsMuted(const IDashboardDataTreeViewEntry& InData) const
+	{
+#if ENABLE_AUDIO_DEBUG
+		if (FAudioDeviceManager* AudioDeviceManager = FAudioDeviceManager::Get())
+		{
+			const FSoundDashboardEntry& SoundAssetDashboardEntry = FSoundDashboardViewFactoryPrivate::CastEntry(InData);
+			const FName SoundAssetName(SoundAssetDashboardEntry.GetDisplayName().ToString());
+
+			return AudioDeviceManager->GetDebugger().IsMuteSoundWave(SoundAssetName);
+		}
+#endif // ENABLE_AUDIO_DEBUG
+
+		return false;
+	}
+
+	TArray<TObjectPtr<UObject>> FSoundDashboardViewFactory::GetSelectedEditableAssets() const
+	{
+		TArray<TObjectPtr<UObject>> Objects;
+
+		if (!FilteredEntriesListView.IsValid())
+		{
+			return Objects;
+		}
+
+		const TArray<TSharedPtr<IDashboardDataTreeViewEntry>> SelectedItems = FilteredEntriesListView->GetSelectedItems();
+
+		Algo::TransformIf(SelectedItems, Objects,
+			[](const TSharedPtr<IDashboardDataTreeViewEntry>& SelectedItem)
+			{
+				if (SelectedItem.IsValid())
+				{
+					IObjectTreeDashboardEntry& RowData = *StaticCastSharedPtr<IObjectTreeDashboardEntry>(SelectedItem).Get();
+
+					if (TObjectPtr<UObject> Object = RowData.GetObject())
+					{
+						return Object->IsAsset();
+					}
+				}
+
+				return false;
+			},
+			[](const TSharedPtr<IDashboardDataTreeViewEntry>& SelectedItem) -> TObjectPtr<UObject>
+			{
+				if (SelectedItem.IsValid())
+				{
+					IObjectTreeDashboardEntry& RowData = *StaticCastSharedPtr<IObjectTreeDashboardEntry>(SelectedItem).Get();
+					return RowData.GetObject();
+				}
+
+				return nullptr;
+			}
+		);
+
+		return Objects;
+	}
+#endif // WITH_EDITOR
+
+	void FSoundDashboardViewFactory::PinSound()
+	{
+		// @TODO UE-220744: pending to implement
+	}
+
+#if WITH_EDITOR
+	void FSoundDashboardViewFactory::BrowseSoundAsset() const
+	{
+		if (GEditor)
+		{
+			TArray<UObject*> EditableAssets = GetSelectedEditableAssets();
+			GEditor->SyncBrowserToObjects(EditableAssets);
+		}
+	}
+
+	void FSoundDashboardViewFactory::OpenSoundAsset() const
+	{
+		if (GEditor && FilteredEntriesListView.IsValid())
+		{
+			TArray<UObject*> Objects = GetSelectedEditableAssets();
+			if (UAssetEditorSubsystem* AssetSubsystem = GEditor->GetEditorSubsystem<UAssetEditorSubsystem>())
+			{
+				AssetSubsystem->OpenEditorForAssets(Objects);
+			}
+		}
+	}
+#endif // WITH_EDITOR
+
+	void FSoundDashboardViewFactory::HideSound()
+	{
+		// @TODO UE-220744: pending to implement
+	}
+
+	void FSoundDashboardViewFactory::SoundInfo() const
+	{
+		// @TODO UE-220744: pending to implement
+	}
+} // namespace UE::Audio::Insights
+
+#undef LOCTEXT_NAMESPACE
diff --git a/Engine/Plugins/AudioInsights/Source/AudioInsights/Public/Views/SoundDashboardViewFactory.h b/Engine/Plugins/AudioInsights/Source/AudioInsights/Public/Views/SoundDashboardViewFactory.h
new file mode 100644
--- /dev/null
+++ b/Engine/Plugins/AudioInsights/Source/AudioInsights/Public/Views/SoundDashboardViewFactory.h
@@ -0,0 +1,109 @@
+// Copyright Epic Games, Inc. All Rights Reserved.
+#pragma once
+
+#include "Views/TreeDashboardViewFactory.h"
+
+#include "Filters/GenericFilter.h"
+
+class FUICommandList;
+
+namespace UE::Audio::Insights
+{
+	enum class ESoundDashboardFilterFlags : uint32
+	{
+		None = 0,
+		MetaSound = 1 << 0,
+		SoundCue  = 1 << 1,
+		ProceduralSource = 1 << 2,
+	};
+
+	ENUM_CLASS_FLAGS(ESoundDashboardFilterFlags);
+
+	class FSoundDashboardFilter : public FGenericFilter<ESoundDashboardFilterFlags>
+	{
+	public:
+		FSoundDashboardFilter(ESoundDashboardFilterFlags InFlags, 
+			const FString& InName, 
+			const FText& InDisplayName, 
+			const FName& InIconName,
+			const FText& InToolTipText, 
+			FLinearColor InColor, 
+			TSharedPtr<FFilterCategory> InCategory)
+			: FGenericFilter<ESoundDashboardFilterFlags>(InCategory, InName, InDisplayName, FGenericFilter<ESoundDashboardFilterFlags>::FOnItemFiltered())
+			, Flags(InFlags)
+		{
+			ToolTip  = InToolTipText;
+			Color    = InColor;
+			IconName = InIconName;
+		}
+
+		bool IsActive() const {	return bIsActive; }
+		ESoundDashboardFilterFlags GetFlags() const { return Flags; }
+
+	private:
+		virtual void ActiveStateChanged(bool bActive) override { bIsActive = bActive; }
+		virtual bool PassesFilter(ESoundDashboardFilterFlags InItem) const override { return EnumHasAnyFlags(InItem, Flags); }
+
+		ESoundDashboardFilterFlags Flags;
+		bool bIsActive = false;
+	};
+
+	class AUDIOINSIGHTS_API FSoundDashboardViewFactory : public FTraceTreeDashboardViewFactory
+	{
+	public:
+		FSoundDashboardViewFactory();
+		virtual ~FSoundDashboardViewFactory();
+
+		virtual FName GetName() const override;
+		virtual FText GetDisplayName() const override;
+		virtual FSlateIcon GetIcon() const override;
+		virtual EDefaultDashboardTabStack GetDefaultTabStack() const override;
+		virtual TSharedRef<SWidget> MakeWidget(TSharedRef<SDockTab> OwnerTab, const FSpawnTabArgs& SpawnTabArgs) override;
+		virtual void ProcessEntries(FTraceTreeDashboardViewFactory::EProcessReason Reason) override;
+
+	private:
+		void BindCommands();
+
+		virtual TSharedPtr<SWidget> GetFilterBarWidget() override;
+		virtual TSharedPtr<SWidget> GetFilterBarButtonWidget() override;
+
+		bool IsRootItem(const TSharedPtr<IDashboardDataTreeViewEntry>& InEntry) const;
+
+		TSharedRef<SWidget> GenerateWidgetForRootColumn(const TSharedRef<FTraceTreeDashboardViewFactory::SRowWidget>& InRowWidget, const TSharedRef<IDashboardDataTreeViewEntry>& InRowData, const FName& InColumn, const FText& InValueText);
+		virtual TSharedRef<SWidget> GenerateWidgetForColumn(const TSharedRef<FTraceTreeDashboardViewFactory::SRowWidget>& InRowWidget, const TSharedRef<IDashboardDataTreeViewEntry>& InRowData, const FName& InColumn) override;
+		
+		virtual TSharedPtr<SWidget> OnConstructContextMenu() override;
+
+		virtual const TMap<FName, FTraceTreeDashboardViewFactory::FHeaderRowColumnData>& GetHeaderRowColumns() const override;
+		virtual const TMap<FName, FTraceTreeDashboardViewFactory::FColumnData>& GetColumns() const override;
+
+		virtual void SortTable() override;
+
+#if WITH_EDITOR
+		TSharedRef<SWidget> MakeMuteSoloWidget();
+
+		FReply MuteSelected();
+		FReply SoloSelected();
+
+		bool IsSoloed(const IDashboardDataTreeViewEntry& InData) const;
+		bool IsMuted(const IDashboardDataTreeViewEntry& InData) const;
+
+		TArray<TObjectPtr<UObject>> GetSelectedEditableAssets() const;
+#endif // WITH_EDITOR
+
+		void PinSound();
+
+#if WITH_EDITOR
+		void BrowseSoundAsset() const;
+		void OpenSoundAsset() const;
+#endif
+		void HideSound();
+		void SoundInfo() const;
+
+		TSharedPtr<FUICommandList> CommandList;
+		TSharedPtr<SWidget> SoundsFilterBar;
+		TSharedPtr<SWidget> SoundsFilterBarButton;
+
+		ESoundDashboardFilterFlags SelectedFilterFlags = ESoundDashboardFilterFlags::None;
+	};
+} // namespace UE::Audio::Insights
diff --git a/Engine/Plugins/AudioInsights/Source/AudioInsightsEditor/Private/AudioInsightsEditorModule.cpp b/Engine/Plugins/AudioInsights/Source/AudioInsightsEditor/Private/AudioInsightsEditorModule.cpp
--- a/Engine/Plugins/AudioInsights/Source/AudioInsightsEditor/Private/AudioInsightsEditorModule.cpp
+++ b/Engine/Plugins/AudioInsights/Source/AudioInsightsEditor/Private/AudioInsightsEditorModule.cpp
@@ -13,6 +13,7 @@
 #include "Views/AudioMetersDashboardViewFactory.h"
 #include "Views/LogDashboardViewFactory.h"
 #include "Views/MixerSourceDashboardViewFactory.h"
+#include "Views/SoundDashboardViewFactory.h"
 #include "Views/SubmixesDashboardViewFactory.h"
 #include "Views/ViewportDashboardViewFactory.h"
 #include "Views/VirtualLoopDashboardViewFactory.h"
@@ -40,6 +41,7 @@ namespace UE::Audio::Insights
 			DashboardFactory->RegisterViewFactory(MakeShared<FVirtualLoopDashboardViewFactory>());
 			DashboardFactory->RegisterViewFactory(MakeShared<FSubmixesDashboardViewFactory>());
 			DashboardFactory->RegisterViewFactory(MakeShared<FAudioBusesDashboardViewFactory>());
+			DashboardFactory->RegisterViewFactory(MakeShared<FSoundDashboardViewFactory>());
 			DashboardFactory->RegisterViewFactory(MakeShared<FAudioMetersDashboardViewFactory>());
 			DashboardFactory->RegisterViewFactory(MakeShared<FAudioAnalyzerRackDashboardViewFactory>());
 		}
