[Audio Insights] Added missing typename to fix issue triggered by builds when submitting CL 41026884 (Sound Dashboard part 7)

diff --git a/Engine/Plugins/AudioInsights/Source/AudioInsights/Public/Views/SAudioFilterBar.h b/Engine/Plugins/AudioInsights/Source/AudioInsights/Public/Views/SAudioFilterBar.h
--- a/Engine/Plugins/AudioInsights/Source/AudioInsights/Public/Views/SAudioFilterBar.h
+++ b/Engine/Plugins/AudioInsights/Source/AudioInsights/Public/Views/SAudioFilterBar.h
@@ -13,13 +13,13 @@ namespace UE::Audio::Insights
 	public:
 		SLATE_BEGIN_ARGS(SAudioFilterBar){}
 			SLATE_ARGUMENT(TArray<TSharedRef<FFilterBase<EnumType>>>, CustomFilters)
-			SLATE_EVENT(SBasicFilterBar<EnumType>::FOnFilterChanged, OnFilterChanged)
+			SLATE_EVENT(typename SBasicFilterBar<EnumType>::FOnFilterChanged, OnFilterChanged)
 		SLATE_END_ARGS()
 
 		void Construct(const FArguments& InArgs)
 		{
 			SBasicFilterBar<EnumType>::Construct(
-				SBasicFilterBar<EnumType>::FArguments()
+				typename SBasicFilterBar<EnumType>::FArguments()
 				.CustomFilters(InArgs._CustomFilters)
 				.CanChangeOrientation(false)
 				.FilterBarLayout(EFilterBarLayout::Horizontal)
