[Audio Insights] Adjusted even more fix of issue triggered by builds when submitting CL 41026884 (Sound Dashboard part 7)

diff --git a/Engine/Plugins/AudioInsights/Source/AudioInsights/Public/Views/SAudioFilterBar.h b/Engine/Plugins/AudioInsights/Source/AudioInsights/Public/Views/SAudioFilterBar.h
--- a/Engine/Plugins/AudioInsights/Source/AudioInsights/Public/Views/SAudioFilterBar.h
+++ b/Engine/Plugins/AudioInsights/Source/AudioInsights/Public/Views/SAudioFilterBar.h
@@ -11,16 +11,20 @@ namespace UE::Audio::Insights
 	class SAudioFilterBar : public SBasicFilterBar<EnumType>
 	{
 	public:
+		using FOnFilterChanged = typename SBasicFilterBar<EnumType>::FOnFilterChanged;
+
 		SLATE_BEGIN_ARGS(SAudioFilterBar){}
 			SLATE_ARGUMENT(TArray<TSharedRef<FFilterBase<EnumType>>>, CustomFilters)
-			SLATE_EVENT(typename SBasicFilterBar<EnumType>::FOnFilterChanged, OnFilterChanged)
+			SLATE_EVENT(SAudioFilterBar<EnumType>::FOnFilterChanged, OnFilterChanged)
 		SLATE_END_ARGS()
 
 		void Construct(const FArguments& InArgs)
 		{
-			SBasicFilterBar<EnumType>::Construct(
-				typename SBasicFilterBar<EnumType>::FArguments()
-				.CustomFilters(InArgs._CustomFilters)
+			typename SBasicFilterBar<EnumType>::FArguments Args;
+			Args._OnFilterChanged = InArgs._OnFilterChanged;
+			Args._CustomFilters = InArgs._CustomFilters;
+
+			SBasicFilterBar<EnumType>::Construct(Args
 				.CanChangeOrientation(false)
 				.FilterBarLayout(EFilterBarLayout::Horizontal)
 				.UseSectionsForCategories(true)
