[Audio Insights] Don't enable Audio or AudioMixer trace channel on Provider construction, only when an Audio Insights trace begins.

diff --git a/Engine/Plugins/AudioInsights/Source/AudioInsights/Private/AudioInsightsTraceChannelHandle.cpp b/Engine/Plugins/AudioInsights/Source/AudioInsights/Private/AudioInsightsTraceChannelHandle.cpp
--- a/Engine/Plugins/AudioInsights/Source/AudioInsights/Private/AudioInsightsTraceChannelHandle.cpp
+++ b/Engine/Plugins/AudioInsights/Source/AudioInsights/Private/AudioInsightsTraceChannelHandle.cpp
@@ -79,8 +79,6 @@ namespace UE::Audio::Insights
 			else
 			{
 				Manager->ActiveChannels.Add(ChannelName, 1);
-				constexpr bool bIsEnabled = true;
-				Trace::ToggleChannel(*ChannelName, bIsEnabled);
 			}
 		}
 	}
@@ -95,8 +93,6 @@ namespace UE::Audio::Insights
 			RefCount--;
 			if (RefCount == 0)
 			{
-				constexpr bool bIsEnabled = false;
-				Trace::ToggleChannel(*ChannelName, bIsEnabled);
 				Manager->ActiveChannels.Remove(ChannelName);
 			}
 		}
