[Audio Insights] - Values don't clear from Audio Insights when closed and reopened
* Bug occurs when a sound is active when the standalone version of Unreal Insights is closed, and stops before it is reopened. Doing so results in the stop message being missed by [REDACTED_PERSON], and "ghost" active sound entries being left in the sound list.
* Now in the standalone version of Unreal Insights, instead of adding a device entry when we receive a start message, we instead cache the message and wait for a corresponding volume message to also be sent. Active sounds send their volume data every update, so we remove "ghost" entries from the list.

diff --git a/Engine/Plugins/AudioInsights/Source/AudioInsights/Private/Providers/MixerSourceTraceProvider.cpp b/Engine/Plugins/AudioInsights/Source/AudioInsights/Private/Providers/MixerSourceTraceProvider.cpp
--- a/Engine/Plugins/AudioInsights/Source/AudioInsights/Private/Providers/MixerSourceTraceProvider.cpp
+++ b/Engine/Plugins/AudioInsights/Source/AudioInsights/Private/Providers/MixerSourceTraceProvider.cpp
@@ -253,10 +253,21 @@ namespace UE::Audio::Insights
 
 	bool FMixerSourceTraceProvider::ProcessMessages()
 	{
-		auto GetOrCreateEntry = [this](const FMixerSourceMessageBase& Msg)
+		auto GetOrCreateEntry = [this](const FMixerSourceStartMessage& Msg)
 		{
 			TSharedPtr<FMixerSourceDashboardEntry>* ToReturn = nullptr;
-			UpdateDeviceEntry(Msg.DeviceId, Msg.PlayOrder, [&ToReturn, &Msg](TSharedPtr<FMixerSourceDashboardEntry>& Entry)
+
+#if !WITH_EDITOR
+			// In standalone Unreal Insights, we don't necessarily know if a sound is still active until we have a corresponding volume message.
+			// Cache any start messages received and only create a device entry once we have also received a volume message.
+			if (SessionCachedMessages.IsValid())
+			{
+				FMixerSourceStartMessage& StartMessage = SessionCachedMessages->PlayOrderToStartMessageMap.FindOrAdd(Msg.PlayOrder, Msg);
+				return ToReturn;
+			}
+#endif // !WITH_EDITOR
+
+			UpdateDeviceEntry(Msg.DeviceId, Msg.PlayOrder, [this, &ToReturn, &Msg](TSharedPtr<FMixerSourceDashboardEntry>& Entry)
 			{
 				if (!Entry.IsValid())
 				{
@@ -271,6 +282,38 @@ namespace UE::Audio::Insights
 			return ToReturn;
 		};
 
+#if !WITH_EDITOR
+		auto GetOrLateCreateEntry = [this](const FMixerSourceMessageBase& Msg)
+		{
+			// If we have cached a start message and have received a volume message for the same PlayOrderID in standalone Unreal Insights,
+			// we can safetly assume that that this sound is active.
+			// Create a new entry and return the result.
+			if (FMixerSourceStartMessage* StartMessage = SessionCachedMessages->PlayOrderToStartMessageMap.Find(Msg.PlayOrder))
+			{
+				UpdateDeviceEntry(StartMessage->DeviceId, StartMessage->PlayOrder, [this, &StartMessage](TSharedPtr<FMixerSourceDashboardEntry>& Entry)
+				{
+					if (!Entry.IsValid())
+					{
+						Entry = MakeShared<FMixerSourceDashboardEntry>();
+						Entry->DeviceId = StartMessage->DeviceId;
+						Entry->PlayOrder = StartMessage->PlayOrder;
+					}
+
+					Entry->Timestamp = StartMessage->Timestamp;
+
+					Entry->Name = *StartMessage->Name;
+					Entry->ComponentId = StartMessage->ComponentId;
+					Entry->SourceId = StartMessage->SourceId;
+
+					SessionCachedMessages->StartCachedMessages.EmplaceBack(*StartMessage);
+					SessionCachedMessages->PlayOrderToStartMessageMap.Remove(StartMessage->PlayOrder);
+				});
+			}
+
+			return FindDeviceEntry(Msg.DeviceId, Msg.PlayOrder);
+		};
+#endif //!WITH_EDITOR
+
 		auto GetEntry = [this](const FMixerSourceMessageBase& Msg)
 		{
 			return FindDeviceEntry(Msg.DeviceId, Msg.PlayOrder);
@@ -280,6 +323,11 @@ namespace UE::Audio::Insights
 		[this](const FMixerSourceStartMessage& Msg, TSharedPtr<FMixerSourceDashboardEntry>* OutEntry)
 		{
 #if !WITH_EDITOR
+			if (OutEntry == nullptr || !(*OutEntry).IsValid())
+			{
+				return;
+			}
+
 			if (SessionCachedMessages.IsValid())
 			{
 				SessionCachedMessages->StartCachedMessages.EmplaceBack(Msg);
@@ -294,10 +342,15 @@ namespace UE::Audio::Insights
 
 		TArray<int32, TInlineAllocator<64>> EntriesWithPoppedDataPoints;
 
-		ProcessMessageQueue<FMixerSourceVolumeMessage>(TraceMessages.VolumeMessages, GetEntry,
+		ProcessMessageQueue<FMixerSourceVolumeMessage>(TraceMessages.VolumeMessages,
+#if WITH_EDITOR
+		GetEntry,
+#else
+		GetOrLateCreateEntry,
+#endif // WITH_EDITOR
 		[this, &EntriesWithPoppedDataPoints](const FMixerSourceVolumeMessage& Msg, TSharedPtr<FMixerSourceDashboardEntry>* OutEntry)
 		{
-			if (OutEntry == nullptr)
+			if (OutEntry == nullptr || !(*OutEntry).IsValid())
 			{
 				return;
 			}
@@ -322,7 +375,7 @@ namespace UE::Audio::Insights
 		ProcessMessageQueue<FMixerSourcePitchMessage>(TraceMessages.PitchMessages, GetEntry,
 		[this, &EntriesWithPoppedDataPoints](const FMixerSourcePitchMessage& Msg, TSharedPtr<FMixerSourceDashboardEntry>* OutEntry)
 		{
-			if (OutEntry == nullptr)
+			if (OutEntry == nullptr || !(*OutEntry).IsValid())
 			{
 				return;
 			}
@@ -347,7 +400,7 @@ namespace UE::Audio::Insights
 		ProcessMessageQueue<FMixerSourceLPFFreqMessage>(TraceMessages.LPFFreqMessages, GetEntry,
 		[this, &EntriesWithPoppedDataPoints](const FMixerSourceLPFFreqMessage& Msg, TSharedPtr<FMixerSourceDashboardEntry>* OutEntry)
 		{
-			if (OutEntry == nullptr)
+			if (OutEntry == nullptr || !(*OutEntry).IsValid())
 			{
 				return;
 			}
@@ -372,7 +425,7 @@ namespace UE::Audio::Insights
 		ProcessMessageQueue<FMixerSourceHPFFreqMessage>(TraceMessages.HPFFreqMessages, GetEntry,
 		[this, &EntriesWithPoppedDataPoints](const FMixerSourceHPFFreqMessage& Msg, TSharedPtr<FMixerSourceDashboardEntry>* OutEntry)
 		{
-			if (OutEntry == nullptr)
+			if (OutEntry == nullptr || !(*OutEntry).IsValid())
 			{
 				return;
 			}
@@ -397,7 +450,7 @@ namespace UE::Audio::Insights
 		ProcessMessageQueue<FMixerSourceEnvelopeMessage>(TraceMessages.EnvelopeMessages, GetEntry,
 		[this, &EntriesWithPoppedDataPoints](const FMixerSourceEnvelopeMessage& Msg, TSharedPtr<FMixerSourceDashboardEntry>* OutEntry)
 		{
-			if (OutEntry == nullptr)
+			if (OutEntry == nullptr || !(*OutEntry).IsValid())
 			{
 				return;
 			}
@@ -422,7 +475,7 @@ namespace UE::Audio::Insights
 		ProcessMessageQueue<FMixerSourceDistanceAttenuationMessage>(TraceMessages.DistanceAttenuationMessages, GetEntry,
 		[this, &EntriesWithPoppedDataPoints](const FMixerSourceDistanceAttenuationMessage& Msg, TSharedPtr<FMixerSourceDashboardEntry>* OutEntry)
 		{
-			if (OutEntry == nullptr)
+			if (OutEntry == nullptr || !(*OutEntry).IsValid())
 			{
 				return;
 			}
@@ -447,7 +500,7 @@ namespace UE::Audio::Insights
 		[this](const FMixerSourceStopMessage& Msg, TSharedPtr<FMixerSourceDashboardEntry>* OutEntry)
 		{
 #if !WITH_EDITOR
-			if (SessionCachedMessages.IsValid())
+			if (SessionCachedMessages.IsValid() || !(*OutEntry).IsValid())
 			{
 				SessionCachedMessages->StopCachedMessages.EmplaceBack(Msg);
 			}
diff --git a/Engine/Plugins/AudioInsights/Source/AudioInsights/Public/Messages/MixerSourceTraceMessages.h b/Engine/Plugins/AudioInsights/Source/AudioInsights/Public/Messages/MixerSourceTraceMessages.h
--- a/Engine/Plugins/AudioInsights/Source/AudioInsights/Public/Messages/MixerSourceTraceMessages.h
+++ b/Engine/Plugins/AudioInsights/Source/AudioInsights/Public/Messages/MixerSourceTraceMessages.h
@@ -137,6 +137,8 @@ namespace UE::Audio::Insights
 		TraceServices::TPagedArray<FMixerSourceStartMessage> StartCachedMessages;
 		TraceServices::TPagedArray<FMixerSourceStopMessage> StopCachedMessages;
 		TraceServices::TPagedArray<FMixerSourceVolumeMessage> VolumeCachedMessages;
+
+		TMap<uint32, FMixerSourceStartMessage> PlayOrderToStartMessageMap;
 	};
 #endif // !WITH_EDITOR
 
