[Audio Insights] Sound Dashboard: Removed old comment related to an issue with OnFilterChanged_Lambda. This issue was already fixed in CL 41029982.

diff --git a/Engine/Plugins/AudioInsights/Source/AudioInsights/Private/Views/SoundDashboardViewFactory.cpp b/Engine/Plugins/AudioInsights/Source/AudioInsights/Private/Views/SoundDashboardViewFactory.cpp
--- a/Engine/Plugins/AudioInsights/Source/AudioInsights/Private/Views/SoundDashboardViewFactory.cpp
+++ b/Engine/Plugins/AudioInsights/Source/AudioInsights/Private/Views/SoundDashboardViewFactory.cpp
@@ -206,7 +206,6 @@ namespace UE::Audio::Insights
 
 			SAssignNew(SoundsFilterBar, SAudioFilterBar<ESoundDashboardFilterFlags>)
 			.CustomFilters(Filters)
-			// @TODO UE-220744: OnFilterChanged_Lambda not getting invoked, pending to investigate
 			.OnFilterChanged_Lambda([this, Filters]()
 			{
 				SelectedFilterFlags = ESoundDashboardFilterFlags::None;
