[Audio Insights] Sound Dashboard: Show number of children in all rows except for the leaf ones

diff --git a/Engine/Plugins/AudioInsights/Source/AudioInsights/Private/Views/SoundDashboardViewFactory.cpp b/Engine/Plugins/AudioInsights/Source/AudioInsights/Private/Views/SoundDashboardViewFactory.cpp
--- a/Engine/Plugins/AudioInsights/Source/AudioInsights/Private/Views/SoundDashboardViewFactory.cpp
+++ b/Engine/Plugins/AudioInsights/Source/AudioInsights/Private/Views/SoundDashboardViewFactory.cpp
@@ -431,6 +431,20 @@ namespace UE::Audio::Insights
 
 					return ValueText;
 				})
+			]
+			// Number of children text (if is not leaf item)
+			+ SHorizontalBox::Slot()
+			.Padding(6.0f, 2.0f, 0.0f, 2.0f)
+			.AutoWidth()
+			[
+				InColumn == "Name" && !InRowData->Children.IsEmpty()
+					? SNew(STextBlock)
+					  .ColorAndOpacity(FSlateColor(FColor::White.WithAlpha(128)))
+					  .Text_Lambda([InRowData]()
+					  {
+					      return FText::FromString("(" + FString::FromInt(InRowData->Children.Num()) + ")");
+					  })
+					: SNullWidget::NullWidget
 			];
 	}
 
