[Audio Insights] Sound Dashboard: removed "More Information" item from right click context menu as agreed.

diff --git a/Engine/Plugins/AudioInsights/Source/AudioInsights/Private/SoundDashboardCommands.cpp b/Engine/Plugins/AudioInsights/Source/AudioInsights/Private/SoundDashboardCommands.cpp
--- a/Engine/Plugins/AudioInsights/Source/AudioInsights/Private/SoundDashboardCommands.cpp
+++ b/Engine/Plugins/AudioInsights/Source/AudioInsights/Private/SoundDashboardCommands.cpp
@@ -21,7 +21,6 @@ namespace UE::Audio::Insights
 		UI_COMMAND(Browse, "Browse To Asset", "Browses to the selected sound asset in the content browser.", EUserInterfaceActionType::Button, FInputChord(EKeys::B, EModifierKey::Control));
 		UI_COMMAND(Edit, "Edit", "Opens the selected sound for edit.", EUserInterfaceActionType::Button, FInputChord(EKeys::E, EModifierKey::Control));
 		UI_COMMAND(Hide, "Hide", "Hides the selected sound.", EUserInterfaceActionType::Button, FInputChord(EKeys::H, EModifierKey::Control));
-		UI_COMMAND(Info, "More Information", "Provides more information about this sound.", EUserInterfaceActionType::Button, FInputChord(EKeys::I, EModifierKey::Control));
 	}
 } // namespace UE::Audio::Insights
 
diff --git a/Engine/Plugins/AudioInsights/Source/AudioInsights/Private/SoundDashboardCommands.h b/Engine/Plugins/AudioInsights/Source/AudioInsights/Private/SoundDashboardCommands.h
--- a/Engine/Plugins/AudioInsights/Source/AudioInsights/Private/SoundDashboardCommands.h
+++ b/Engine/Plugins/AudioInsights/Source/AudioInsights/Private/SoundDashboardCommands.h
@@ -19,13 +19,11 @@ namespace UE::Audio::Insights
 		TSharedPtr<const FUICommandInfo> GetBrowseCommand() const { return Browse; }
 		TSharedPtr<const FUICommandInfo> GetEditCommand() const   { return Edit; }
 		TSharedPtr<const FUICommandInfo> GetHideCommand() const   { return Hide; }
-		TSharedPtr<const FUICommandInfo> GetInfoCommand() const   { return Info; }
 		
 	private:
 		TSharedPtr<FUICommandInfo> Pin;
 		TSharedPtr<FUICommandInfo> Browse;
 		TSharedPtr<FUICommandInfo> Edit;
 		TSharedPtr<FUICommandInfo> Hide;
-		TSharedPtr<FUICommandInfo> Info;
 	};
 } // namespace UE::Audio::Insights
diff --git a/Engine/Plugins/AudioInsights/Source/AudioInsights/Private/Views/SoundDashboardViewFactory.cpp b/Engine/Plugins/AudioInsights/Source/AudioInsights/Private/Views/SoundDashboardViewFactory.cpp
--- a/Engine/Plugins/AudioInsights/Source/AudioInsights/Private/Views/SoundDashboardViewFactory.cpp
+++ b/Engine/Plugins/AudioInsights/Source/AudioInsights/Private/Views/SoundDashboardViewFactory.cpp
@@ -94,7 +94,6 @@ namespace UE::Audio::Insights
 		CommandList->MapAction(Commands.GetEditCommand(),   FExecuteAction::CreateRaw(this, &FSoundDashboardViewFactory::OpenSoundAsset));
 #endif // WITH_EDITOR
 		CommandList->MapAction(Commands.GetHideCommand(), FExecuteAction::CreateRaw(this, &FSoundDashboardViewFactory::HideSound));
-		CommandList->MapAction(Commands.GetInfoCommand(), FExecuteAction::CreateRaw(this, &FSoundDashboardViewFactory::SoundInfo));
 	}
 
 #if WITH_EDITOR
@@ -480,7 +479,6 @@ namespace UE::Audio::Insights
 			MenuBuilder.AddMenuEntry(Commands.GetEditCommand(),   NAME_None, TAttribute<FText>(), TAttribute<FText>(), FSlateStyle::Get().CreateIcon("AudioInsights.Icon.SoundDashboard.Edit"));
 #endif // WITH_EDITOR
 			MenuBuilder.AddMenuEntry(Commands.GetHideCommand(), NAME_None, TAttribute<FText>(), TAttribute<FText>(), FSlateStyle::Get().CreateIcon("AudioInsights.Icon.SoundDashboard.Hide"));
-			MenuBuilder.AddMenuEntry(Commands.GetInfoCommand(), NAME_None, TAttribute<FText>(), TAttribute<FText>(), FSlateStyle::Get().CreateIcon("AudioInsights.Icon.SoundDashboard.Info"));
 		}
 
 		MenuBuilder.EndSection();
@@ -1039,11 +1037,6 @@ namespace UE::Audio::Insights
 	{
 		// @TODO UE-220744: pending to implement
 	}
-
-	void FSoundDashboardViewFactory::SoundInfo() const
-	{
-		// @TODO UE-220744: pending to implement
-	}
 } // namespace UE::Audio::Insights
 
 #undef LOCTEXT_NAMESPACE
diff --git a/Engine/Plugins/AudioInsights/Source/AudioInsights/Public/Views/SoundDashboardViewFactory.h b/Engine/Plugins/AudioInsights/Source/AudioInsights/Public/Views/SoundDashboardViewFactory.h
--- a/Engine/Plugins/AudioInsights/Source/AudioInsights/Public/Views/SoundDashboardViewFactory.h
+++ b/Engine/Plugins/AudioInsights/Source/AudioInsights/Public/Views/SoundDashboardViewFactory.h
@@ -99,7 +99,6 @@ namespace UE::Audio::Insights
 		void OpenSoundAsset() const;
 #endif
 		void HideSound();
-		void SoundInfo() const;
 
 		TSharedPtr<FUICommandList> CommandList;
 		TSharedPtr<SWidget> SoundsFilterBar;
