[Audio Insights] Sound Dashboard: colors for categories now in AudioInsightsStyle (so we can use the same color for the category icons in the tree view) + added SoundWave, SoundCueTemplate, Pinned and Hidden categories to filter bar + added table row style to alternate color between rows.

diff --git a/Engine/Plugins/AudioInsights/Source/AudioInsights/Private/Views/SoundDashboardViewFactory.cpp b/Engine/Plugins/AudioInsights/Source/AudioInsights/Private/Views/SoundDashboardViewFactory.cpp
--- a/Engine/Plugins/AudioInsights/Source/AudioInsights/Private/Views/SoundDashboardViewFactory.cpp
+++ b/Engine/Plugins/AudioInsights/Source/AudioInsights/Private/Views/SoundDashboardViewFactory.cpp
@@ -181,7 +181,7 @@ namespace UE::Audio::Insights
 				LOCTEXT("SoundDashboard_Filter_MetaSoundNameText", "MetaSound"),
 				"AudioInsights.Icon.SoundDashboard.MetaSound",
 				FText::GetEmpty(),
-				FLinearColor(0.008f, 0.76f, 0.078f),
+				FSlateStyle::Get().GetColor("SoundDashboard.MetaSoundColor"),
 				FilterCategory),
 
 				MakeShared<FSoundDashboardFilter>(
@@ -190,7 +190,7 @@ namespace UE::Audio::Insights
 				LOCTEXT("SoundDashboard_Filter_SoundCueNameText", "Sound Cue"),
 				"AudioInsights.Icon.SoundDashboard.SoundCue",
 				FText::GetEmpty(),
-				FLinearColor(0.022f, 0.49f, 0.98f),
+				FSlateStyle::Get().GetColor("SoundDashboard.SoundCueColor"),
 				FilterCategory),
 
 				MakeShared<FSoundDashboardFilter>(
@@ -199,7 +199,7 @@ namespace UE::Audio::Insights
 				LOCTEXT("SoundDashboard_Filter_ProceduralSourceNameText", "Procedural Source"),
 				"AudioInsights.Icon.SoundDashboard.ProceduralSource",
 				FText::GetEmpty(),
-				FLinearColor(0.98f, 0.32f, 0.006f),
+				FSlateStyle::Get().GetColor("SoundDashboard.ProceduralSourceColor"),
 				FilterCategory),
 
 				MakeShared<FSoundDashboardFilter>(
@@ -208,7 +208,7 @@ namespace UE::Audio::Insights
 				LOCTEXT("SoundDashboard_Filter_SoundWaveNameText", "Sound Wave"),
 				"AudioInsights.Icon.SoundDashboard.SoundWave",
 				FText::GetEmpty(),
-				FLinearColor(FColor::Purple),
+				FSlateStyle::Get().GetColor("SoundDashboard.SoundWaveColor"),
 				FilterCategory),
 
 				MakeShared<FSoundDashboardFilter>(
@@ -217,8 +217,26 @@ namespace UE::Audio::Insights
 				LOCTEXT("SoundDashboard_Filter_SoundCueTemplateNameText", "Sound Cue Template"),
 				"AudioInsights.Icon.SoundDashboard.SoundCue",
 				FText::GetEmpty(),
-				FLinearColor(FColor::Red),
+				FSlateStyle::Get().GetColor("SoundDashboard.SoundCueTemplateColor"),
 				FilterCategory),
+
+				MakeShared<FSoundDashboardFilter>(
+				ESoundDashboardFilterFlags::Pinned,
+				"Pinned",
+				LOCTEXT("SoundDashboard_Filter_PinnedNameText", "Pinned"),
+				"AudioInsights.Icon.SoundDashboard.Pin",
+				FText::GetEmpty(),
+				FSlateStyle::Get().GetColor("SoundDashboard.PinnedColor"),
+				FilterCategory),
+
+				MakeShared<FSoundDashboardFilter>(
+				ESoundDashboardFilterFlags::Hidden,
+				"Hidden",
+				LOCTEXT("SoundDashboard_Filter_HiddenNameText", "Hidden"),
+				"AudioInsights.Icon.SoundDashboard.Hide",
+				FText::GetEmpty(),
+				FSlateStyle::Get().GetColor("SoundDashboard.HiddenColor"),
+				FilterCategory)
 			};
 
 			SAssignNew(SoundsFilterBar, SAudioFilterBar<ESoundDashboardFilterFlags>)
@@ -294,19 +312,19 @@ namespace UE::Audio::Insights
 				switch (SoundDashboardEntry.EntryType)
 				{
 				case ESoundDashboardEntryType::MetaSound:
-					return FSlateColor(FColor::Green);
+					return FSlateColor(FSlateStyle::Get().GetColor("SoundDashboard.MetaSoundColor"));
 
 				case ESoundDashboardEntryType::SoundCue:
-					return FSlateColor(FColor::Blue);
+					return FSlateColor(FSlateStyle::Get().GetColor("SoundDashboard.SoundCueColor"));
 
 				case ESoundDashboardEntryType::ProceduralSource:
-					return FSlateColor(FColor::Orange);
+					return FSlateColor(FSlateStyle::Get().GetColor("SoundDashboard.ProceduralSourceColor"));
 
 				case ESoundDashboardEntryType::SoundWave:
-					return FSlateColor(FColor::Purple);
+					return FSlateColor(FSlateStyle::Get().GetColor("SoundDashboard.SoundWaveColor"));
 
 				case ESoundDashboardEntryType::SoundCueTemplate:
-					return FSlateColor(FColor::Red);
+					return FSlateColor(FSlateStyle::Get().GetColor("SoundDashboard.SoundCueTemplateColor"));
 
 				case ESoundDashboardEntryType::None:
 				default:
diff --git a/Engine/Plugins/AudioInsights/Source/AudioInsights/Private/Views/TreeDashboardViewFactory.cpp b/Engine/Plugins/AudioInsights/Source/AudioInsights/Private/Views/TreeDashboardViewFactory.cpp
--- a/Engine/Plugins/AudioInsights/Source/AudioInsights/Private/Views/TreeDashboardViewFactory.cpp
+++ b/Engine/Plugins/AudioInsights/Source/AudioInsights/Private/Views/TreeDashboardViewFactory.cpp
@@ -22,7 +22,10 @@ namespace UE::Audio::Insights
 	{
 		Data = InData;
 		Factory = InFactory;
-		SMultiColumnTableRow<TSharedPtr<IDashboardDataTreeViewEntry>>::Construct(FSuperRowType::FArguments(), InOwnerTable);
+
+		FSuperRowType::FArguments Args = FSuperRowType::FArguments().Style(&FSlateStyle::Get().GetWidgetStyle<FTableRowStyle>("SoundDashboard.TableViewRow"));
+
+		SMultiColumnTableRow<TSharedPtr<IDashboardDataTreeViewEntry>>::Construct(Args, InOwnerTable);
 	}
 
 	TSharedRef<SWidget> FTraceTreeDashboardViewFactory::SRowWidget::GenerateWidgetForColumn(const FName& Column)
diff --git a/Engine/Plugins/AudioInsights/Source/AudioInsights/Public/AudioInsightsStyle.h b/Engine/Plugins/AudioInsights/Source/AudioInsights/Public/AudioInsightsStyle.h
--- a/Engine/Plugins/AudioInsights/Source/AudioInsights/Public/AudioInsightsStyle.h
+++ b/Engine/Plugins/AudioInsights/Source/AudioInsights/Public/AudioInsightsStyle.h
@@ -6,6 +6,8 @@
 #include "Styling/SlateStyle.h"
 #include "Styling/SlateStyleMacros.h"
 #include "Styling/SlateStyleRegistry.h"
+#include "Styling/SlateTypes.h"
+#include "Styling/StyleColors.h"
 #include "Textures/SlateIcon.h"
 
 #define LOCTEXT_NAMESPACE "AudioInsights"
@@ -147,7 +149,8 @@ namespace UE::Audio::Insights
 
 			Set("AudioInsights.Thumbnail", new IMAGE_BRUSH_SVG(TEXT("Icons/audio_insights"), Icon16));
 
-			// Sound Dashboard icons
+			/* Sound Dashboard */
+			// Icons
 			Set("AudioInsights.Icon.SoundDashboard.Browse",           new IMAGE_BRUSH_SVG(TEXT("Icons/SoundDashboard/browse"),           Icon16));
 			Set("AudioInsights.Icon.SoundDashboard.Edit",             new IMAGE_BRUSH_SVG(TEXT("Icons/SoundDashboard/edit"),             Icon16));
 			Set("AudioInsights.Icon.SoundDashboard.Filter",           new IMAGE_BRUSH_SVG(TEXT("Icons/SoundDashboard/filter"),           Icon16));
@@ -161,6 +164,23 @@ namespace UE::Audio::Insights
 			Set("AudioInsights.Icon.SoundDashboard.SoundCue",         new IMAGE_BRUSH_SVG(TEXT("Icons/SoundDashboard/soundcue"),         Icon16));
 			Set("AudioInsights.Icon.SoundDashboard.SoundWave",        new IMAGE_BRUSH_SVG(TEXT("Icons/SoundDashboard/soundwave"),        Icon16));
 			Set("AudioInsights.Icon.SoundDashboard.Tab",              new IMAGE_BRUSH_SVG(TEXT("Icons/SoundDashboard/tab"),              Icon16));
+			
+			// Category colors
+			Set("SoundDashboard.MetaSoundColor", FLinearColor(0.008f, 0.76f, 0.078f));
+			Set("SoundDashboard.SoundCueColor", FLinearColor(0.022f, 0.49f, 0.98f));
+			Set("SoundDashboard.ProceduralSourceColor", FLinearColor(0.98f, 0.32f, 0.006f));
+			Set("SoundDashboard.SoundWaveColor", FLinearColor(0.12f, 0.093f, 0.64f));
+			Set("SoundDashboard.SoundCueTemplateColor", FLinearColor(0.98f, 0.01f, 0.01f));
+			Set("SoundDashboard.PinnedColor", FLinearColor(0.9f, 0.9f, 0.9f));
+			Set("SoundDashboard.HiddenColor", FLinearColor(0.4f, 0.4f, 0.4f));
+
+			// Table style
+			Set("SoundDashboard.TableViewRow", FTableRowStyle(FAppStyle::Get().GetWidgetStyle<FTableRowStyle>("TableView.Row"))
+				.SetOddRowBackgroundBrush(FSlateColorBrush(FStyleColors::Recessed))
+				.SetEvenRowBackgroundBrush(FSlateColorBrush(FStyleColors::Background))
+				.SetOddRowBackgroundHoveredBrush(FSlateColorBrush(FStyleColors::SelectHover))
+				.SetEvenRowBackgroundHoveredBrush(FSlateColorBrush(FStyleColors::SelectHover))
+			);
 
 			FSlateStyleRegistry::RegisterSlateStyle(*this);
 		}
diff --git a/Engine/Plugins/AudioInsights/Source/AudioInsights/Public/Views/SoundDashboardViewFactory.h b/Engine/Plugins/AudioInsights/Source/AudioInsights/Public/Views/SoundDashboardViewFactory.h
--- a/Engine/Plugins/AudioInsights/Source/AudioInsights/Public/Views/SoundDashboardViewFactory.h
+++ b/Engine/Plugins/AudioInsights/Source/AudioInsights/Public/Views/SoundDashboardViewFactory.h
@@ -16,6 +16,9 @@ namespace UE::Audio::Insights
 		SoundCue  = 1 << 1,
 		ProceduralSource = 1 << 2,
 		SoundWave = 1 << 3,
+		SoundCueTemplate = 1 << 4,
+		Pinned = 1 << 5,
+		Hidden = 1 << 6
 	};
 
 	ENUM_CLASS_FLAGS(ESoundDashboardFilterFlags);
