[Audio Insights] Audio Buses Dashboard: Remove item coloring logic when the audio bus is registered as it causes confusion + done to be consistent with the Submixes dashboard too (which doesn't have item coloring).

diff --git a/Engine/Plugins/AudioInsights/Source/AudioInsightsEditor/Private/Messages/AudioBusMessages.h b/Engine/Plugins/AudioInsights/Source/AudioInsightsEditor/Private/Messages/AudioBusMessages.h
--- a/Engine/Plugins/AudioInsights/Source/AudioInsightsEditor/Private/Messages/AudioBusMessages.h
+++ b/Engine/Plugins/AudioInsights/Source/AudioInsightsEditor/Private/Messages/AudioBusMessages.h
@@ -85,7 +85,6 @@ namespace UE::Audio::Insights
 
 		TWeakObjectPtr<UAudioBus> AudioBus;
 		uint32 AudioBusId = INDEX_NONE;
-		bool bIsRegistered = false;
 		bool bHasActivity = false;
 	};
 } // namespace UE::Audio::Insights
diff --git a/Engine/Plugins/AudioInsights/Source/AudioInsightsEditor/Private/Providers/AudioBusProvider.cpp b/Engine/Plugins/AudioInsights/Source/AudioInsightsEditor/Private/Providers/AudioBusProvider.cpp
--- a/Engine/Plugins/AudioInsights/Source/AudioInsightsEditor/Private/Providers/AudioBusProvider.cpp
+++ b/Engine/Plugins/AudioInsights/Source/AudioInsightsEditor/Private/Providers/AudioBusProvider.cpp
@@ -227,8 +227,6 @@ namespace UE::Audio::Insights
 			{
 				EntryRef.EntryType = FSoftObjectPath(EntryRef.Name).IsAsset() ? EAudioBusEntryType::AssetBased : EAudioBusEntryType::CodeGenerated;
 			}
-
-			EntryRef.bIsRegistered = true;
 		});
 
 		ProcessMessageQueue<FAudioBusDeactivateMessage>(TraceMessages.DeactivateMessages, GetEntry,
@@ -238,8 +236,6 @@ namespace UE::Audio::Insights
 			{
 				FAudioBusAssetDashboardEntry& EntryRef = *OutEntry->Get();
 
-				EntryRef.bIsRegistered = false;
-
 				if (EntryRef.EntryType == EAudioBusEntryType::CodeGenerated)
 				{
 					RemoveDeviceEntry(Msg.DeviceId, Msg.AudioBusId);
@@ -254,13 +250,10 @@ namespace UE::Audio::Insights
 			{
 				FAudioBusAssetDashboardEntry& EntryRef = *OutEntry->Get();
 
-				// We should only receive activity messages from registered buses
-				const bool bIsRegistered = EntryRef.bIsRegistered || (!EntryRef.Name.IsEmpty() && EntryRef.EntryType != EAudioBusEntryType::None);
-				if (bIsRegistered)
+				if (!EntryRef.Name.IsEmpty() && EntryRef.EntryType != EAudioBusEntryType::None)
 				{
 					EntryRef.bHasActivity = Msg.bHasActivity;
 					EntryRef.Timestamp = Msg.Timestamp;
-					EntryRef.bIsRegistered = true;
 				}
 			}
 		});
diff --git a/Engine/Plugins/AudioInsights/Source/AudioInsightsEditor/Private/Views/AudioBusesDashboardViewFactory.cpp b/Engine/Plugins/AudioInsights/Source/AudioInsightsEditor/Private/Views/AudioBusesDashboardViewFactory.cpp
--- a/Engine/Plugins/AudioInsights/Source/AudioInsightsEditor/Private/Views/AudioBusesDashboardViewFactory.cpp
+++ b/Engine/Plugins/AudioInsights/Source/AudioInsightsEditor/Private/Views/AudioBusesDashboardViewFactory.cpp
@@ -144,15 +144,6 @@ namespace UE::Audio::Insights
 						SNew(STextBlock)
 						.Text(ValueText)
 						.MinDesiredWidth(300)
-						.ColorAndOpacity_Lambda([InRowData]()
-						{
-							static const FSlateColor DefaultColor(FColor::White);
-							static const FSlateColor AudioBusAssetColor(FColor(80, 200, 255));
-							
-							const FAudioBusAssetDashboardEntry& AudioBusAssetDashboardEntry = AudioBusesPrivate::CastEntry(InRowData.Get());
-
-							return AudioBusAssetDashboardEntry.bIsRegistered ? AudioBusAssetColor : DefaultColor;
-						})
 						.OnDoubleClicked_Lambda([this, InRowData](const FGeometry& MyGeometry, const FPointerEvent& PointerEvent)
 						{
 							if (GEditor)
