[Audio Insights] Removed old FTraceChannelManager and FTraceChannelHandle classes since they are not used to toggle trace channels anymore (this is already handled by FTraceModule)

diff --git a/Engine/Plugins/AudioInsights/Source/AudioInsights/Private/AudioInsightsTraceChannelHandle.cpp b/Engine/Plugins/AudioInsights/Source/AudioInsights/Private/AudioInsightsTraceChannelHandle.cpp
deleted file mode 100644
diff --git a/Engine/Plugins/AudioInsights/Source/AudioInsights/Private/AudioInsightsTraceModule.cpp b/Engine/Plugins/AudioInsights/Source/AudioInsights/Private/AudioInsightsTraceModule.cpp
--- a/Engine/Plugins/AudioInsights/Source/AudioInsights/Private/AudioInsightsTraceModule.cpp
+++ b/Engine/Plugins/AudioInsights/Source/AudioInsights/Private/AudioInsightsTraceModule.cpp
@@ -25,21 +25,12 @@ namespace UE::Audio::Insights
 	}
 
 	FTraceModule::FTraceModule()
-		: ChannelManager(MakeShared<FTraceChannelManager>())
 #if !WITH_EDITOR
-		, AudioChannels({ FTraceModulePrivate::AudioChannelName, FTraceModulePrivate::AudioMixerChannelName } )
+		: AudioChannels({ FTraceModulePrivate::AudioChannelName, FTraceModulePrivate::AudioMixerChannelName } )
 		, EmptyArray()
 #endif
 	{
-		// Don't run providers in any commandlet to avoid additional, unnecessary overhead as audio insights is dormant.
-		if (!IsRunningCommandlet())
-		{
-			TSharedPtr<FMixerSourceTraceProvider> SourceProvider = MakeShared<FMixerSourceTraceProvider>(ChannelManager);
-			TSharedPtr<FVirtualLoopTraceProvider> VirtualLoopProvider = MakeShared<FVirtualLoopTraceProvider>(ChannelManager);
 
-			TraceProviders.Add(SourceProvider->GetName(), SourceProvider);
-			TraceProviders.Add(VirtualLoopProvider->GetName(), VirtualLoopProvider);
-		}
 	}
 
 	void FTraceModule::GetModuleInfo(TraceServices::FModuleInfo& OutModuleInfo)
@@ -53,11 +44,6 @@ namespace UE::Audio::Insights
 		TraceProviders.Add(TraceProvider->GetName(), TraceProvider);
 	}
 
-	TSharedRef<FTraceChannelManager> FTraceModule::GetChannelManager()
-	{
-		return ChannelManager;
-	}
-
 	const FName FTraceModule::GetName()
 	{
 		const FLazyName TraceName = { "TraceModule_AudioTrace" };
diff --git a/Engine/Plugins/AudioInsights/Source/AudioInsights/Private/Providers/MixerSourceTraceProvider.h b/Engine/Plugins/AudioInsights/Source/AudioInsights/Private/Providers/MixerSourceTraceProvider.h
--- a/Engine/Plugins/AudioInsights/Source/AudioInsights/Private/Providers/MixerSourceTraceProvider.h
+++ b/Engine/Plugins/AudioInsights/Source/AudioInsights/Private/Providers/MixerSourceTraceProvider.h
@@ -1,7 +1,6 @@
 // Copyright Epic Games, Inc. All Rights Reserved.
 #pragma once
 
-#include "AudioInsightsTraceChannelHandle.h"
 #include "AudioInsightsTraceProviderBase.h"
 #include "Messages/MixerSourceTraceMessages.h"
 #include "UObject/NameTypes.h"
@@ -14,10 +13,10 @@ namespace UE::Audio::Insights
 		, public TSharedFromThis<FMixerSourceTraceProvider>
 	{
 	public:
-		FMixerSourceTraceProvider(TSharedRef<FTraceChannelManager> InManager)
+		FMixerSourceTraceProvider()
 			: TDeviceDataMapTraceProvider<uint32, TSharedPtr<FMixerSourceDashboardEntry>>(GetName_Static())
 		{
-			Channels.Add(InManager->CreateHandle({ TEXT("AudioMixerChannel") }));
+			
 		}
 
 		virtual ~FMixerSourceTraceProvider() = default;
@@ -39,7 +38,5 @@ namespace UE::Audio::Insights
 #endif // !WITH_EDITOR
 
 		FMixerSourceMessages TraceMessages;
-
-		TSet<FTraceChannelHandle> Channels;
 	};
 } // namespace UE::Audio::Insights
diff --git a/Engine/Plugins/AudioInsights/Source/AudioInsights/Private/Providers/SoundTraceProvider.cpp b/Engine/Plugins/AudioInsights/Source/AudioInsights/Private/Providers/SoundTraceProvider.cpp
--- a/Engine/Plugins/AudioInsights/Source/AudioInsights/Private/Providers/SoundTraceProvider.cpp
+++ b/Engine/Plugins/AudioInsights/Source/AudioInsights/Private/Providers/SoundTraceProvider.cpp
@@ -34,10 +34,10 @@ namespace UE::Audio::Insights
 #endif // !WITH_EDITOR
 	}
 
-	FSoundTraceProvider::FSoundTraceProvider(TSharedRef<FTraceChannelManager> InManager)
+	FSoundTraceProvider::FSoundTraceProvider()
 		: TDeviceDataMapTraceProvider<uint32, TSharedPtr<FSoundDashboardEntry>>(GetName_Static())
 	{
-		Channels.Add(InManager->CreateHandle({ TEXT("AudioMixerChannel") }));
+		
 	}
 
 	FName FSoundTraceProvider::GetName_Static()
diff --git a/Engine/Plugins/AudioInsights/Source/AudioInsights/Private/Providers/SoundTraceProvider.h b/Engine/Plugins/AudioInsights/Source/AudioInsights/Private/Providers/SoundTraceProvider.h
--- a/Engine/Plugins/AudioInsights/Source/AudioInsights/Private/Providers/SoundTraceProvider.h
+++ b/Engine/Plugins/AudioInsights/Source/AudioInsights/Private/Providers/SoundTraceProvider.h
@@ -1,7 +1,6 @@
 // Copyright Epic Games, Inc. All Rights Reserved.
 #pragma once
 
-#include "AudioInsightsTraceChannelHandle.h"
 #include "AudioInsightsTraceProviderBase.h"
 #include "Messages/SoundTraceMessages.h"
 #include "UObject/NameTypes.h"
@@ -13,7 +12,7 @@ namespace UE::Audio::Insights
 		, public TSharedFromThis<FSoundTraceProvider>
 	{
 	public:
-		FSoundTraceProvider(TSharedRef<FTraceChannelManager> InManager);
+		FSoundTraceProvider();
 		virtual ~FSoundTraceProvider() = default;
 
 		static FName GetName_Static();
@@ -34,7 +33,5 @@ namespace UE::Audio::Insights
 #endif // !WITH_EDITOR
 
 		FSoundMessages TraceMessages;
-
-		TSet<FTraceChannelHandle> Channels;
 	};
 } // namespace UE::Audio::Insights
diff --git a/Engine/Plugins/AudioInsights/Source/AudioInsights/Private/Providers/VirtualLoopTraceProvider.h b/Engine/Plugins/AudioInsights/Source/AudioInsights/Private/Providers/VirtualLoopTraceProvider.h
--- a/Engine/Plugins/AudioInsights/Source/AudioInsights/Private/Providers/VirtualLoopTraceProvider.h
+++ b/Engine/Plugins/AudioInsights/Source/AudioInsights/Private/Providers/VirtualLoopTraceProvider.h
@@ -1,7 +1,6 @@
 // Copyright Epic Games, Inc. All Rights Reserved.
 #pragma once
 
-#include "AudioInsightsTraceChannelHandle.h"
 #include "AudioInsightsTraceProviderBase.h"
 #include "Messages/VirtualLoopTraceMessages.h"
 
@@ -13,10 +12,10 @@ namespace UE::Audio::Insights
 		, public TSharedFromThis<FVirtualLoopTraceProvider>
 	{
 	public:
-		FVirtualLoopTraceProvider(TSharedRef<FTraceChannelManager> InChannelManager)
+		FVirtualLoopTraceProvider()
 			: TDeviceDataMapTraceProvider<uint32, TSharedPtr<FVirtualLoopDashboardEntry>>(GetName_Static())
 		{
-			Channels.Add(InChannelManager->CreateHandle(TEXT("AudioChannel")));
+			
 		}
 
 		virtual ~FVirtualLoopTraceProvider() = default;
@@ -38,6 +37,5 @@ namespace UE::Audio::Insights
 		virtual bool ProcessMessages() override;
 
 		FVirtualLoopMessages TraceMessages;
-		TSet<FTraceChannelHandle> Channels;
 	};
 } // namespace UE::Audio::Insights
diff --git a/Engine/Plugins/AudioInsights/Source/AudioInsights/Private/Views/MixerSourceDashboardViewFactory.cpp b/Engine/Plugins/AudioInsights/Source/AudioInsights/Private/Views/MixerSourceDashboardViewFactory.cpp
--- a/Engine/Plugins/AudioInsights/Source/AudioInsights/Private/Views/MixerSourceDashboardViewFactory.cpp
+++ b/Engine/Plugins/AudioInsights/Source/AudioInsights/Private/Views/MixerSourceDashboardViewFactory.cpp
@@ -66,13 +66,18 @@ namespace UE::Audio::Insights
 
 	FMixerSourceDashboardViewFactory::FMixerSourceDashboardViewFactory()
 	{
-		FTraceModule& TraceModule = static_cast<FTraceModule&>(FAudioInsightsModule::GetChecked().GetTraceModule());
+		FTraceModule& AudioInsightsTraceModule = static_cast<FTraceModule&>(FAudioInsightsModule::GetChecked().GetTraceModule());
+
+		const TSharedPtr<FMixerSourceTraceProvider> MixerSourceTraceProvider = MakeShared<FMixerSourceTraceProvider>();
+
+		AudioInsightsTraceModule.AddTraceProvider(MixerSourceTraceProvider);
+
 		Providers = TArray<TSharedPtr<FTraceProviderBase>>
 		{
-			TraceModule.FindAudioTraceProvider<FMixerSourceTraceProvider>()
+			MixerSourceTraceProvider
 		};
 
-		TraceModule.OnAnalysisStarting.AddRaw(this, &FMixerSourceDashboardViewFactory::OnAnalysisStarting);
+		AudioInsightsTraceModule.OnAnalysisStarting.AddRaw(this, &FMixerSourceDashboardViewFactory::OnAnalysisStarting);
 	}
 
 	FMixerSourceDashboardViewFactory::~FMixerSourceDashboardViewFactory()
diff --git a/Engine/Plugins/AudioInsights/Source/AudioInsights/Private/Views/SoundDashboardViewFactory.cpp b/Engine/Plugins/AudioInsights/Source/AudioInsights/Private/Views/SoundDashboardViewFactory.cpp
--- a/Engine/Plugins/AudioInsights/Source/AudioInsights/Private/Views/SoundDashboardViewFactory.cpp
+++ b/Engine/Plugins/AudioInsights/Source/AudioInsights/Private/Views/SoundDashboardViewFactory.cpp
@@ -43,7 +43,7 @@ namespace UE::Audio::Insights
 	{
 		FTraceModule& AudioInsightsTraceModule = static_cast<FTraceModule&>(FAudioInsightsModule::GetChecked().GetTraceModule());
 
-		const TSharedPtr<FSoundTraceProvider> SoundsTraceProvider = MakeShared<FSoundTraceProvider>(AudioInsightsTraceModule.GetChannelManager());
+		const TSharedPtr<FSoundTraceProvider> SoundsTraceProvider = MakeShared<FSoundTraceProvider>();
 
 		AudioInsightsTraceModule.AddTraceProvider(SoundsTraceProvider);
 
diff --git a/Engine/Plugins/AudioInsights/Source/AudioInsights/Private/Views/VirtualLoopDashboardViewFactory.cpp b/Engine/Plugins/AudioInsights/Source/AudioInsights/Private/Views/VirtualLoopDashboardViewFactory.cpp
--- a/Engine/Plugins/AudioInsights/Source/AudioInsights/Private/Views/VirtualLoopDashboardViewFactory.cpp
+++ b/Engine/Plugins/AudioInsights/Source/AudioInsights/Private/Views/VirtualLoopDashboardViewFactory.cpp
@@ -28,10 +28,15 @@ namespace UE::Audio::Insights
 
 	FVirtualLoopDashboardViewFactory::FVirtualLoopDashboardViewFactory()
 	{
-		const FTraceModule& TraceModule = static_cast<FTraceModule&>(FAudioInsightsModule::GetChecked().GetTraceModule());
+		FTraceModule& AudioInsightsTraceModule = static_cast<FTraceModule&>(FAudioInsightsModule::GetChecked().GetTraceModule());
+
+		const TSharedPtr<FVirtualLoopTraceProvider> VirtualLoopTraceProvider = MakeShared<FVirtualLoopTraceProvider>();
+
+		AudioInsightsTraceModule.AddTraceProvider(VirtualLoopTraceProvider);
+
 		Providers = TArray<TSharedPtr<FTraceProviderBase>>
 		{
-			TraceModule.FindAudioTraceProvider<FVirtualLoopTraceProvider>()
+			VirtualLoopTraceProvider
 		};
 	}
 
diff --git a/Engine/Plugins/AudioInsights/Source/AudioInsights/Public/AudioInsightsTraceChannelHandle.h b/Engine/Plugins/AudioInsights/Source/AudioInsights/Public/AudioInsightsTraceChannelHandle.h
deleted file mode 100644
diff --git a/Engine/Plugins/AudioInsights/Source/AudioInsights/Public/AudioInsightsTraceModule.h b/Engine/Plugins/AudioInsights/Source/AudioInsights/Public/AudioInsightsTraceModule.h
--- a/Engine/Plugins/AudioInsights/Source/AudioInsights/Public/AudioInsightsTraceModule.h
+++ b/Engine/Plugins/AudioInsights/Source/AudioInsights/Public/AudioInsightsTraceModule.h
@@ -33,8 +33,6 @@ namespace UE::Audio::Insights
 
 		virtual void AddTraceProvider(TSharedPtr<FTraceProviderBase> TraceProvider) override;
 
-		virtual TSharedRef<FTraceChannelManager> GetChannelManager() override;
-
 		virtual void StartTraceAnalysis(const bool bOnlyTraceAudioChannels) override;
 		virtual bool IsTraceAnalysisActive() const override;
 		virtual void StopTraceAnalysis() override;
@@ -66,7 +64,6 @@ namespace UE::Audio::Insights
 		bool GetAudioTracesAreEnabled() const;
 #endif
 
-		TSharedRef<FTraceChannelManager> ChannelManager;
 		TMap<FName, TSharedPtr<FTraceProviderBase>> TraceProviders;
 		TArray<FString> ChannelsToRestore;
 
diff --git a/Engine/Plugins/AudioInsights/Source/AudioInsights/Public/IAudioInsightsTraceModule.h b/Engine/Plugins/AudioInsights/Source/AudioInsights/Public/IAudioInsightsTraceModule.h
--- a/Engine/Plugins/AudioInsights/Source/AudioInsights/Public/IAudioInsightsTraceModule.h
+++ b/Engine/Plugins/AudioInsights/Source/AudioInsights/Public/IAudioInsightsTraceModule.h
@@ -1,7 +1,6 @@
 // Copyright Epic Games, Inc. All Rights Reserved.
 #pragma once
 
-#include "AudioInsightsTraceChannelHandle.h"
 #include "AudioInsightsTraceProviderBase.h"
 #include "TraceServices/ModuleService.h"
 
@@ -28,8 +27,6 @@ public:
 
 	virtual bool AudioChannelsCanBeManuallyEnabled() const = 0;
 
-	virtual TSharedRef<UE::Audio::Insights::FTraceChannelManager> GetChannelManager() = 0;
-
 #if !WITH_EDITOR
 	virtual void InitializeSessionInfo(const TraceServices::FSessionInfo& SessionInfo) = 0;
 	virtual void RequestChannelUpdate() = 0;
diff --git a/Engine/Plugins/Runtime/AudioModulation/Source/AudioModulationEditor/Private/Insights/Providers/ControlBusTraceProvider.h b/Engine/Plugins/Runtime/AudioModulation/Source/AudioModulationEditor/Private/Insights/Providers/ControlBusTraceProvider.h
--- a/Engine/Plugins/Runtime/AudioModulation/Source/AudioModulationEditor/Private/Insights/Providers/ControlBusTraceProvider.h
+++ b/Engine/Plugins/Runtime/AudioModulation/Source/AudioModulationEditor/Private/Insights/Providers/ControlBusTraceProvider.h
@@ -1,7 +1,6 @@
 // Copyright Epic Games, Inc. All Rights Reserved.
 #pragma once
 
-#include "AudioInsightsTraceChannelHandle.h"
 #include "AudioInsightsTraceProviderBase.h"
 #include "Insights/Messages/ControlBusTraceMessages.h"
 
@@ -13,10 +12,10 @@ namespace AudioModulationEditor
 		, public TSharedFromThis<FControlBusTraceProvider>
 	{
 	public:
-		FControlBusTraceProvider(TSharedRef<UE::Audio::Insights::FTraceChannelManager> InManager)
+		FControlBusTraceProvider()
 			: UE::Audio::Insights::TDeviceDataMapTraceProvider<uint32, TSharedPtr<FControlBusDashboardEntry>>(GetName_Static())
 		{
-			Channels.Add(InManager->CreateHandle({ TEXT("AudioChannel") }));
+
 		}
 
 		virtual ~FControlBusTraceProvider() = default;
@@ -28,6 +27,5 @@ namespace AudioModulationEditor
 	private:
 
 		FControlBusMessages TraceMessages;
-		TSet<UE::Audio::Insights::FTraceChannelHandle> Channels;
 	};
 } // namespace AudioModulationEditor
diff --git a/Engine/Plugins/Runtime/AudioModulation/Source/AudioModulationEditor/Private/Insights/Providers/ModulationMatrixTraceProvider.cpp b/Engine/Plugins/Runtime/AudioModulation/Source/AudioModulationEditor/Private/Insights/Providers/ModulationMatrixTraceProvider.cpp
--- a/Engine/Plugins/Runtime/AudioModulation/Source/AudioModulationEditor/Private/Insights/Providers/ModulationMatrixTraceProvider.cpp
+++ b/Engine/Plugins/Runtime/AudioModulation/Source/AudioModulationEditor/Private/Insights/Providers/ModulationMatrixTraceProvider.cpp
@@ -7,11 +7,9 @@
 
 namespace AudioModulationEditor
 {
-	FModulationMatrixTraceProvider::FModulationMatrixTraceProvider(TSharedRef<UE::Audio::Insights::FTraceChannelManager> InManager)
+	FModulationMatrixTraceProvider::FModulationMatrixTraceProvider()
 		: UE::Audio::Insights::TDeviceDataMapTraceProvider<uint32, TSharedPtr<FModulationMatrixDashboardEntry>>(GetName_Static())
 	{
-		Channels.Add(InManager->CreateHandle({ TEXT("AudioChannel") }));
-
 		FAudioDeviceManagerDelegates::OnAudioDeviceDestroyed.AddRaw(this, &FModulationMatrixTraceProvider::OnAudioDeviceDestroyed);
 	}
 
diff --git a/Engine/Plugins/Runtime/AudioModulation/Source/AudioModulationEditor/Private/Insights/Providers/ModulationMatrixTraceProvider.h b/Engine/Plugins/Runtime/AudioModulation/Source/AudioModulationEditor/Private/Insights/Providers/ModulationMatrixTraceProvider.h
--- a/Engine/Plugins/Runtime/AudioModulation/Source/AudioModulationEditor/Private/Insights/Providers/ModulationMatrixTraceProvider.h
+++ b/Engine/Plugins/Runtime/AudioModulation/Source/AudioModulationEditor/Private/Insights/Providers/ModulationMatrixTraceProvider.h
@@ -1,7 +1,6 @@
 // Copyright Epic Games, Inc. All Rights Reserved.
 #pragma once
 
-#include "AudioInsightsTraceChannelHandle.h"
 #include "AudioInsightsTraceProviderBase.h"
 #include "Containers/Map.h"
 #include "Delegates/DelegateCombinations.h"
@@ -14,7 +13,7 @@ namespace AudioModulationEditor
 		, public TSharedFromThis<FModulationMatrixTraceProvider>
 	{
 	public:
-		FModulationMatrixTraceProvider(TSharedRef<UE::Audio::Insights::FTraceChannelManager> InManager);
+		FModulationMatrixTraceProvider();
 
 		virtual ~FModulationMatrixTraceProvider();
 
@@ -43,7 +42,6 @@ namespace AudioModulationEditor
 		void OnAudioDeviceDestroyed(::Audio::FDeviceId InDeviceId);
 
 		FModulationMatrixMessages TraceMessages;
-		TSet<UE::Audio::Insights::FTraceChannelHandle> Channels;
 
 		TMap<::Audio::FDeviceId, TSet<uint32>> DeviceIdToActiveModulatorSourceIdsMap;
 		BusIdToBusInfoMap ActiveControlBuses;
diff --git a/Engine/Plugins/Runtime/AudioModulation/Source/AudioModulationEditor/Private/Insights/Views/ControlBusDashboardViewFactory.cpp b/Engine/Plugins/Runtime/AudioModulation/Source/AudioModulationEditor/Private/Insights/Views/ControlBusDashboardViewFactory.cpp
--- a/Engine/Plugins/Runtime/AudioModulation/Source/AudioModulationEditor/Private/Insights/Views/ControlBusDashboardViewFactory.cpp
+++ b/Engine/Plugins/Runtime/AudioModulation/Source/AudioModulationEditor/Private/Insights/Views/ControlBusDashboardViewFactory.cpp
@@ -33,7 +33,7 @@ namespace AudioModulationEditor
 		IAudioInsightsEditorModule& InsightsModule = IAudioInsightsEditorModule::GetChecked();
 		IAudioInsightsTraceModule& InsightsTraceModule = InsightsModule.GetTraceModule();
 
-		TSharedPtr<FControlBusTraceProvider> ControlBusProvider = MakeShared<FControlBusTraceProvider>(InsightsTraceModule.GetChannelManager());
+		TSharedPtr<FControlBusTraceProvider> ControlBusProvider = MakeShared<FControlBusTraceProvider>();
 		InsightsTraceModule.AddTraceProvider(StaticCastSharedPtr<UE::Audio::Insights::FTraceProviderBase>(ControlBusProvider));
 
 		Providers =
diff --git a/Engine/Plugins/Runtime/AudioModulation/Source/AudioModulationEditor/Private/Insights/Views/ModulationMatrixDashboardViewFactory.cpp b/Engine/Plugins/Runtime/AudioModulation/Source/AudioModulationEditor/Private/Insights/Views/ModulationMatrixDashboardViewFactory.cpp
--- a/Engine/Plugins/Runtime/AudioModulation/Source/AudioModulationEditor/Private/Insights/Views/ModulationMatrixDashboardViewFactory.cpp
+++ b/Engine/Plugins/Runtime/AudioModulation/Source/AudioModulationEditor/Private/Insights/Views/ModulationMatrixDashboardViewFactory.cpp
@@ -32,7 +32,7 @@ namespace AudioModulationEditor
 	{
 		IAudioInsightsTraceModule& AudioInsightsTraceModule = IAudioInsightsEditorModule::GetChecked().GetTraceModule();
 
-		ModulationMatrixTraceProvider = MakeShared<FModulationMatrixTraceProvider>(AudioInsightsTraceModule.GetChannelManager());
+		ModulationMatrixTraceProvider = MakeShared<FModulationMatrixTraceProvider>();
 
 		AudioInsightsTraceModule.AddTraceProvider(ModulationMatrixTraceProvider);
 
