[Audio Insights] Home Screen - Children Sounds
- Implement tree structure of sounds in the Audio Insights Sounds tab in both editor and standalone
- Only show individual SoundWave instances below Active Sound entries if the type of entry can have multiple SoundWaves as children.
- Using MixerSource messages where possible to reduce the amount of duplicated data sent.
- Pitch updated to use FinalPitch rather than TargetPitch so we capture any affect from modulators
- Don't display any data points for sound dashboard rows that have children - their data is not set.
- Auto-expand new sound entries when they first enter the sounds dashboard
- If a sound is set to always play, show it's priority as MAX

diff --git a/Engine/Plugins/AudioInsights/Source/AudioInsights/Private/Providers/SoundTraceProvider.cpp b/Engine/Plugins/AudioInsights/Source/AudioInsights/Private/Providers/SoundTraceProvider.cpp
--- a/Engine/Plugins/AudioInsights/Source/AudioInsights/Private/Providers/SoundTraceProvider.cpp
+++ b/Engine/Plugins/AudioInsights/Source/AudioInsights/Private/Providers/SoundTraceProvider.cpp
@@ -2,6 +2,7 @@
 #include "Providers/SoundTraceProvider.h"
 
 #include "AudioInsightsModule.h"
+#include "Messages/MixerSourceTraceMessages.h"
 #include "Trace/Analyzer.h"
 #include "TraceServices/Model/AnalysisSession.h"
 
@@ -10,10 +11,161 @@
 #include "Async/ParallelFor.h"
 #endif // !WITH_EDITOR
 
+#define LOCTEXT_NAMESPACE "AudioInsights"
+
 namespace UE::Audio::Insights
 {
 	namespace FSoundTraceProviderPrivate
 	{
+		const FString MetaSoundDisplayName			{ LOCTEXT("AudioDashboard_Sounds_EntryType_MetaSound", "MetaSounds").ToString() };
+		const FString SoundCueDisplayName			{ LOCTEXT("AudioDashboard_Sounds_EntryType_SoundCues", "Sound Cues").ToString() };
+		const FString SoundWaveDisplayName			{ LOCTEXT("AudioDashboard_Sounds_EntryType_SoundWaves", "Sound Waves").ToString() };
+		const FString ProceduralSourceDisplayName	{ LOCTEXT("AudioDashboard_Sounds_EntryType_ProceduralSources", "Procedural Sources").ToString() };
+		const FString SoundCueTemplateName			{ LOCTEXT("AudioDashboard_Sounds_EntryType_SoundCueTemplate", "Sound Cue Templates").ToString() };
+		const FString UncategorizedSoundName		{ LOCTEXT("AudioDashboard_Sounds_EntryType_SoundCueTemplate", "Others").ToString() };
+
+		FSoundDashboardEntry& CastEntry(IDashboardDataTreeViewEntry& InData)
+		{
+			return static_cast<FSoundDashboardEntry&>(InData);
+		};
+
+		FString GetEntryTypeDisplayName(const ESoundDashboardEntryType EntryType)
+		{
+			switch (EntryType)
+			{
+				case ESoundDashboardEntryType::MetaSound:
+					return MetaSoundDisplayName;
+				case ESoundDashboardEntryType::SoundCue:
+					return SoundCueDisplayName;
+				case ESoundDashboardEntryType::SoundWave:
+					return SoundWaveDisplayName;
+				case ESoundDashboardEntryType::ProceduralSource:
+					return ProceduralSourceDisplayName;
+				case ESoundDashboardEntryType::SoundCueTemplate:
+					return SoundCueTemplateName;
+			}
+
+			return UncategorizedSoundName;
+		}
+
+		bool EntryTypeHasSoundWaveEntries(const ESoundDashboardEntryType EntryType)
+		{
+			switch (EntryType)
+			{
+				case ESoundDashboardEntryType::SoundCue:
+				case ESoundDashboardEntryType::SoundCueTemplate:
+					return true;
+
+				case ESoundDashboardEntryType::MetaSound:
+				case ESoundDashboardEntryType::ProceduralSource:
+				case ESoundDashboardEntryType::SoundWave:
+					return false;
+				}
+
+			return true;
+		}
+
+		void UpdateParameterEntry(const FMixerSourceParameterMessage& Msg, const float DataPoint, ::Audio::TCircularAudioBuffer<FDataPoint>& OutDataPoints, TArray<int32, TInlineAllocator<64>>& OutEntriesWithPoppedDataPoints)
+		{
+			if (!OutEntriesWithPoppedDataPoints.Contains(Msg.PlayOrder))
+			{
+				OutDataPoints.Pop(OutDataPoints.Num());
+				OutEntriesWithPoppedDataPoints.Add(Msg.PlayOrder);
+			}
+
+			OutDataPoints.Push({ Msg.Timestamp, DataPoint });
+		}
+
+		//////////////////////////////////////////////////////////////////////////////////////////////////////////////
+		// Finds an Asset entry inside a Category entry, which will hold a collection of Active Sound instances that are all the same Asset
+		void GetAssetEntry(const TSharedPtr<FSoundDashboardEntry>* CategoryEntry, const FString& AssetName, TSharedPtr<FSoundDashboardEntry>& OutAssetEntry)
+		{
+			OutAssetEntry = nullptr;
+			if (CategoryEntry == nullptr || !CategoryEntry->IsValid())
+			{
+				return;
+			}
+
+			TSharedPtr<IDashboardDataTreeViewEntry>* AssetEntry = CategoryEntry->Get()->Children.FindByPredicate([&AssetName](TSharedPtr<IDashboardDataTreeViewEntry> ChildEntry)
+			{
+				FSoundDashboardEntry& NewEntry = CastEntry(*ChildEntry.Get());
+				return NewEntry.Name == AssetName;
+			});
+
+			if (AssetEntry == nullptr || !AssetEntry->IsValid())
+			{
+				return;
+			}
+
+			OutAssetEntry = StaticCastSharedPtr<FSoundDashboardEntry>(*AssetEntry);
+		}
+
+		//////////////////////////////////////////////////////////////////////////////////////////////////////////////
+		// Finds an Active Sound entry inside a Category entry
+		void FindActiveSoundEntryInCategory(const TSharedPtr<FSoundDashboardEntry>* CategoryEntry, const FString& AssetName, const uint32 ActiveSoundPlayOrder, TSharedPtr<FSoundDashboardEntry>& OutActiveSoundEntry)
+		{
+			OutActiveSoundEntry = nullptr;
+			if (CategoryEntry == nullptr || !CategoryEntry->IsValid())
+			{
+				return;
+			}
+
+			TSharedPtr<IDashboardDataTreeViewEntry>* AssetEntry = CategoryEntry->Get()->Children.FindByPredicate([&AssetName](TSharedPtr<IDashboardDataTreeViewEntry> ChildEntry)
+			{
+				FSoundDashboardEntry& NewEntry = CastEntry(*ChildEntry.Get());
+				return NewEntry.Name == AssetName;
+			});
+
+			if (AssetEntry == nullptr || !AssetEntry->IsValid())
+			{
+				return;
+			}
+
+			TSharedPtr<IDashboardDataTreeViewEntry>* ActiveSoundEntry = AssetEntry->Get()->Children.FindByPredicate([ActiveSoundPlayOrder](TSharedPtr<IDashboardDataTreeViewEntry> ChildEntry)
+			{
+				FSoundDashboardEntry& NewEntry = CastEntry(*ChildEntry.Get());
+				return NewEntry.PlayOrder == ActiveSoundPlayOrder;
+			});
+
+			if (ActiveSoundEntry == nullptr || !ActiveSoundEntry->IsValid())
+			{
+				return;
+			}
+
+			OutActiveSoundEntry = StaticCastSharedPtr<FSoundDashboardEntry>(*ActiveSoundEntry);
+		}
+
+		//////////////////////////////////////////////////////////////////////////////////////////////////////////////
+		// Will return either the Active Sound entry or the child SoundWave entry depending on the sound's EntryType
+		FSoundDashboardEntry* FindSoundEntryForActiveSound(const uint32 SoundWavePlayOrder, const TSharedPtr<FSoundDashboardEntry>* ActiveSoundEntry)
+		{
+			if (ActiveSoundEntry == nullptr || !ActiveSoundEntry->IsValid())
+			{
+				return nullptr;
+			}
+
+			if (!EntryTypeHasSoundWaveEntries((*ActiveSoundEntry)->EntryType))
+			{
+				return ActiveSoundEntry->Get();
+			}
+
+			for (TSharedPtr<IDashboardDataTreeViewEntry> SoundWaveEntry : (*ActiveSoundEntry)->Children)
+			{
+				if (SoundWaveEntry == nullptr || !SoundWaveEntry.IsValid())
+				{
+					continue;
+				}
+
+				FSoundDashboardEntry* SoundWaveEntryCast = static_cast<FSoundDashboardEntry*>(SoundWaveEntry.Get());
+				if (SoundWaveEntryCast->PlayOrder == SoundWavePlayOrder)
+				{
+					return SoundWaveEntryCast;
+				}
+			}
+
+			return nullptr;
+		}
+
 #if !WITH_EDITOR
 		template<typename T>
 		const T* FindClosestMessageToTimestamp(const TraceServices::TPagedArray<T>& InCachedMessages, const double InTimeMarker, const uint32 InPlayOrder)
@@ -23,7 +175,7 @@ namespace UE::Audio::Insights
 			// Iterate backwards from TimeMarker until we find the matching PlayOrder
 			for (auto It = InCachedMessages.GetIteratorFromItem(ClosestMessageToTimeStampIndex); It; --It)
 			{
-				if (It->PlayOrder == InPlayOrder)
+				if (It->ActiveSoundPlayOrder == InPlayOrder)
 				{
 					return &(*It);
 				}
@@ -35,9 +187,9 @@ namespace UE::Audio::Insights
 	}
 
 	FSoundTraceProvider::FSoundTraceProvider()
-		: TDeviceDataMapTraceProvider<uint32, TSharedPtr<FSoundDashboardEntry>>(GetName_Static())
+		: TDeviceDataMapTraceProvider<ESoundDashboardEntryType, TSharedPtr<FSoundDashboardEntry>>(GetName_Static())
 	{
-		
+		InitLambdas();
 	}
 
 	FName FSoundTraceProvider::GetName_Static()
@@ -63,13 +215,14 @@ namespace UE::Audio::Insights
 
 				UE::Trace::IAnalyzer::FInterfaceBuilder& Builder = Context.InterfaceBuilder;
 
-				Builder.RouteEvent(RouteId_Start,     "Audio", "SoundStart");
-				Builder.RouteEvent(RouteId_Stop,      "Audio", "SoundStop");
-				Builder.RouteEvent(RouteId_Priority,  "Audio", "SoundPriority");
-				Builder.RouteEvent(RouteId_Distance,  "Audio", "SoundDistance");
-				Builder.RouteEvent(RouteId_Amplitude, "Audio", "SoundAmplitude");
-				Builder.RouteEvent(RouteId_Volume,    "Audio", "SoundVolume");
-				Builder.RouteEvent(RouteId_Pitch,     "Audio", "SoundPitch");
+				Builder.RouteEvent(RouteId_ActiveSoundStart,	"Audio", "SoundStart");
+				Builder.RouteEvent(RouteId_SoundWaveStart,		"Audio", "SoundWaveStart");
+				Builder.RouteEvent(RouteId_Stop,				"Audio", "SoundStop");
+				Builder.RouteEvent(RouteId_Priority,			"Audio", "SoundPriority");
+				Builder.RouteEvent(RouteId_Distance,			"Audio", "SoundDistance");
+				Builder.RouteEvent(RouteId_Amplitude,			"Audio", "MixerSourceEnvelope");
+				Builder.RouteEvent(RouteId_Volume,				"Audio", "MixerSourceVolume");
+				Builder.RouteEvent(RouteId_Pitch,				"Audio", "MixerSourcePitch");
 			}
 
 			virtual bool OnEvent(uint16 RouteId, EStyle Style, const FOnEventContext& Context) override
@@ -80,9 +233,15 @@ namespace UE::Audio::Insights
 
 				switch (RouteId)
 				{
-					case RouteId_Start:
+					case RouteId_ActiveSoundStart:
+					{
+						Messages.ActiveSoundStartMessages.Enqueue(FSoundStartMessage{ Context });
+						break;
+					}
+
+					case RouteId_SoundWaveStart:
 					{
-						Messages.StartMessages.Enqueue(FSoundStartMessage{ Context });
+						Messages.SoundWaveStartMessages.Enqueue(FSoundWaveStartMessage(Context));
 						break;
 					}
 
@@ -106,19 +265,19 @@ namespace UE::Audio::Insights
 
 					case RouteId_Amplitude:
 					{
-						Messages.AmplitudeMessages.Enqueue(FSoundAmplitudeMessage{ Context });
+						Messages.AmplitudeMessages.Enqueue(FMixerSourceEnvelopeMessage{ Context });
 						break;
 					}
 
 					case RouteId_Volume:
 					{
-						Messages.VolumeMessages.Enqueue(FSoundVolumeMessage{ Context });
+						Messages.VolumeMessages.Enqueue(FMixerSourceVolumeMessage{ Context });
 						break;
 					}
 
 					case RouteId_Pitch:
 					{
-						Messages.PitchMessages.Enqueue(FSoundPitchMessage{ Context });
+						Messages.PitchMessages.Enqueue(FMixerSourcePitchMessage{ Context });
 						break;
 					}
 
@@ -141,7 +300,8 @@ namespace UE::Audio::Insights
 		private:
 			enum : uint16
 			{
-				RouteId_Start,
+				RouteId_ActiveSoundStart,
+				RouteId_SoundWaveStart,
 				RouteId_Stop,
 				RouteId_Priority,
 				RouteId_Distance,
@@ -153,6 +313,8 @@ namespace UE::Audio::Insights
 			TraceServices::IAnalysisSession& Session;
 		};
 
+		ActiveSoundToEntryKeysMap.Empty();
+
 		return new FSoundTraceAnalyzer(AsShared(), InSession);
 	}
 
@@ -166,36 +328,36 @@ namespace UE::Audio::Insights
 	bool FSoundTraceProvider::ProcessMessages()
 	{
 		// Helper lambdas
-		auto GetOrCreateEntry = [this](const FSoundMessageBase& Msg)
+		TSharedPtr<FSoundDashboardEntry> ActiveSoundEntryReturn = nullptr;
+		auto GetActiveSoundEntryFromStartMessage = [this, &ActiveSoundEntryReturn](const FSoundWaveStartMessage& Msg) -> TSharedPtr<FSoundDashboardEntry>*
 		{
-			TSharedPtr<FSoundDashboardEntry>* ToReturn = nullptr;
-
-			UpdateDeviceEntry(Msg.DeviceId, Msg.PlayOrder, [&ToReturn, &Msg](TSharedPtr<FSoundDashboardEntry>& Entry)
-			{
-				if (!Entry.IsValid())
-				{
-					Entry = MakeShared<FSoundDashboardEntry>();
-					Entry->DeviceId = Msg.DeviceId;
-					Entry->PlayOrder = Msg.PlayOrder;
-				}
-
-				Entry->Timestamp = Msg.Timestamp;
-				ToReturn = &Entry;
-			});
-
-			return ToReturn;
+			GetActiveSoundEntryFromIDs(Msg.ActiveSoundPlayOrder, Msg.DeviceId, ActiveSoundEntryReturn);
+			return &ActiveSoundEntryReturn;
 		};
 
-		auto GetEntry = [this](const FSoundMessageBase& Msg)
+		TSharedPtr<FSoundDashboardEntry> ParamReturn = nullptr;
+		auto GetActiveSoundEntryFromParameterMessage = [this, &ParamReturn](const FMixerSourceParameterMessage& Msg) -> TSharedPtr<FSoundDashboardEntry>*
 		{
-			return FindDeviceEntry(Msg.DeviceId, Msg.PlayOrder);
+			GetActiveSoundEntryFromIDs(Msg.ActiveSoundPlayOrder, Msg.DeviceId, ParamReturn);
+			return &ParamReturn;
 		};
 
+		TSharedPtr<FSoundDashboardEntry> SoundEntryReturn = nullptr;
+		auto GetSoundEntryFromStopMessage = [this, &SoundEntryReturn](const FSoundStopMessage& Msg) -> TSharedPtr<FSoundDashboardEntry>*
+		{
+			GetActiveSoundEntryFromIDs(Msg.PlayOrder, Msg.DeviceId, SoundEntryReturn);
+			return &SoundEntryReturn;
+		};
 
 		// Process messages
-		ProcessMessageQueue<FSoundStartMessage>(TraceMessages.StartMessages, GetOrCreateEntry,
-		[this](const FSoundStartMessage& Msg, TSharedPtr<FSoundDashboardEntry>* OutEntry)
+		ProcessMessageQueue<FSoundStartMessage>(TraceMessages.ActiveSoundStartMessages, GetOrCreateActiveSoundEntry,
+		[this](const FSoundStartMessage& Msg, const TSharedPtr<FSoundDashboardEntry>* OutActiveSoundEntry)
 		{
+			if (OutActiveSoundEntry == nullptr || !OutActiveSoundEntry->IsValid())
+			{
+				return;
+			}
+
 #if !WITH_EDITOR
 			if (SessionCachedMessages.IsValid())
 			{
@@ -203,17 +365,20 @@ namespace UE::Audio::Insights
 			}
 #endif // !WITH_EDITOR
 
-			FSoundDashboardEntry& EntryRef = *OutEntry->Get();
-			EntryRef.Name = *Msg.Name;
-			EntryRef.EntryType = Msg.EntryType;
+			FSoundDashboardEntry& ActiveSoundEntryRef = *OutActiveSoundEntry->Get();
+			ActiveSoundEntryRef.Name = *Msg.Name;
+			ActiveSoundEntryRef.EntryType = Msg.EntryType;
 		});
 
-		TArray<int32, TInlineAllocator<64>> EntriesWithPoppedDataPoints;
-
-		ProcessMessageQueue<FSoundPriorityMessage>(TraceMessages.PriorityMessages, GetEntry,
-		[this, &EntriesWithPoppedDataPoints](const FSoundPriorityMessage& Msg, TSharedPtr<FSoundDashboardEntry>* OutEntry)
+		ProcessMessageQueue<FSoundWaveStartMessage>(TraceMessages.SoundWaveStartMessages, GetActiveSoundEntryFromStartMessage,
+		[this](const FSoundWaveStartMessage& Msg, const TSharedPtr<FSoundDashboardEntry>* OutActiveSoundEntry)
 		{
-			if (OutEntry == nullptr)
+			if (OutActiveSoundEntry == nullptr || !OutActiveSoundEntry->IsValid())
+			{
+				return;
+			}
+
+			if (!FSoundTraceProviderPrivate::EntryTypeHasSoundWaveEntries((*OutActiveSoundEntry)->EntryType))
 			{
 				return;
 			}
@@ -221,24 +386,47 @@ namespace UE::Audio::Insights
 #if !WITH_EDITOR
 			if (SessionCachedMessages.IsValid())
 			{
-				SessionCachedMessages->PriorityCachedMessages.EmplaceBack(Msg);
+				SessionCachedMessages->SoundWaveStartCachedMessages.EmplaceBack(Msg);
 			}
 #endif // !WITH_EDITOR
+			
+			TSharedPtr<FSoundDashboardEntry> SoundWaveEntry = MakeShared<FSoundDashboardEntry>();
+			SoundWaveEntry->DeviceId = Msg.DeviceId;
+			SoundWaveEntry->PlayOrder = Msg.PlayOrder;
+			SoundWaveEntry->Timestamp = Msg.Timestamp;
+			SoundWaveEntry->Name = *Msg.Name;
+			SoundWaveEntry->EntryType = Msg.EntryType;
+			
+			(*OutActiveSoundEntry)->Children.Add(SoundWaveEntry);
+		});
+
+		TArray<int32, TInlineAllocator<64>> EntriesWithPoppedDataPoints;
 
-			if (!EntriesWithPoppedDataPoints.Contains(Msg.PlayOrder))
+		ProcessMessageQueue<FSoundPriorityMessage>(TraceMessages.PriorityMessages, GetActiveSoundEntryFromParameterMessage,
+		[this, &EntriesWithPoppedDataPoints](const FSoundPriorityMessage& Msg, const TSharedPtr<FSoundDashboardEntry>* OutActiveSoundEntry)
+		{
+			FSoundDashboardEntry* SoundEntry = FSoundTraceProviderPrivate::FindSoundEntryForActiveSound(Msg.PlayOrder, OutActiveSoundEntry);
+			if (SoundEntry == nullptr)
+			{
+				return;
+			}
+
+#if !WITH_EDITOR
+			if (SessionCachedMessages.IsValid())
 			{
-				(*OutEntry)->PriorityDataPoints.Pop((*OutEntry)->PriorityDataPoints.Num());
-				EntriesWithPoppedDataPoints.Add(Msg.PlayOrder);
+				SessionCachedMessages->PriorityCachedMessages.EmplaceBack(Msg);
 			}
+#endif // !WITH_EDITOR
 
-			(*OutEntry)->PriorityDataPoints.Push({ Msg.Timestamp, Msg.Priority });
+			FSoundTraceProviderPrivate::UpdateParameterEntry(Msg, Msg.Priority, SoundEntry->PriorityDataPoints, EntriesWithPoppedDataPoints);
 		});
-		
 		EntriesWithPoppedDataPoints.Reset();
-		ProcessMessageQueue<FSoundDistanceMessage>(TraceMessages.DistanceMessages, GetEntry,
-		[this, &EntriesWithPoppedDataPoints](const FSoundDistanceMessage& Msg, TSharedPtr<FSoundDashboardEntry>* OutEntry)
+
+		ProcessMessageQueue<FSoundDistanceMessage>(TraceMessages.DistanceMessages, GetActiveSoundEntryFromParameterMessage,
+		[this, &EntriesWithPoppedDataPoints](const FSoundDistanceMessage& Msg, const TSharedPtr<FSoundDashboardEntry>* OutActiveSoundEntry)
 		{
-			if (OutEntry == nullptr)
+			FSoundDashboardEntry* SoundEntry = FSoundTraceProviderPrivate::FindSoundEntryForActiveSound(Msg.PlayOrder, OutActiveSoundEntry);
+			if (SoundEntry == nullptr)
 			{
 				return;
 			}
@@ -250,20 +438,15 @@ namespace UE::Audio::Insights
 			}
 #endif // !WITH_EDITOR
 
-			if (!EntriesWithPoppedDataPoints.Contains(Msg.PlayOrder))
-			{
-				(*OutEntry)->DistanceDataPoints.Pop((*OutEntry)->DistanceDataPoints.Num());
-				EntriesWithPoppedDataPoints.Add(Msg.PlayOrder);
-			}
-
-			(*OutEntry)->DistanceDataPoints.Push({ Msg.Timestamp, Msg.Distance });
+			FSoundTraceProviderPrivate::UpdateParameterEntry(Msg, Msg.Distance, SoundEntry->DistanceDataPoints, EntriesWithPoppedDataPoints);
 		});
-		
 		EntriesWithPoppedDataPoints.Reset();
-		ProcessMessageQueue<FSoundAmplitudeMessage>(TraceMessages.AmplitudeMessages, GetEntry,
-		[this, &EntriesWithPoppedDataPoints](const FSoundAmplitudeMessage& Msg, TSharedPtr<FSoundDashboardEntry>* OutEntry)
+
+		ProcessMessageQueue<FMixerSourceEnvelopeMessage>(TraceMessages.AmplitudeMessages, GetActiveSoundEntryFromParameterMessage,
+		[this, &EntriesWithPoppedDataPoints](const FMixerSourceEnvelopeMessage& Msg, const TSharedPtr<FSoundDashboardEntry>* OutActiveSoundEntry)
 		{
-			if (OutEntry == nullptr)
+			FSoundDashboardEntry* SoundEntry = FSoundTraceProviderPrivate::FindSoundEntryForActiveSound(Msg.PlayOrder, OutActiveSoundEntry);
+			if (SoundEntry == nullptr)
 			{
 				return;
 			}
@@ -275,19 +458,15 @@ namespace UE::Audio::Insights
 			}
 #endif // !WITH_EDITOR
 
-			if (!EntriesWithPoppedDataPoints.Contains(Msg.PlayOrder))
-			{
-				(*OutEntry)->AmplitudeDataPoints.Pop((*OutEntry)->AmplitudeDataPoints.Num());
-				EntriesWithPoppedDataPoints.Add(Msg.PlayOrder);
-			}
-
-			(*OutEntry)->AmplitudeDataPoints.Push({ Msg.Timestamp, Msg.Amplitude });
+			FSoundTraceProviderPrivate::UpdateParameterEntry(Msg, Msg.Envelope, SoundEntry->AmplitudeDataPoints, EntriesWithPoppedDataPoints);
 		});
+		EntriesWithPoppedDataPoints.Reset();
 
-		ProcessMessageQueue<FSoundVolumeMessage>(TraceMessages.VolumeMessages, GetEntry,
-		[this, &EntriesWithPoppedDataPoints](const FSoundVolumeMessage& Msg, TSharedPtr<FSoundDashboardEntry>* OutEntry)
+		ProcessMessageQueue<FMixerSourceVolumeMessage>(TraceMessages.VolumeMessages, GetActiveSoundEntryFromParameterMessage,
+		[this, &EntriesWithPoppedDataPoints](const FMixerSourceVolumeMessage& Msg, const TSharedPtr<FSoundDashboardEntry>* OutActiveSoundEntry)
 		{
-			if (OutEntry == nullptr)
+			FSoundDashboardEntry* SoundEntry = FSoundTraceProviderPrivate::FindSoundEntryForActiveSound(Msg.PlayOrder, OutActiveSoundEntry);
+			if (SoundEntry == nullptr)
 			{
 				return;
 			}
@@ -299,20 +478,15 @@ namespace UE::Audio::Insights
 			}
 #endif // !WITH_EDITOR
 
-			if (!EntriesWithPoppedDataPoints.Contains(Msg.PlayOrder))
-			{
-				(*OutEntry)->VolumeDataPoints.Pop((*OutEntry)->VolumeDataPoints.Num());
-				EntriesWithPoppedDataPoints.Add(Msg.PlayOrder);
-			}
-
-			(*OutEntry)->VolumeDataPoints.Push({ Msg.Timestamp, Msg.Volume });
+			FSoundTraceProviderPrivate::UpdateParameterEntry(Msg, Msg.Volume, SoundEntry->VolumeDataPoints, EntriesWithPoppedDataPoints);
 		});
-
 		EntriesWithPoppedDataPoints.Reset();
-		ProcessMessageQueue<FSoundPitchMessage>(TraceMessages.PitchMessages, GetEntry,
-		[this, &EntriesWithPoppedDataPoints](const FSoundPitchMessage& Msg, TSharedPtr<FSoundDashboardEntry>* OutEntry)
+
+		ProcessMessageQueue<FMixerSourcePitchMessage>(TraceMessages.PitchMessages, GetActiveSoundEntryFromParameterMessage,
+		[this, &EntriesWithPoppedDataPoints](const FMixerSourcePitchMessage& Msg, const TSharedPtr<FSoundDashboardEntry>* OutActiveSoundEntry)
 		{
-			if (OutEntry == nullptr)
+			FSoundDashboardEntry* SoundEntry = FSoundTraceProviderPrivate::FindSoundEntryForActiveSound(Msg.PlayOrder, OutActiveSoundEntry);
+			if (SoundEntry == nullptr)
 			{
 				return;
 			}
@@ -324,17 +498,11 @@ namespace UE::Audio::Insights
 			}
 #endif // !WITH_EDITOR
 
-			if (!EntriesWithPoppedDataPoints.Contains(Msg.PlayOrder))
-			{
-				(*OutEntry)->PitchDataPoints.Pop((*OutEntry)->PitchDataPoints.Num());
-				EntriesWithPoppedDataPoints.Add(Msg.PlayOrder);
-			}
-
-			(*OutEntry)->PitchDataPoints.Push({ Msg.Timestamp, Msg.Pitch });
+			FSoundTraceProviderPrivate::UpdateParameterEntry(Msg, Msg.Pitch, SoundEntry->PitchDataPoints, EntriesWithPoppedDataPoints);
 		});
 
-		ProcessMessageQueue<FSoundStopMessage>(TraceMessages.StopMessages, GetEntry,
-		[this](const FSoundStopMessage& Msg, TSharedPtr<FSoundDashboardEntry>* OutEntry)
+		ProcessMessageQueue<FSoundStopMessage>(TraceMessages.StopMessages, GetSoundEntryFromStopMessage,
+		[this](const FSoundStopMessage& Msg, const TSharedPtr<FSoundDashboardEntry>* OutActiveSoundEntry)
 		{
 #if !WITH_EDITOR
 			if (SessionCachedMessages.IsValid())
@@ -343,15 +511,141 @@ namespace UE::Audio::Insights
 			}
 #endif // !WITH_EDITOR
 
-			if (OutEntry && (*OutEntry)->Timestamp < Msg.Timestamp)
+			if (OutActiveSoundEntry && OutActiveSoundEntry->IsValid() && (*OutActiveSoundEntry)->Timestamp < Msg.Timestamp)
 			{
-				RemoveDeviceEntry(Msg.DeviceId, Msg.PlayOrder);
+				RemoveActiveSoundEntry(Msg, *OutActiveSoundEntry);
 			}
 		});
 
 		return true;
 	}
 
+	void FSoundTraceProvider::InitLambdas()
+	{
+		GetOrCreateActiveSoundEntry = [this](const FSoundStartMessage& Msg)
+		{
+			TSharedPtr<FSoundDashboardEntry>* ReturnedSoundEntry = nullptr;
+
+			UpdateDeviceEntry(Msg.DeviceId, Msg.EntryType, [this, &ReturnedSoundEntry, &Msg](TSharedPtr<FSoundDashboardEntry>& CategoryEntry)
+			{
+				if (!CategoryEntry.IsValid())
+				{
+					CategoryEntry = MakeShared<FSoundDashboardEntry>();
+					CategoryEntry->DeviceId = Msg.DeviceId;
+					CategoryEntry->EntryType = Msg.EntryType;
+					CategoryEntry->Name = FSoundTraceProviderPrivate::GetEntryTypeDisplayName(Msg.EntryType);
+				}
+
+				TSharedPtr<IDashboardDataTreeViewEntry>* AssetEntry = CategoryEntry->Children.FindByPredicate([&Msg](TSharedPtr<IDashboardDataTreeViewEntry> ChildEntry)
+				{
+					FSoundDashboardEntry& NewEntry = FSoundTraceProviderPrivate::CastEntry(*ChildEntry.Get());
+					return NewEntry.Name == Msg.Name;
+				});
+
+				TSharedPtr<FSoundDashboardEntry> NewAssetEntry = nullptr;
+				if (AssetEntry == nullptr || !AssetEntry->IsValid())
+				{
+					NewAssetEntry = MakeShared<FSoundDashboardEntry>();
+					NewAssetEntry->DeviceId = Msg.DeviceId;
+					NewAssetEntry->EntryType = Msg.EntryType;
+					NewAssetEntry->Name = Msg.Name;
+
+					CategoryEntry->Children.Add(NewAssetEntry);
+				}
+				else
+				{
+					NewAssetEntry = StaticCastSharedPtr<FSoundDashboardEntry>(*AssetEntry);
+				}
+
+				TSharedPtr<IDashboardDataTreeViewEntry>* SoundEntry = NewAssetEntry->Children.FindByPredicate([&Msg](TSharedPtr<IDashboardDataTreeViewEntry> ChildEntry)
+				{
+					FSoundDashboardEntry& NewEntry = FSoundTraceProviderPrivate::CastEntry(*ChildEntry.Get());
+					return NewEntry.PlayOrder == Msg.PlayOrder;
+				});
+
+				if (SoundEntry == nullptr || !SoundEntry->IsValid())
+				{
+					TSharedPtr<FSoundDashboardEntry> NewEntry = MakeShared<FSoundDashboardEntry>();
+					NewEntry->DeviceId = Msg.DeviceId;
+					NewEntry->PlayOrder = Msg.PlayOrder;
+
+					NewAssetEntry->Children.Add(NewEntry);
+
+					ActiveSoundToEntryKeysMap.Add(Msg.PlayOrder, SoundEntryKeys{ Msg.EntryType, Msg.Name });
+
+					ReturnedSoundEntry = &NewEntry;
+				}
+				else
+				{
+					TSharedPtr<FSoundDashboardEntry> NewEntry = StaticCastSharedPtr<FSoundDashboardEntry>(*SoundEntry);
+					ReturnedSoundEntry = &NewEntry;
+				}
+
+				if (ReturnedSoundEntry == nullptr)
+				{
+					return;
+				}
+
+				(*ReturnedSoundEntry)->Timestamp = Msg.Timestamp;
+			});
+
+			return ReturnedSoundEntry;
+		};
+	}
+
+	void FSoundTraceProvider::GetActiveSoundEntryFromIDs(const uint32 ActiveSoundPlayOrder, const ::Audio::FDeviceId DeviceId, TSharedPtr<FSoundDashboardEntry>& OutActiveSoundEntry)
+	{
+		const SoundEntryKeys* SoundEntryKeys = ActiveSoundToEntryKeysMap.Find(ActiveSoundPlayOrder);
+		if (SoundEntryKeys == nullptr)
+		{
+			return;
+		}
+
+		TSharedPtr<FSoundDashboardEntry>* CategoryEntry = FindDeviceEntry(DeviceId, SoundEntryKeys->EntryType);
+
+		FSoundTraceProviderPrivate::FindActiveSoundEntryInCategory(CategoryEntry, SoundEntryKeys->AssetName, ActiveSoundPlayOrder, OutActiveSoundEntry);
+	}
+
+	void FSoundTraceProvider::RemoveActiveSoundEntry(const FSoundStopMessage& Msg, TSharedPtr<FSoundDashboardEntry> OutActiveSoundEntry)
+	{
+		if (!OutActiveSoundEntry.IsValid())
+		{
+			return;
+		}
+
+		const uint32 ActiveSoundPlayOrder = OutActiveSoundEntry->PlayOrder;
+		const SoundEntryKeys* SoundEntryKeys = ActiveSoundToEntryKeysMap.Find(ActiveSoundPlayOrder);
+
+		if (SoundEntryKeys == nullptr)
+		{
+			return;
+		}
+
+		TSharedPtr<FSoundDashboardEntry>* CategoryEntry = FindDeviceEntry(Msg.DeviceId, SoundEntryKeys->EntryType);
+
+		// First remove the entry from the asset name
+		TSharedPtr<FSoundDashboardEntry> AssetEntry = nullptr;
+		FSoundTraceProviderPrivate::GetAssetEntry(CategoryEntry, SoundEntryKeys->AssetName, AssetEntry);
+		if (AssetEntry == nullptr || !AssetEntry.IsValid())
+		{
+			return;
+		}
+
+		AssetEntry->Children.Remove(OutActiveSoundEntry);
+
+		if (AssetEntry->Children.Num() == 0)
+		{
+			(*CategoryEntry)->Children.Remove(AssetEntry);
+
+			if ((*CategoryEntry)->Children.Num() == 0)
+			{
+				RemoveDeviceEntry(Msg.DeviceId, OutActiveSoundEntry->EntryType);
+			}
+		}
+
+		ActiveSoundToEntryKeysMap.Remove(ActiveSoundPlayOrder);
+	}
+
 #if !WITH_EDITOR
 	void FSoundTraceProvider::OnTimingViewTimeMarkerChanged(double TimeMarker)
 	{
@@ -364,7 +658,22 @@ namespace UE::Audio::Insights
 
 		DeviceDataMap.Empty();
 
+		TSharedPtr<FSoundDashboardEntry> SoundEntryReturn = nullptr;
+		auto GetActiveSoundEntryFromStopMessage = [this, &SoundEntryReturn](const FSoundStopMessage& Msg) -> TSharedPtr<FSoundDashboardEntry>*
+		{
+			GetActiveSoundEntryFromIDs(Msg.PlayOrder, Msg.DeviceId, SoundEntryReturn);
+			return &SoundEntryReturn;
+		};
+
+		TSharedPtr<FSoundDashboardEntry> ParamReturn = nullptr;
+		auto GetActiveSoundEntryFromParameterMessage = [this, &ParamReturn](const FMixerSourceParameterMessage& Msg) -> TSharedPtr<FSoundDashboardEntry>*
+		{
+			GetActiveSoundEntryFromIDs(Msg.ActiveSoundPlayOrder, Msg.DeviceId, ParamReturn);
+			return &ParamReturn;
+		};
+
 		// Collect all the start messages registered until this point in time 
+		TSharedPtr<FSoundDashboardEntry>* ReturnedSoundEntry = nullptr;
 		for (const FSoundStartMessage& StartCachedMessage : SessionCachedMessages->StartCachedMessages)
 		{
 			if (StartCachedMessage.Timestamp > TimeMarker)
@@ -372,18 +681,47 @@ namespace UE::Audio::Insights
 				break;
 			}
 
-			UpdateDeviceEntry(StartCachedMessage.DeviceId, StartCachedMessage.PlayOrder, [&StartCachedMessage](TSharedPtr<FSoundDashboardEntry>& Entry)
+			TSharedPtr<FSoundDashboardEntry>* ActiveSoundEntry = GetOrCreateActiveSoundEntry(StartCachedMessage);
+			if (ActiveSoundEntry == nullptr || !ActiveSoundEntry->IsValid())
 			{
-				if (!Entry.IsValid())
-				{
-					Entry = MakeShared<FSoundDashboardEntry>();
-					Entry->DeviceId  = StartCachedMessage.DeviceId;
-					Entry->PlayOrder = StartCachedMessage.PlayOrder;
-				}
-				Entry->Timestamp = StartCachedMessage.Timestamp;
+				continue;
+			}
 
-				Entry->Name = *StartCachedMessage.Name;
-			});
+			FSoundDashboardEntry& EntryRef = *ActiveSoundEntry->Get();
+			EntryRef.Name = *StartCachedMessage.Name;
+			EntryRef.EntryType = StartCachedMessage.EntryType;
+		}
+
+		////////////////////////////////////////
+		// Add all soundwave start messages registered until this point in time 
+		TSharedPtr<FSoundDashboardEntry> ParentSoundEntryReturn = nullptr;
+		for (const FSoundWaveStartMessage& SoundWaveStartCachedMessage : SessionCachedMessages->SoundWaveStartCachedMessages)
+		{
+			if (SoundWaveStartCachedMessage.Timestamp > TimeMarker)
+			{
+				break;
+			}
+
+			if (!FSoundTraceProviderPrivate::EntryTypeHasSoundWaveEntries(SoundWaveStartCachedMessage.EntryType))
+			{
+				return;
+			}
+
+			GetActiveSoundEntryFromIDs(SoundWaveStartCachedMessage.ActiveSoundPlayOrder, SoundWaveStartCachedMessage.DeviceId, ParentSoundEntryReturn);
+
+			if (!ParentSoundEntryReturn->IsValid())
+			{
+				return;
+			}
+
+			TSharedPtr<FSoundDashboardEntry> ChildEntry = MakeShared<FSoundDashboardEntry>();
+			ChildEntry->DeviceId = SoundWaveStartCachedMessage.DeviceId;
+			ChildEntry->PlayOrder = SoundWaveStartCachedMessage.PlayOrder;
+			ChildEntry->Timestamp = SoundWaveStartCachedMessage.Timestamp;
+			ChildEntry->Name = *SoundWaveStartCachedMessage.Name;
+			ChildEntry->EntryType = SoundWaveStartCachedMessage.EntryType;
+
+			ParentSoundEntryReturn->Children.Add(ChildEntry);
 		}
 
 		// Selectively remove start messages collected in the step above by knowing which sounds were stopped.
@@ -395,14 +733,14 @@ namespace UE::Audio::Insights
 				break;
 			}
 
-			const TSharedPtr<FSoundDashboardEntry>* OutEntry = FindDeviceEntry(StopCachedMessage.DeviceId, StopCachedMessage.PlayOrder);
+			TSharedPtr<FSoundDashboardEntry>* ActiveSoundEntry = GetActiveSoundEntryFromStopMessage(StopCachedMessage);
 
-			if (OutEntry && (*OutEntry)->Timestamp < StopCachedMessage.Timestamp)
+			if (ActiveSoundEntry && (*ActiveSoundEntry)->Timestamp < StopCachedMessage.Timestamp)
 			{
-				RemoveDeviceEntry(StopCachedMessage.DeviceId, StopCachedMessage.PlayOrder);
+				RemoveActiveSoundEntry(StopCachedMessage, *ActiveSoundEntry);
 			}
 		}
-
+		
 		// For now we only retrieve information from AudioDeviceId 1 (main device in standalone games)
 		const FDeviceData* DeviceData = DeviceDataMap.Find(1);
 		if (DeviceData)
@@ -412,13 +750,13 @@ namespace UE::Audio::Insights
 			{
 				FSoundPriorityMessage PriorityMessage;
 				FSoundDistanceMessage DistanceMessage;
-				FSoundAmplitudeMessage AmplitudeMessage;
-				FSoundVolumeMessage VolumeMessage;
-				FSoundPitchMessage PitchMessage;
+				FMixerSourceEnvelopeMessage AmplitudeMessage;
+				FMixerSourceVolumeMessage VolumeMessage;
+				FMixerSourcePitchMessage PitchMessage;
 			};
 
 			TArray<uint32> PlayOrderArray;
-			DeviceData->GenerateKeyArray(PlayOrderArray);
+			ActiveSoundToEntryKeysMap.GetKeys(PlayOrderArray);
 
 			TArray<CachedEntryInfo> CachedEntryInfos;
 			CachedEntryInfos.SetNumUninitialized(PlayOrderArray.Num());
@@ -444,21 +782,21 @@ namespace UE::Audio::Insights
 				}
 
 				// Amplitude
-				const FSoundAmplitudeMessage* FoundAmplitudeCachedMessage = FindClosestMessageToTimestamp(SessionCachedMessages->AmplitudeCachedMessages, TimeMarker, PlayOrder);
+				const FMixerSourceEnvelopeMessage* FoundAmplitudeCachedMessage = FindClosestMessageToTimestamp(SessionCachedMessages->AmplitudeCachedMessages, TimeMarker, PlayOrder);
 				if (FoundAmplitudeCachedMessage)
 				{
 					CachedEntryInfos[Index].AmplitudeMessage = *FoundAmplitudeCachedMessage;
 				}
 
 				// Volume
-				const FSoundVolumeMessage* FoundVolumeCachedMessage = FindClosestMessageToTimestamp(SessionCachedMessages->VolumeCachedMessages, TimeMarker, PlayOrder);
+				const FMixerSourceVolumeMessage* FoundVolumeCachedMessage = FindClosestMessageToTimestamp(SessionCachedMessages->VolumeCachedMessages, TimeMarker, PlayOrder);
 				if (FoundVolumeCachedMessage)
 				{
 					CachedEntryInfos[Index].VolumeMessage = *FoundVolumeCachedMessage;
 				}
 
 				// Pitch
-				const FSoundPitchMessage* FoundPitchCachedMessage = FindClosestMessageToTimestamp(SessionCachedMessages->PitchCachedMessages, TimeMarker, PlayOrder);
+				const FMixerSourcePitchMessage* FoundPitchCachedMessage = FindClosestMessageToTimestamp(SessionCachedMessages->PitchCachedMessages, TimeMarker, PlayOrder);
 				if (FoundPitchCachedMessage)
 				{
 					CachedEntryInfos[Index].PitchMessage = *FoundPitchCachedMessage;
@@ -468,77 +806,33 @@ namespace UE::Audio::Insights
 			// Update the device entries with the collected info
 			for (const CachedEntryInfo& CachedEntryInfo : CachedEntryInfos)
 			{
-				UpdateDeviceEntry(CachedEntryInfo.PriorityMessage.DeviceId, CachedEntryInfo.PriorityMessage.PlayOrder, [&CachedEntryInfo](TSharedPtr<FSoundDashboardEntry>& Entry)
-				{
-					if (!Entry.IsValid())
-					{
-						Entry = MakeShared<FSoundDashboardEntry>();
-						Entry->DeviceId  = CachedEntryInfo.PriorityMessage.DeviceId;
-						Entry->PlayOrder = CachedEntryInfo.PriorityMessage.PlayOrder;
-					}
-
-					Entry->Timestamp = CachedEntryInfo.PriorityMessage.Timestamp;
-					Entry->PriorityDataPoints.Push({ CachedEntryInfo.PriorityMessage.Timestamp, CachedEntryInfo.PriorityMessage.Priority });
-				});
-
-				UpdateDeviceEntry(CachedEntryInfo.DistanceMessage.DeviceId, CachedEntryInfo.DistanceMessage.PlayOrder, [&CachedEntryInfo](TSharedPtr<FSoundDashboardEntry>& Entry)
-				{
-					if (!Entry.IsValid())
-					{
-						Entry = MakeShared<FSoundDashboardEntry>();
-						Entry->DeviceId  = CachedEntryInfo.DistanceMessage.DeviceId;
-						Entry->PlayOrder = CachedEntryInfo.DistanceMessage.PlayOrder;
-					}
-
-					Entry->Timestamp = CachedEntryInfo.DistanceMessage.Timestamp;
-					Entry->DistanceDataPoints.Push({ CachedEntryInfo.DistanceMessage.Timestamp, CachedEntryInfo.DistanceMessage.Distance });
-				});
-
-				UpdateDeviceEntry(CachedEntryInfo.AmplitudeMessage.DeviceId, CachedEntryInfo.AmplitudeMessage.PlayOrder, [&CachedEntryInfo](TSharedPtr<FSoundDashboardEntry>& Entry)
+				TSharedPtr<FSoundDashboardEntry>* ActiveSoundEntry = GetActiveSoundEntryFromParameterMessage(CachedEntryInfo.PriorityMessage);
+				
+				if (ActiveSoundEntry == nullptr || !ActiveSoundEntry->IsValid())
 				{
-					if (!Entry.IsValid())
-					{
-						Entry = MakeShared<FSoundDashboardEntry>();
-						Entry->DeviceId  = CachedEntryInfo.AmplitudeMessage.DeviceId;
-						Entry->PlayOrder = CachedEntryInfo.AmplitudeMessage.PlayOrder;
-					}
-
-					Entry->Timestamp = CachedEntryInfo.AmplitudeMessage.Timestamp;
-					Entry->AmplitudeDataPoints.Push({ CachedEntryInfo.AmplitudeMessage.Timestamp, CachedEntryInfo.AmplitudeMessage.Amplitude });
-				});
+					continue;
+				}
 
-				UpdateDeviceEntry(CachedEntryInfo.VolumeMessage.DeviceId, CachedEntryInfo.VolumeMessage.PlayOrder, [&CachedEntryInfo](TSharedPtr<FSoundDashboardEntry>& Entry)
+				FSoundDashboardEntry* SoundEntry = FSoundTraceProviderPrivate::FindSoundEntryForActiveSound(CachedEntryInfo.PriorityMessage.PlayOrder, ActiveSoundEntry);
+				if (SoundEntry == nullptr)
 				{
-					if (!Entry.IsValid())
-					{
-						Entry = MakeShared<FSoundDashboardEntry>();
-						Entry->DeviceId  = CachedEntryInfo.VolumeMessage.DeviceId;
-						Entry->PlayOrder = CachedEntryInfo.VolumeMessage.PlayOrder;
-					}
-
-					Entry->Timestamp = CachedEntryInfo.VolumeMessage.Timestamp;
-					Entry->VolumeDataPoints.Push({ CachedEntryInfo.VolumeMessage.Timestamp, CachedEntryInfo.VolumeMessage.Volume });
-				});
+					continue;
+				}
 
-				UpdateDeviceEntry(CachedEntryInfo.PitchMessage.DeviceId, CachedEntryInfo.PitchMessage.PlayOrder, [&CachedEntryInfo](TSharedPtr<FSoundDashboardEntry>& Entry)
-				{
-					if (!Entry.IsValid())
-					{
-						Entry = MakeShared<FSoundDashboardEntry>();
-						Entry->DeviceId  = CachedEntryInfo.PitchMessage.DeviceId;
-						Entry->PlayOrder = CachedEntryInfo.PitchMessage.PlayOrder;
-					}
+				SoundEntry->PriorityDataPoints.Push({ CachedEntryInfo.PriorityMessage.Timestamp, CachedEntryInfo.PriorityMessage.Priority });
+				SoundEntry->DistanceDataPoints.Push({ CachedEntryInfo.DistanceMessage.Timestamp, CachedEntryInfo.DistanceMessage.Distance });
+				SoundEntry->AmplitudeDataPoints.Push({ CachedEntryInfo.AmplitudeMessage.Timestamp, CachedEntryInfo.AmplitudeMessage.Envelope });
+				SoundEntry->VolumeDataPoints.Push({ CachedEntryInfo.VolumeMessage.Timestamp, CachedEntryInfo.VolumeMessage.Volume });
+				SoundEntry->PitchDataPoints.Push({ CachedEntryInfo.PitchMessage.Timestamp, CachedEntryInfo.PitchMessage.Pitch });
 
-					Entry->Timestamp = CachedEntryInfo.PitchMessage.Timestamp;
-					Entry->PitchDataPoints.Push({ CachedEntryInfo.PitchMessage.Timestamp, CachedEntryInfo.PitchMessage.Pitch });
-				});
+				SoundEntry->Timestamp = CachedEntryInfo.PriorityMessage.Timestamp;
 			}
 		}
 
 		// Call parent method to update LastMessageId
 		FTraceProviderBase::OnTimingViewTimeMarkerChanged(TimeMarker);
 	}
-
 #endif // !WITH_EDITOR
 
 } // namespace UE::Audio::Insights
+#undef LOCTEXT_NAMESPACE
\ No newline at end of file
diff --git a/Engine/Plugins/AudioInsights/Source/AudioInsights/Private/Providers/SoundTraceProvider.h b/Engine/Plugins/AudioInsights/Source/AudioInsights/Private/Providers/SoundTraceProvider.h
--- a/Engine/Plugins/AudioInsights/Source/AudioInsights/Private/Providers/SoundTraceProvider.h
+++ b/Engine/Plugins/AudioInsights/Source/AudioInsights/Private/Providers/SoundTraceProvider.h
@@ -5,10 +5,20 @@
 #include "Messages/SoundTraceMessages.h"
 #include "UObject/NameTypes.h"
 
+/**
+ *	Trace provider for the "Sounds Dashboard"
+ *	
+ *	"Sounds Dashboard" displays active audio in the following tree structure:
+ *	- Category											(i.e. The EntryType of the sound asset, e.g. SoundCue, MetaSound etc)
+ *		- Asset											(i.e. Sound instances from the same asset are grouped here)
+ *			- Active Sound								(i.e. The Active Sound instance playing)
+ *				- SoundWave Instances (optional)		(i.e. If an entry type can have multiple SoundWaves playing simultaneously (e.g. SoundCue)
+ *														, individual Soundwaves are shown here)
+ */
 namespace UE::Audio::Insights
 {
 	class FSoundTraceProvider
-		: public TDeviceDataMapTraceProvider<uint32, TSharedPtr<FSoundDashboardEntry>>
+		: public TDeviceDataMapTraceProvider<ESoundDashboardEntryType, TSharedPtr<FSoundDashboardEntry>>
 		, public TSharedFromThis<FSoundTraceProvider>
 	{
 	public:
@@ -26,6 +36,10 @@ namespace UE::Audio::Insights
 		virtual bool ProcessMessages() override;
 
 	private:
+		void InitLambdas();
+		void GetActiveSoundEntryFromIDs(const uint32 PlayOrderID, const ::Audio::FDeviceId DeviceId, TSharedPtr<FSoundDashboardEntry>& OutSoundEntry);
+		void RemoveActiveSoundEntry(const FSoundStopMessage& Msg, TSharedPtr<FSoundDashboardEntry> OutEntry);
+
 #if !WITH_EDITOR
 		virtual void OnTimingViewTimeMarkerChanged(double TimeMarker) override;
 
@@ -33,5 +47,15 @@ namespace UE::Audio::Insights
 #endif // !WITH_EDITOR
 
 		FSoundMessages TraceMessages;
+
+		struct SoundEntryKeys
+		{
+			ESoundDashboardEntryType EntryType;
+			FString AssetName;
+		};
+
+		TMap<uint32, SoundEntryKeys> ActiveSoundToEntryKeysMap;
+
+		TFunction<TSharedPtr<FSoundDashboardEntry>*(const FSoundStartMessage& Msg)> GetOrCreateActiveSoundEntry;
 	};
 } // namespace UE::Audio::Insights
diff --git a/Engine/Plugins/AudioInsights/Source/AudioInsights/Private/Views/SoundDashboardViewFactory.cpp b/Engine/Plugins/AudioInsights/Source/AudioInsights/Private/Views/SoundDashboardViewFactory.cpp
--- a/Engine/Plugins/AudioInsights/Source/AudioInsights/Private/Views/SoundDashboardViewFactory.cpp
+++ b/Engine/Plugins/AudioInsights/Source/AudioInsights/Private/Views/SoundDashboardViewFactory.cpp
@@ -1,6 +1,7 @@
 // Copyright Epic Games, Inc. All Rights Reserved.
 #include "Views/SoundDashboardViewFactory.h"
 
+#include "Algo/Accumulate.h"
 #include "AudioInsightsModule.h"
 #include "AudioInsightsStyle.h"
 #include "Messages/SoundTraceMessages.h"
@@ -364,7 +365,13 @@ namespace UE::Audio::Insights
 					.ColorAndOpacity(FSlateColor(FColor::White.WithAlpha(128)))
 					.Text_Lambda([InRowData]()
 					{
-						return FText::FromString("(" + FString::FromInt(InRowData->Children.Num()) + ")");
+						const uint32 TotalNumChildren = Algo::Accumulate(InRowData->Children, 0,
+							[](uint32 Accum, TSharedPtr<IDashboardDataTreeViewEntry> InChild)
+							{
+								return Accum + (InChild ? InChild->Children.Num() : 0);
+							});
+
+						return FText::FromString("(" + FString::FromInt(TotalNumChildren) + ")");
 					})
 				];
 		}
@@ -546,12 +553,12 @@ namespace UE::Audio::Insights
 			{
 				"Priority",
 				{
-					.DisplayName = LOCTEXT("SoundDashboard_PriorityColumnDisplayName", "Priority"),
-					.IconName = NAME_None,
+					.DisplayName	  = LOCTEXT("SoundDashboard_PriorityColumnDisplayName", "Priority"),
+					.IconName		  = NAME_None,
 					.bShowDisplayName = true,
-					.bDefaultHidden = false,
-					.FillWidth = 0.08f,
-					.Alignment = EHorizontalAlignment::HAlign_Left
+					.bDefaultHidden	  = false,
+					.FillWidth		  = 0.08f,
+					.Alignment		  = EHorizontalAlignment::HAlign_Left
 				}
 			},
 			{
@@ -669,9 +676,25 @@ namespace UE::Audio::Insights
 					.GetDisplayValue = [](const IDashboardDataTreeViewEntry& InData)
 					{
 						const FSoundDashboardEntry& SoundDashboardEntry = CastEntry(InData);
+						if (SoundDashboardEntry.Children.Num() > 0)
+						{
+							return FText();
+						}
 
 						const ::Audio::TCircularAudioBuffer<FDataPoint>& PriorityDataPoints = SoundDashboardEntry.PriorityDataPoints;
-						return FText::AsNumber(GetLastEntryArrayValue(PriorityDataPoints), FSlateStyle::Get().GetAmpFloatFormat());
+						const float PriorityValue = GetLastEntryArrayValue(PriorityDataPoints);
+
+						// Max priority as defined in SoundWave.cpp
+						static constexpr float VolumeWeightedMaxPriority = TNumericLimits<float>::Max() / MAX_VOLUME;
+
+						if (PriorityValue >= VolumeWeightedMaxPriority)
+						{
+							return LOCTEXT("AudioDashboard_Sounds_Max", "MAX");
+						}
+						else
+						{
+							return FText::AsNumber(GetLastEntryArrayValue(PriorityDataPoints), FSlateStyle::Get().GetAmpFloatFormat());
+						}
 					}
 				}
 			},
@@ -681,9 +704,13 @@ namespace UE::Audio::Insights
 					.GetDisplayValue = [](const IDashboardDataTreeViewEntry& InData)
 					{
 						const FSoundDashboardEntry& SoundDashboardEntry = CastEntry(InData);
+						if (SoundDashboardEntry.Children.Num() > 0)
+						{
+							return FText();
+						}
 
 						const ::Audio::TCircularAudioBuffer<FDataPoint>& DistanceDataPoints = SoundDashboardEntry.DistanceDataPoints;
-						return FText::AsNumber(GetLastEntryArrayValue(DistanceDataPoints), FSlateStyle::Get().GetAmpFloatFormat());
+						return FText::AsNumber(GetLastEntryArrayValue(DistanceDataPoints), FSlateStyle::Get().GetDefaultFloatFormat());
 					}
 				}
 			},
@@ -693,6 +720,10 @@ namespace UE::Audio::Insights
 					.GetDisplayValue = [](const IDashboardDataTreeViewEntry& InData)
 					{
 						const FSoundDashboardEntry& SoundDashboardEntry = CastEntry(InData);
+						if (SoundDashboardEntry.Children.Num() > 0)
+						{
+							return FText();
+						}
 
 						const ::Audio::TCircularAudioBuffer<FDataPoint>& AmplitudeDataPoints = SoundDashboardEntry.AmplitudeDataPoints;
 						return FText::AsNumber(GetLastEntryArrayValue(AmplitudeDataPoints), FSlateStyle::Get().GetAmpFloatFormat());
@@ -705,6 +736,10 @@ namespace UE::Audio::Insights
 					.GetDisplayValue = [](const IDashboardDataTreeViewEntry& InData)
 					{
 						const FSoundDashboardEntry& SoundDashboardEntry = CastEntry(InData);
+						if (SoundDashboardEntry.Children.Num() > 0)
+						{
+							return FText();
+						}
 
 						const ::Audio::TCircularAudioBuffer<FDataPoint>& VolumeDataPoints = SoundDashboardEntry.VolumeDataPoints;
 						return FText::AsNumber(GetLastEntryArrayValue(VolumeDataPoints), FSlateStyle::Get().GetAmpFloatFormat());
@@ -717,6 +752,10 @@ namespace UE::Audio::Insights
 					.GetDisplayValue = [](const IDashboardDataTreeViewEntry& InData)
 					{
 						const FSoundDashboardEntry& SoundDashboardEntry = CastEntry(InData);
+						if (SoundDashboardEntry.Children.Num() > 0)
+						{
+							return FText();
+						}
 
 						const ::Audio::TCircularAudioBuffer<FDataPoint>& PitchDataPoints = SoundDashboardEntry.PitchDataPoints;
 						return FText::AsNumber(GetLastEntryArrayValue(PitchDataPoints), FSlateStyle::Get().GetPitchFloatFormat());
diff --git a/Engine/Plugins/AudioInsights/Source/AudioInsights/Private/Views/TreeDashboardViewFactory.cpp b/Engine/Plugins/AudioInsights/Source/AudioInsights/Private/Views/TreeDashboardViewFactory.cpp
--- a/Engine/Plugins/AudioInsights/Source/AudioInsights/Private/Views/TreeDashboardViewFactory.cpp
+++ b/Engine/Plugins/AudioInsights/Source/AudioInsights/Private/Views/TreeDashboardViewFactory.cpp
@@ -129,6 +129,12 @@ namespace UE::Audio::Insights
 						if (InParent.IsValid() && !InParent->Children.IsEmpty())
 						{
 							OutChildren = InParent->Children;
+
+							if (FilteredEntriesListView.IsValid() && InParent->ShouldInitExpandChildren())
+							{
+								FilteredEntriesListView->SetItemExpansion(InParent, true);
+								InParent->ResetShouldInitExpandChildren();
+							}
 						}
 					})
 				]
diff --git a/Engine/Plugins/AudioInsights/Source/AudioInsights/Public/AudioInsightsDataSource.h b/Engine/Plugins/AudioInsights/Source/AudioInsights/Public/AudioInsightsDataSource.h
--- a/Engine/Plugins/AudioInsights/Source/AudioInsights/Public/AudioInsightsDataSource.h
+++ b/Engine/Plugins/AudioInsights/Source/AudioInsights/Public/AudioInsightsDataSource.h
@@ -20,6 +20,8 @@ namespace UE::Audio::Insights
 		virtual ~IDashboardDataTreeViewEntry() = default;
 
 		virtual bool IsValid() const = 0;
+		virtual bool ShouldInitExpandChildren() const = 0;
+		virtual void ResetShouldInitExpandChildren() = 0;
 
 		TArray<TSharedPtr<IDashboardDataTreeViewEntry>> Children;
 	};
diff --git a/Engine/Plugins/AudioInsights/Source/AudioInsights/Public/Messages/MixerSourceTraceMessages.h b/Engine/Plugins/AudioInsights/Source/AudioInsights/Public/Messages/MixerSourceTraceMessages.h
--- a/Engine/Plugins/AudioInsights/Source/AudioInsights/Public/Messages/MixerSourceTraceMessages.h
+++ b/Engine/Plugins/AudioInsights/Source/AudioInsights/Public/Messages/MixerSourceTraceMessages.h
@@ -50,25 +50,39 @@ namespace UE::Audio::Insights
 		uint64 ComponentId = TNumericLimits<uint64>::Max();
 	};
 
-#define DEFINE_MIXERSOURCE_PARAM_MESSAGE(ClassName, ParamName, Type, Default)			\
-	struct ClassName : public FMixerSourceMessageBase									\
-	{																					\
-		ClassName() = default;															\
-		ClassName(const Trace::IAnalyzer::FOnEventContext& InContext)					\
-			: FMixerSourceMessageBase(InContext)										\
-		{																				\
-			const Trace::IAnalyzer::FEventData& EventData = InContext.EventData;		\
-			ParamName = EventData.GetValue<Type>(#ParamName);							\
-		}																				\
-		Type ParamName = Default;														\
+	struct FMixerSourceParameterMessage : public FMixerSourceMessageBase
+	{
+		FMixerSourceParameterMessage() = default;
+
+		FMixerSourceParameterMessage(const Trace::IAnalyzer::FOnEventContext& InContext)
+			: FMixerSourceMessageBase(InContext)
+		{
+			ActiveSoundPlayOrder = InContext.EventData.GetValue<uint32>("ActiveSoundPlayOrder");
+		}
+
+		uint32 ActiveSoundPlayOrder;
+	};
+
+#define DEFINE_MIXERSOURCE_PARAM_MESSAGE(ClassName, SuperClassName, ParamName, Type, Default)	\
+	struct ClassName : public SuperClassName													\
+	{																							\
+		ClassName() = default;																	\
+		ClassName(const Trace::IAnalyzer::FOnEventContext& InContext)							\
+			: SuperClassName(InContext)															\
+		{																						\
+			const Trace::IAnalyzer::FEventData& EventData = InContext.EventData;				\
+			ParamName = EventData.GetValue<Type>(#ParamName);									\
+		}																						\
+		Type ParamName = Default;																\
 	};
 
-	DEFINE_MIXERSOURCE_PARAM_MESSAGE(FMixerSourceDistanceAttenuationMessage, DistanceAttenuation, float, 0.0f)
-	DEFINE_MIXERSOURCE_PARAM_MESSAGE(FMixerSourceEnvelopeMessage, Envelope, float, 0.0f)
-	DEFINE_MIXERSOURCE_PARAM_MESSAGE(FMixerSourceHPFFreqMessage, HPFFrequency, float, MIN_FILTER_FREQUENCY)
-	DEFINE_MIXERSOURCE_PARAM_MESSAGE(FMixerSourceLPFFreqMessage, LPFFrequency, float, MAX_FILTER_FREQUENCY)
-	DEFINE_MIXERSOURCE_PARAM_MESSAGE(FMixerSourcePitchMessage, Pitch, float, 1.0f)
-	DEFINE_MIXERSOURCE_PARAM_MESSAGE(FMixerSourceVolumeMessage, Volume, float, 1.0f)
+
+	DEFINE_MIXERSOURCE_PARAM_MESSAGE(FMixerSourceDistanceAttenuationMessage, FMixerSourceMessageBase, DistanceAttenuation, float, 0.0f)
+	DEFINE_MIXERSOURCE_PARAM_MESSAGE(FMixerSourceHPFFreqMessage, FMixerSourceMessageBase, HPFFrequency, float, MIN_FILTER_FREQUENCY)
+	DEFINE_MIXERSOURCE_PARAM_MESSAGE(FMixerSourceLPFFreqMessage, FMixerSourceMessageBase, LPFFrequency, float, MAX_FILTER_FREQUENCY)
+	DEFINE_MIXERSOURCE_PARAM_MESSAGE(FMixerSourceEnvelopeMessage, FMixerSourceParameterMessage, Envelope, float, 0.0f)
+	DEFINE_MIXERSOURCE_PARAM_MESSAGE(FMixerSourcePitchMessage, FMixerSourceParameterMessage, Pitch, float, 1.0f)
+	DEFINE_MIXERSOURCE_PARAM_MESSAGE(FMixerSourceVolumeMessage, FMixerSourceParameterMessage, Volume, float, 1.0f)
 #undef DEFINE_MIXERSOURCE_PARAM_MESSAGE
 
 	using FDataPoint = TPair<double, float>; // (Timestamp, Value)
diff --git a/Engine/Plugins/AudioInsights/Source/AudioInsights/Public/Messages/SoundTraceMessages.h b/Engine/Plugins/AudioInsights/Source/AudioInsights/Public/Messages/SoundTraceMessages.h
--- a/Engine/Plugins/AudioInsights/Source/AudioInsights/Public/Messages/SoundTraceMessages.h
+++ b/Engine/Plugins/AudioInsights/Source/AudioInsights/Public/Messages/SoundTraceMessages.h
@@ -4,6 +4,7 @@
 #include "AudioDefines.h"
 #include "DSP/Dsp.h"
 #include "Messages/AnalyzerMessageQueue.h"
+#include "Messages/MixerSourceTraceMessages.h"
 #include "Trace/Analyzer.h"
 #include "UObject/SoftObjectPath.h"
 #include "Views/TreeDashboardViewFactory.h"
@@ -45,13 +46,26 @@ namespace UE::Audio::Insights
 		ESoundDashboardEntryType EntryType;
 	};
 
+	struct FSoundWaveStartMessage : public FSoundStartMessage
+	{
+		FSoundWaveStartMessage() = default;
+
+		FSoundWaveStartMessage(const Trace::IAnalyzer::FOnEventContext& InContext)
+			: FSoundStartMessage(InContext)
+		{
+			ActiveSoundPlayOrder = InContext.EventData.GetValue<uint32>("ActiveSoundPlayOrder");
+		}
+
+		uint32 ActiveSoundPlayOrder;
+	};
+
 	using FSoundStopMessage = FSoundMessageBase;
 
-	struct FSoundPriorityMessage : public FSoundMessageBase
+	struct FSoundPriorityMessage : public FMixerSourceParameterMessage
 	{
 		FSoundPriorityMessage() = default;
 		FSoundPriorityMessage(const UE::Trace::IAnalyzer::FOnEventContext& InContext)
-			: FSoundMessageBase(InContext)
+			: FMixerSourceParameterMessage(InContext)
 		{
 			Priority = InContext.EventData.GetValue<float>("Priority");
 		}
@@ -59,11 +73,11 @@ namespace UE::Audio::Insights
 		float Priority = 0.0f;
 	};
 
-	struct FSoundDistanceMessage : public FSoundMessageBase
+	struct FSoundDistanceMessage : public FMixerSourceParameterMessage
 	{
 		FSoundDistanceMessage() = default;
 		FSoundDistanceMessage(const UE::Trace::IAnalyzer::FOnEventContext& InContext)
-			: FSoundMessageBase(InContext)
+			: FMixerSourceParameterMessage(InContext)
 		{
 			Distance = InContext.EventData.GetValue<float>("Distance");
 		}
@@ -71,53 +85,19 @@ namespace UE::Audio::Insights
 		float Distance = 0.0f;
 	};
 
-	struct FSoundAmplitudeMessage : public FSoundMessageBase
-	{
-		FSoundAmplitudeMessage() = default;
-		FSoundAmplitudeMessage(const UE::Trace::IAnalyzer::FOnEventContext& InContext)
-			: FSoundMessageBase(InContext)
-		{
-			Amplitude = InContext.EventData.GetValue<float>("Amplitude");
-		}
-
-		float Amplitude = 0.0f;
-	};
-
-	struct FSoundVolumeMessage : public FSoundMessageBase
-	{
-		FSoundVolumeMessage() = default;
-		FSoundVolumeMessage(const UE::Trace::IAnalyzer::FOnEventContext& InContext)
-			: FSoundMessageBase(InContext)
-		{
-			Volume = InContext.EventData.GetValue<float>("Volume");
-		}
-
-		float Volume = 1.0f;
-	};
-
-	struct FSoundPitchMessage : public FSoundMessageBase
-	{
-		FSoundPitchMessage() = default;
-		FSoundPitchMessage(const UE::Trace::IAnalyzer::FOnEventContext& InContext)
-			: FSoundMessageBase(InContext)
-		{
-			Pitch = InContext.EventData.GetValue<float>("Pitch");
-		}
-
-		float Pitch = 1.0f;
-	};
 
 
 	// Trace message queues
 	class FSoundMessages
 	{
 	public:
-		TAnalyzerMessageQueue<FSoundStartMessage> StartMessages { 0.1 };
+		TAnalyzerMessageQueue<FSoundStartMessage> ActiveSoundStartMessages { 0.1 };
+		TAnalyzerMessageQueue<FSoundWaveStartMessage> SoundWaveStartMessages{ 0.1 };
 		TAnalyzerMessageQueue<FSoundPriorityMessage> PriorityMessages{ 0.1 };
 		TAnalyzerMessageQueue<FSoundDistanceMessage> DistanceMessages { 0.1 };
-		TAnalyzerMessageQueue<FSoundAmplitudeMessage> AmplitudeMessages { 0.1 };
-		TAnalyzerMessageQueue<FSoundVolumeMessage> VolumeMessages { 0.1 };
-		TAnalyzerMessageQueue<FSoundPitchMessage> PitchMessages { 0.1 };
+		TAnalyzerMessageQueue<FMixerSourceEnvelopeMessage> AmplitudeMessages { 0.1 };
+		TAnalyzerMessageQueue<FMixerSourceVolumeMessage> VolumeMessages { 0.1 };
+		TAnalyzerMessageQueue<FMixerSourcePitchMessage> PitchMessages { 0.1 };
 		TAnalyzerMessageQueue<FSoundStopMessage> StopMessages { 0.1 };
 	};
 
@@ -126,6 +106,7 @@ namespace UE::Audio::Insights
 	{
 		FSoundSessionCachedMessages(TraceServices::IAnalysisSession& InSession)
 			: StartCachedMessages(InSession.GetLinearAllocator(), 16384)
+			, SoundWaveStartCachedMessages(InSession.GetLinearAllocator(), 16384)
 			, PriorityCachedMessages(InSession.GetLinearAllocator(), 16384)
 			, DistanceCachedMessages(InSession.GetLinearAllocator(), 16384)
 			, AmplitudeCachedMessages(InSession.GetLinearAllocator(), 16384)
@@ -137,11 +118,12 @@ namespace UE::Audio::Insights
 		}
 
 		TraceServices::TPagedArray<FSoundStartMessage> StartCachedMessages;
+		TraceServices::TPagedArray<FSoundWaveStartMessage> SoundWaveStartCachedMessages;
 		TraceServices::TPagedArray<FSoundPriorityMessage> PriorityCachedMessages;
 		TraceServices::TPagedArray<FSoundDistanceMessage> DistanceCachedMessages;
-		TraceServices::TPagedArray<FSoundAmplitudeMessage> AmplitudeCachedMessages;
-		TraceServices::TPagedArray<FSoundVolumeMessage> VolumeCachedMessages;
-		TraceServices::TPagedArray<FSoundPitchMessage> PitchCachedMessages;
+		TraceServices::TPagedArray<FMixerSourceEnvelopeMessage> AmplitudeCachedMessages;
+		TraceServices::TPagedArray<FMixerSourceVolumeMessage> VolumeCachedMessages;
+		TraceServices::TPagedArray<FMixerSourcePitchMessage> PitchCachedMessages;
 		TraceServices::TPagedArray<FSoundStopMessage> StopCachedMessages;
 	};
 #endif // !WITH_EDITOR
@@ -167,11 +149,15 @@ namespace UE::Audio::Insights
 
 		virtual bool IsValid() const override { return PlayOrder != INDEX_NONE; }
 
+		virtual bool ShouldInitExpandChildren() const override { return bShouldForceExpandChildren; }
+		virtual void ResetShouldInitExpandChildren() override { bShouldForceExpandChildren = false; }
+
 		::Audio::FDeviceId DeviceId = INDEX_NONE;
 		uint32 PlayOrder = INDEX_NONE;
 		double Timestamp = 0.0;
 		FString Name;
 		ESoundDashboardEntryType EntryType = ESoundDashboardEntryType::None;
+		bool bShouldForceExpandChildren = true;
 
 		::Audio::TCircularAudioBuffer<FDataPoint> PriorityDataPoints;
 		::Audio::TCircularAudioBuffer<FDataPoint> DistanceDataPoints;
diff --git a/Engine/Source/Runtime/AudioMixer/Private/AudioMixerSource.cpp b/Engine/Source/Runtime/AudioMixer/Private/AudioMixerSource.cpp
--- a/Engine/Source/Runtime/AudioMixer/Private/AudioMixerSource.cpp
+++ b/Engine/Source/Runtime/AudioMixer/Private/AudioMixerSource.cpp
@@ -406,6 +406,7 @@ namespace Audio
 			InitParams.bEnableBaseSubmix = WaveInstance->bEnableBaseSubmix && !bForceAudioLinkOnAllSourcesCVAr;
 			InitParams.bEnableSubmixSends = WaveInstance->bEnableSubmixSends;
 			InitParams.PlayOrder = WaveInstance->GetPlayOrder();
+			InitParams.ActiveSoundPlayOrder = WaveInstance->ActiveSound != nullptr ? WaveInstance->ActiveSound->GetPlayOrder() : INDEX_NONE;
 			bPreviousBusEnablement = WaveInstance->bEnableBusSends;
 			DynamicBusSendInfos.Reset();
 
diff --git a/Engine/Source/Runtime/AudioMixer/Private/AudioMixerSourceManager.cpp b/Engine/Source/Runtime/AudioMixer/Private/AudioMixerSourceManager.cpp
--- a/Engine/Source/Runtime/AudioMixer/Private/AudioMixerSourceManager.cpp
+++ b/Engine/Source/Runtime/AudioMixer/Private/AudioMixerSourceManager.cpp
@@ -225,6 +225,7 @@ UE_TRACE_EVENT_BEGIN(Audio, MixerSourceVolume)
 	UE_TRACE_EVENT_FIELD(uint32, DeviceId)
 	UE_TRACE_EVENT_FIELD(uint64, Timestamp)
 	UE_TRACE_EVENT_FIELD(uint32, PlayOrder)
+	UE_TRACE_EVENT_FIELD(uint32, ActiveSoundPlayOrder)
 	UE_TRACE_EVENT_FIELD(float, Volume)
 UE_TRACE_EVENT_END()
 
@@ -239,6 +240,7 @@ UE_TRACE_EVENT_BEGIN(Audio, MixerSourcePitch)
 	UE_TRACE_EVENT_FIELD(uint32, DeviceId)
 	UE_TRACE_EVENT_FIELD(uint64, Timestamp)
 	UE_TRACE_EVENT_FIELD(uint32, PlayOrder)
+	UE_TRACE_EVENT_FIELD(uint32, ActiveSoundPlayOrder)
 	UE_TRACE_EVENT_FIELD(float, Pitch)
 UE_TRACE_EVENT_END()
 
@@ -254,6 +256,7 @@ UE_TRACE_EVENT_BEGIN(Audio, MixerSourceEnvelope)
 	UE_TRACE_EVENT_FIELD(uint32, DeviceId)
 	UE_TRACE_EVENT_FIELD(uint64, Timestamp)
 	UE_TRACE_EVENT_FIELD(uint32, PlayOrder)
+	UE_TRACE_EVENT_FIELD(uint32, ActiveSoundPlayOrder)
 	UE_TRACE_EVENT_FIELD(float, Envelope)
 UE_TRACE_EVENT_END()
 
@@ -1037,6 +1040,7 @@ namespace Audio
 			SourceInfo.bIsExternalSend = InitParams.bIsExternalSend;
 			SourceInfo.bIsVorbis = InitParams.bIsVorbis;
 			SourceInfo.PlayOrder = InitParams.PlayOrder;
+			SourceInfo.ActiveSoundPlayOrder = InitParams.ActiveSoundPlayOrder;
 			SourceInfo.AudioComponentID = InitParams.AudioComponentID;
 			SourceInfo.bIsSoundfield = InitParams.bIsSoundfield;
 
@@ -2353,7 +2357,8 @@ namespace Audio
 					<< MixerSourcePitch.DeviceId(MixerDevice->DeviceID)
 					<< MixerSourcePitch.Timestamp(FPlatformTime::Cycles64())
 					<< MixerSourcePitch.PlayOrder(SourceInfo.PlayOrder)
-					<< MixerSourcePitch.Pitch(TargetPitch);
+					<< MixerSourcePitch.ActiveSoundPlayOrder(SourceInfo.ActiveSoundPlayOrder)
+					<< MixerSourcePitch.Pitch(FinalPitch);
 			}
 #endif // UE_AUDIO_PROFILERTRACE_ENABLED
 
@@ -3000,6 +3005,7 @@ namespace Audio
 						<< MixerSourceVolume.DeviceId(MixerDevice->DeviceID)
 						<< MixerSourceVolume.Timestamp(FPlatformTime::Cycles64())
 						<< MixerSourceVolume.PlayOrder(SourceInfo.PlayOrder)
+						<< MixerSourceVolume.ActiveSoundPlayOrder(SourceInfo.ActiveSoundPlayOrder)
 						<< MixerSourceVolume.Volume(VolumeDestination);
 
 					UE_TRACE_LOG(Audio, MixerSourceDistanceAttenuation, AudioMixerChannel)
@@ -3151,6 +3157,7 @@ namespace Audio
 						<< MixerSourceEnvelope.DeviceId(MixerDevice->DeviceID)
 						<< MixerSourceEnvelope.Timestamp(FPlatformTime::Cycles64())
 						<< MixerSourceEnvelope.PlayOrder(SourceInfo.PlayOrder)
+						<< MixerSourceEnvelope.ActiveSoundPlayOrder(SourceInfo.ActiveSoundPlayOrder)
 						<< MixerSourceEnvelope.Envelope(SourceInfo.SourceEnvelopeValue);
 				}
 #endif // UE_AUDIO_PROFILERTRACE_ENABLED
diff --git a/Engine/Source/Runtime/AudioMixer/Private/AudioMixerSourceManager.h b/Engine/Source/Runtime/AudioMixer/Private/AudioMixerSourceManager.h
--- a/Engine/Source/Runtime/AudioMixer/Private/AudioMixerSourceManager.h
+++ b/Engine/Source/Runtime/AudioMixer/Private/AudioMixerSourceManager.h
@@ -168,6 +168,7 @@ namespace Audio
 		bool bShouldSourceBufferListenerZeroBuffer = false;
 
 		uint32 PlayOrder = INDEX_NONE;
+		uint32 ActiveSoundPlayOrder = INDEX_NONE;
 	};
 
 	struct FSourceManagerInitParams
@@ -565,6 +566,7 @@ namespace Audio
 			int32 NumInputFrames;
 
 			uint32 PlayOrder;
+			uint32 ActiveSoundPlayOrder;
 
 			// ID for associated Audio Component if there is one, 0 otherwise
 			uint64 AudioComponentID;
diff --git a/Engine/Source/Runtime/Engine/Private/SoundWave.cpp b/Engine/Source/Runtime/Engine/Private/SoundWave.cpp
--- a/Engine/Source/Runtime/Engine/Private/SoundWave.cpp
+++ b/Engine/Source/Runtime/Engine/Private/SoundWave.cpp
@@ -106,6 +106,34 @@ FAutoConsoleVariableRef CvarAllowReverbForMultichannelSources(
 	TEXT("0: Disable, >0: Enable"),
 	ECVF_Default);
 
+#if UE_AUDIO_PROFILERTRACE_ENABLED
+UE_TRACE_EVENT_BEGIN(Audio, SoundWaveStart)
+	UE_TRACE_EVENT_FIELD(uint32, DeviceId)
+	UE_TRACE_EVENT_FIELD(uint64, Timestamp)
+	UE_TRACE_EVENT_FIELD(uint32, PlayOrder)
+	UE_TRACE_EVENT_FIELD(uint32, ActiveSoundPlayOrder)
+	UE_TRACE_EVENT_FIELD(UE::Trace::WideString, Name)
+	UE_TRACE_EVENT_FIELD(UE::Trace::WideString, SoundClassName)
+	UE_TRACE_EVENT_FIELD(UE::Trace::WideString, SoundSuperClassName)
+UE_TRACE_EVENT_END()
+
+UE_TRACE_EVENT_BEGIN(Audio, SoundPriority)
+	UE_TRACE_EVENT_FIELD(uint32, DeviceId)
+	UE_TRACE_EVENT_FIELD(uint64, Timestamp)
+	UE_TRACE_EVENT_FIELD(uint32, PlayOrder)
+	UE_TRACE_EVENT_FIELD(uint32, ActiveSoundPlayOrder)
+	UE_TRACE_EVENT_FIELD(float, Priority)
+UE_TRACE_EVENT_END()
+
+UE_TRACE_EVENT_BEGIN(Audio, SoundDistance)
+	UE_TRACE_EVENT_FIELD(uint32, DeviceId)
+	UE_TRACE_EVENT_FIELD(uint64, Timestamp)
+	UE_TRACE_EVENT_FIELD(uint32, PlayOrder)
+	UE_TRACE_EVENT_FIELD(uint32, ActiveSoundPlayOrder)
+	UE_TRACE_EVENT_FIELD(float, Distance)
+UE_TRACE_EVENT_END()
+#endif // UE_AUDIO_PROFILERTRACE_ENABLED
+
 inline FArchive& operator<<(FArchive& Ar, FSoundWaveCuePoint& CuePoint)
 {
 	FSoundWaveCuePoint::StaticStruct()->SerializeItem(Ar, &CuePoint, nullptr);
@@ -3747,7 +3775,8 @@ void USoundWave::Parse(FAudioDevice* AudioDevice, const UPTRINT NodeWaveInstance
 	WaveInstance->bShouldUseAudioLink = ParseParams.bEnableSendToAudioLink;
 	WaveInstance->AudioLinkSettingsOverride = ParseParams.AudioLinkSettingsOverride;
 
-	if (WaveInstance->IsPlaying())
+	const bool bIsWavePlaying = WaveInstance->IsPlaying();
+	if (bIsWavePlaying)
 	{
 		WaveInstances.Add(WaveInstance);
 		ActiveSound.bFinished = false;
@@ -3769,6 +3798,49 @@ void USoundWave::Parse(FAudioDevice* AudioDevice, const UPTRINT NodeWaveInstance
 		}
 	}
 
+#if UE_AUDIO_PROFILERTRACE_ENABLED
+	const bool bChannelEnabled = UE_TRACE_CHANNELEXPR_IS_ENABLED(AudioMixerChannel);
+	const bool bWaveIsActive = bIsWavePlaying || !ActiveSound.bFinished;
+
+	if (bChannelEnabled && bWaveIsActive)
+	{
+		if (bIsNewWave && WaveInstance->SoundClass != nullptr)
+		{
+			const TObjectPtr<UObject> SoundWaveObjectPtr = FSoftObjectPath(WaveInstance->WaveData->GetPathName()).ResolveObject();
+			if (SoundWaveObjectPtr)
+			{
+				const FString SoundWaveClassName = SoundWaveObjectPtr->GetClass()->GetName();
+
+				const TObjectPtr<UClass> SoundWaveSuperClass = SoundWaveObjectPtr->GetClass()->GetSuperClass();
+				const FString SoundWaveSuperClassName = SoundWaveSuperClass ? SoundWaveSuperClass->GetName() : FString();
+
+				UE_TRACE_LOG(Audio, SoundWaveStart, AudioMixerChannel)
+					<< SoundWaveStart.DeviceId(AudioDevice->DeviceID)
+					<< SoundWaveStart.Timestamp(FPlatformTime::Cycles64())
+					<< SoundWaveStart.PlayOrder(WaveInstance->GetPlayOrder())
+					<< SoundWaveStart.ActiveSoundPlayOrder(ActiveSound.GetPlayOrder())
+					<< SoundWaveStart.Name(*SoundWaveObjectPtr->GetPathName())
+					<< SoundWaveStart.SoundClassName(*SoundWaveClassName)
+					<< SoundWaveStart.SoundSuperClassName(*SoundWaveSuperClassName);
+			}
+		}
+
+		UE_TRACE_LOG(Audio, SoundPriority, AudioMixerChannel)
+			<< SoundPriority.DeviceId(AudioDevice->DeviceID)
+			<< SoundPriority.Timestamp(FPlatformTime::Cycles64())
+			<< SoundPriority.PlayOrder(WaveInstance->GetPlayOrder())
+			<< SoundPriority.ActiveSoundPlayOrder(ActiveSound.GetPlayOrder())
+			<< SoundPriority.Priority(WaveInstance->Priority);
+
+		UE_TRACE_LOG(Audio, SoundDistance, AudioMixerChannel)
+			<< SoundDistance.DeviceId(AudioDevice->DeviceID)
+			<< SoundDistance.Timestamp(FPlatformTime::Cycles64())
+			<< SoundDistance.PlayOrder(WaveInstance->GetPlayOrder())
+			<< SoundDistance.ActiveSoundPlayOrder(ActiveSound.GetPlayOrder())
+			<< SoundDistance.Distance(WaveInstance->ListenerToSoundDistance);
+	}
+#endif // UE_AUDIO_PROFILERTRACE_ENABLED
+
 #if !NO_LOGGING
 	// Sanity check
 	if (NumChannels > 2 && !WaveInstance->bIsAmbisonics && WaveInstance->GetUseSpatialization() && !WaveInstance->bReportedSpatializationWarning)
