Home Screen - Sound timeout
- Inactive sounds will now grey out rather than be immediately removed from the dashboard.
- If they are still inactive after 3 seconds, they are removed from the dashboard.
- If they reactivate, they return to being updated.
- Bonus: Made category headers bold in the sounds dashboard to better match original mockup.
- Bonus: Added recursive expansion/collapse in sounds dashboard tree view

diff --git a/Engine/Plugins/AudioInsights/Source/AudioInsights/Private/Providers/SoundTraceProvider.cpp b/Engine/Plugins/AudioInsights/Source/AudioInsights/Private/Providers/SoundTraceProvider.cpp
--- a/Engine/Plugins/AudioInsights/Source/AudioInsights/Private/Providers/SoundTraceProvider.cpp
+++ b/Engine/Plugins/AudioInsights/Source/AudioInsights/Private/Providers/SoundTraceProvider.cpp
@@ -314,6 +314,7 @@ namespace UE::Audio::Insights
 		};
 
 		ActiveSoundToEntryKeysMap.Empty();
+		EntriesTimingOut.Empty();
 
 		return new FSoundTraceAnalyzer(AsShared(), InSession);
 	}
@@ -513,13 +514,52 @@ namespace UE::Audio::Insights
 
 			if (OutActiveSoundEntry && OutActiveSoundEntry->IsValid() && (*OutActiveSoundEntry)->Timestamp < Msg.Timestamp)
 			{
-				RemoveActiveSoundEntry(Msg, *OutActiveSoundEntry);
+				static constexpr double DashboardTimeoutTime = 3.0;
+				(*OutActiveSoundEntry)->TimeoutTimestamp = FPlatformTime::Seconds() + DashboardTimeoutTime;
+				EntriesTimingOut.Push(SoundMessageIDs{ Msg.DeviceId, Msg.PlayOrder });
 			}
 		});
 
 		return true;
 	}
 
+	bool FSoundTraceProvider::ProcessManuallyUpdatedEntries()
+	{
+		// Timeout old dashboard entries
+		const double CurrentTimeStamp = FPlatformTime::Seconds();
+
+		int32 NumToTrim = 0;
+		for (const SoundMessageIDs SoundMessageIDs : EntriesTimingOut)
+		{
+			TSharedPtr<FSoundDashboardEntry> ActiveSoundEntry = nullptr;
+			GetActiveSoundEntryFromIDs(SoundMessageIDs.PlayOrder, SoundMessageIDs.DeviceId, ActiveSoundEntry);
+
+			if (!ActiveSoundEntry.IsValid())
+			{
+				NumToTrim++;
+				continue;
+			}
+
+			if (ActiveSoundEntry->TimeoutTimestamp <= CurrentTimeStamp)
+			{
+				RemoveActiveSoundEntry(SoundMessageIDs.DeviceId, ActiveSoundEntry);
+				NumToTrim++;
+			}
+			else
+			{
+				break;
+			}
+		}
+
+		if (NumToTrim >= 0)
+		{
+			EntriesTimingOut.RemoveAt(0, NumToTrim, EAllowShrinking::No);
+			return true;
+		}
+
+		return false;
+	}
+
 	void FSoundTraceProvider::InitLambdas()
 	{
 		GetOrCreateActiveSoundEntry = [this](const FSoundStartMessage& Msg)
@@ -578,6 +618,15 @@ namespace UE::Audio::Insights
 				else
 				{
 					TSharedPtr<FSoundDashboardEntry> NewEntry = StaticCastSharedPtr<FSoundDashboardEntry>(*SoundEntry);
+					if (NewEntry->TimeoutTimestamp != INVALID_TIMEOUT)
+					{
+						NewEntry->TimeoutTimestamp = INVALID_TIMEOUT;
+
+						EntriesTimingOut.RemoveAll([&Msg](const SoundMessageIDs& SoundMessageIDs)
+						{
+							return Msg.DeviceId == SoundMessageIDs.DeviceId && Msg.PlayOrder == SoundMessageIDs.PlayOrder;
+						});
+					}
 					ReturnedSoundEntry = &NewEntry;
 				}
 
@@ -606,7 +655,7 @@ namespace UE::Audio::Insights
 		FSoundTraceProviderPrivate::FindActiveSoundEntryInCategory(CategoryEntry, SoundEntryKeys->AssetName, ActiveSoundPlayOrder, OutActiveSoundEntry);
 	}
 
-	void FSoundTraceProvider::RemoveActiveSoundEntry(const FSoundStopMessage& Msg, TSharedPtr<FSoundDashboardEntry> OutActiveSoundEntry)
+	void FSoundTraceProvider::RemoveActiveSoundEntry(const ::Audio::FDeviceId DeviceId, TSharedPtr<FSoundDashboardEntry> OutActiveSoundEntry)
 	{
 		if (!OutActiveSoundEntry.IsValid())
 		{
@@ -621,7 +670,7 @@ namespace UE::Audio::Insights
 			return;
 		}
 
-		TSharedPtr<FSoundDashboardEntry>* CategoryEntry = FindDeviceEntry(Msg.DeviceId, SoundEntryKeys->EntryType);
+		TSharedPtr<FSoundDashboardEntry>* CategoryEntry = FindDeviceEntry(DeviceId, SoundEntryKeys->EntryType);
 
 		// First remove the entry from the asset name
 		TSharedPtr<FSoundDashboardEntry> AssetEntry = nullptr;
@@ -639,7 +688,7 @@ namespace UE::Audio::Insights
 
 			if ((*CategoryEntry)->Children.Num() == 0)
 			{
-				RemoveDeviceEntry(Msg.DeviceId, OutActiveSoundEntry->EntryType);
+				RemoveDeviceEntry(DeviceId, OutActiveSoundEntry->EntryType);
 			}
 		}
 
@@ -735,9 +784,9 @@ namespace UE::Audio::Insights
 
 			TSharedPtr<FSoundDashboardEntry>* ActiveSoundEntry = GetActiveSoundEntryFromStopMessage(StopCachedMessage);
 
-			if (ActiveSoundEntry && (*ActiveSoundEntry)->Timestamp < StopCachedMessage.Timestamp)
+			if (ActiveSoundEntry && ActiveSoundEntry->IsValid() && (*ActiveSoundEntry)->Timestamp < StopCachedMessage.Timestamp)
 			{
-				RemoveActiveSoundEntry(StopCachedMessage, *ActiveSoundEntry);
+				RemoveActiveSoundEntry(StopCachedMessage.DeviceId, *ActiveSoundEntry);
 			}
 		}
 		
diff --git a/Engine/Plugins/AudioInsights/Source/AudioInsights/Private/Providers/SoundTraceProvider.h b/Engine/Plugins/AudioInsights/Source/AudioInsights/Private/Providers/SoundTraceProvider.h
--- a/Engine/Plugins/AudioInsights/Source/AudioInsights/Private/Providers/SoundTraceProvider.h
+++ b/Engine/Plugins/AudioInsights/Source/AudioInsights/Private/Providers/SoundTraceProvider.h
@@ -34,11 +34,12 @@ namespace UE::Audio::Insights
 #endif // !WITH_EDITOR
 
 		virtual bool ProcessMessages() override;
+		virtual bool ProcessManuallyUpdatedEntries() override;
 
 	private:
 		void InitLambdas();
 		void GetActiveSoundEntryFromIDs(const uint32 PlayOrderID, const ::Audio::FDeviceId DeviceId, TSharedPtr<FSoundDashboardEntry>& OutSoundEntry);
-		void RemoveActiveSoundEntry(const FSoundStopMessage& Msg, TSharedPtr<FSoundDashboardEntry> OutEntry);
+		void RemoveActiveSoundEntry(const ::Audio::FDeviceId DeviceId, TSharedPtr<FSoundDashboardEntry> OutEntry);
 
 #if !WITH_EDITOR
 		virtual void OnTimingViewTimeMarkerChanged(double TimeMarker) override;
@@ -50,11 +51,18 @@ namespace UE::Audio::Insights
 
 		struct SoundEntryKeys
 		{
-			ESoundDashboardEntryType EntryType;
-			FString AssetName;
+			const ESoundDashboardEntryType EntryType;
+			const FString AssetName;
+		};
+
+		struct SoundMessageIDs
+		{
+			const ::Audio::FDeviceId DeviceId;
+			const uint32 PlayOrder;
 		};
 
 		TMap<uint32, SoundEntryKeys> ActiveSoundToEntryKeysMap;
+		TArray<SoundMessageIDs> EntriesTimingOut;
 
 		TFunction<TSharedPtr<FSoundDashboardEntry>*(const FSoundStartMessage& Msg)> GetOrCreateActiveSoundEntry;
 	};
diff --git a/Engine/Plugins/AudioInsights/Source/AudioInsights/Private/Views/SoundDashboardViewFactory.cpp b/Engine/Plugins/AudioInsights/Source/AudioInsights/Private/Views/SoundDashboardViewFactory.cpp
--- a/Engine/Plugins/AudioInsights/Source/AudioInsights/Private/Views/SoundDashboardViewFactory.cpp
+++ b/Engine/Plugins/AudioInsights/Source/AudioInsights/Private/Views/SoundDashboardViewFactory.cpp
@@ -333,6 +333,8 @@ namespace UE::Audio::Insights
 				}
 			};
 
+			const FSlateFontInfo BoldFontInfo = FCoreStyle::GetDefaultFontStyle("Bold", 10);
+
 			return SNew(SHorizontalBox)
 				// Tree expander arrow
 				+ SHorizontalBox::Slot()
@@ -354,7 +356,10 @@ namespace UE::Audio::Insights
 				.Padding(IconName != NAME_None ? 10.0f : 0.0f, 2.0f, 0.0f, 2.0f)
 				.AutoWidth()
 				[
-					SNew(STextBlock).Text(InValueText)
+					SNew(STextBlock)
+						.Font(BoldFontInfo)
+						.Text(InValueText)
+						.ColorAndOpacity(FSlateColor::UseForeground())
 				]
 				// Number of children text
 				+ SHorizontalBox::Slot()
@@ -408,6 +413,7 @@ namespace UE::Audio::Insights
 
 		const FText ValueText  = ColumnData.GetDisplayValue.IsSet() ? ColumnData.GetDisplayValue(InRowData.Get()) : FText::GetEmpty();
 		const FName& ValueIcon = ColumnData.GetIconName.IsSet() ? ColumnData.GetIconName(InRowData.Get()) : NAME_None;
+		const FSlateColor ValueTextColor = ColumnData.GetTextColorValue.IsSet() ? ColumnData.GetTextColorValue(InRowData.Get()) : FSlateColor::UseForeground();
 
 		if (ValueText.IsEmpty() && ValueIcon.IsNone())
 		{
@@ -441,6 +447,13 @@ namespace UE::Audio::Insights
 
 					return IconName != NAME_None ? FSlateStyle::Get().GetBrush(IconName) : nullptr;
 				})
+				.ColorAndOpacity_Lambda([this, InRowData, InColumn]()
+				{
+					const FColumnData& ColumnData = GetColumns()[InColumn];
+					const FSlateColor TextColor = ColumnData.GetTextColorValue.IsSet() ? ColumnData.GetTextColorValue(InRowData.Get()) : FSlateColor::UseForeground();
+
+					return TextColor;
+				})
 			]
 			// Text (optional)
 			+ SHorizontalBox::Slot()
@@ -455,6 +468,13 @@ namespace UE::Audio::Insights
 
 					return ValueText;
 				})
+				.ColorAndOpacity_Lambda([this, InRowData, InColumn]()
+				{
+					const FColumnData& ColumnData = GetColumns()[InColumn];
+					const FSlateColor TextColor = ColumnData.GetTextColorValue.IsSet() ? ColumnData.GetTextColorValue(InRowData.Get()) : FSlateColor::UseForeground();
+
+					return TextColor;
+				})
 			]
 			// Number of children text (if is not leaf item)
 			+ SHorizontalBox::Slot()
@@ -667,6 +687,12 @@ namespace UE::Audio::Insights
 						}
 
 						return NAME_None;
+					},
+					.GetTextColorValue = [](const IDashboardDataTreeViewEntry& InData) -> FSlateColor
+					{
+						const FSoundDashboardEntry& SoundDashboardEntry = CastEntry(InData);
+						return SoundDashboardEntry.TimeoutTimestamp == INVALID_TIMEOUT ? FSlateColor::UseForeground()
+																					   : FSlateColor(FSlateStyle::Get().GetColor("SoundDashboard.TimingOutTextColor"));
 					}
 				}
 			},
@@ -695,6 +721,12 @@ namespace UE::Audio::Insights
 						{
 							return FText::AsNumber(GetLastEntryArrayValue(PriorityDataPoints), FSlateStyle::Get().GetAmpFloatFormat());
 						}
+					},
+					.GetTextColorValue = [](const IDashboardDataTreeViewEntry& InData) -> FSlateColor
+					{
+						const FSoundDashboardEntry& SoundDashboardEntry = CastEntry(InData);
+						return SoundDashboardEntry.TimeoutTimestamp == INVALID_TIMEOUT ? FSlateColor::UseForeground()
+																					   : FSlateColor(FSlateStyle::Get().GetColor("SoundDashboard.TimingOutTextColor"));
 					}
 				}
 			},
@@ -711,6 +743,12 @@ namespace UE::Audio::Insights
 
 						const ::Audio::TCircularAudioBuffer<FDataPoint>& DistanceDataPoints = SoundDashboardEntry.DistanceDataPoints;
 						return FText::AsNumber(GetLastEntryArrayValue(DistanceDataPoints), FSlateStyle::Get().GetDefaultFloatFormat());
+					},
+					.GetTextColorValue = [](const IDashboardDataTreeViewEntry& InData) -> FSlateColor
+					{
+						const FSoundDashboardEntry& SoundDashboardEntry = CastEntry(InData);
+						return SoundDashboardEntry.TimeoutTimestamp == INVALID_TIMEOUT ? FSlateColor::UseForeground()
+																					   : FSlateColor(FSlateStyle::Get().GetColor("SoundDashboard.TimingOutTextColor"));
 					}
 				}
 			},
@@ -727,6 +765,12 @@ namespace UE::Audio::Insights
 
 						const ::Audio::TCircularAudioBuffer<FDataPoint>& AmplitudeDataPoints = SoundDashboardEntry.AmplitudeDataPoints;
 						return FText::AsNumber(GetLastEntryArrayValue(AmplitudeDataPoints), FSlateStyle::Get().GetAmpFloatFormat());
+					},
+					.GetTextColorValue = [](const IDashboardDataTreeViewEntry& InData) -> FSlateColor
+					{
+						const FSoundDashboardEntry& SoundDashboardEntry = CastEntry(InData);
+						return SoundDashboardEntry.TimeoutTimestamp == INVALID_TIMEOUT ? FSlateColor::UseForeground()
+																					   : FSlateColor(FSlateStyle::Get().GetColor("SoundDashboard.TimingOutTextColor"));
 					}
 				}
 			},
@@ -743,6 +787,12 @@ namespace UE::Audio::Insights
 
 						const ::Audio::TCircularAudioBuffer<FDataPoint>& VolumeDataPoints = SoundDashboardEntry.VolumeDataPoints;
 						return FText::AsNumber(GetLastEntryArrayValue(VolumeDataPoints), FSlateStyle::Get().GetAmpFloatFormat());
+					},
+					.GetTextColorValue = [](const IDashboardDataTreeViewEntry& InData) -> FSlateColor
+					{
+						const FSoundDashboardEntry& SoundDashboardEntry = CastEntry(InData);
+						return SoundDashboardEntry.TimeoutTimestamp == INVALID_TIMEOUT ? FSlateColor::UseForeground()
+																					   : FSlateColor(FSlateStyle::Get().GetColor("SoundDashboard.TimingOutTextColor"));
 					}
 				}
 			},
@@ -759,6 +809,12 @@ namespace UE::Audio::Insights
 
 						const ::Audio::TCircularAudioBuffer<FDataPoint>& PitchDataPoints = SoundDashboardEntry.PitchDataPoints;
 						return FText::AsNumber(GetLastEntryArrayValue(PitchDataPoints), FSlateStyle::Get().GetPitchFloatFormat());
+					},
+					.GetTextColorValue = [](const IDashboardDataTreeViewEntry& InData) -> FSlateColor
+					{
+						const FSoundDashboardEntry& SoundDashboardEntry = CastEntry(InData);
+						return SoundDashboardEntry.TimeoutTimestamp == INVALID_TIMEOUT ? FSlateColor::UseForeground()
+																					   : FSlateColor(FSlateStyle::Get().GetColor("SoundDashboard.TimingOutTextColor"));
 					}
 				}
 			}
diff --git a/Engine/Plugins/AudioInsights/Source/AudioInsights/Private/Views/TreeDashboardViewFactory.cpp b/Engine/Plugins/AudioInsights/Source/AudioInsights/Private/Views/TreeDashboardViewFactory.cpp
--- a/Engine/Plugins/AudioInsights/Source/AudioInsights/Private/Views/TreeDashboardViewFactory.cpp
+++ b/Engine/Plugins/AudioInsights/Source/AudioInsights/Private/Views/TreeDashboardViewFactory.cpp
@@ -137,6 +137,7 @@ namespace UE::Audio::Insights
 							}
 						}
 					})
+					.OnSetExpansionRecursive(this, &FTraceTreeDashboardViewFactory::HandleRecursiveExpansion)
 				]
 			];
 		}
@@ -261,6 +262,16 @@ namespace UE::Audio::Insights
 		RequestSort();
 	}
 
+	void FTraceTreeDashboardViewFactory::HandleRecursiveExpansion(TSharedPtr<IDashboardDataTreeViewEntry> Item, bool bIsItemExpanded)
+	{
+		for (TSharedPtr<IDashboardDataTreeViewEntry> ChildItem : Item->Children)
+		{
+			HandleRecursiveExpansion(ChildItem, bIsItemExpanded);
+		}
+
+		FilteredEntriesListView->SetItemExpansion(Item, bIsItemExpanded);
+	}
+
 	void FTraceTreeDashboardViewFactory::Tick(float InElapsed)
 	{
 		for (const TSharedPtr<FTraceProviderBase>& Provider : Providers)
@@ -270,7 +281,13 @@ namespace UE::Audio::Insights
 				continue;
 			}
 
-			if (const uint64* FoundCurrentUpdateId = UpdateIds.Find(Provider->GetName()))
+			if (Provider->ShouldForceUpdate())
+			{
+				Provider->ResetShouldForceUpdate();
+
+				UpdateFilterReason = EProcessReason::EntriesUpdated;
+			}
+			else if (const uint64* FoundCurrentUpdateId = UpdateIds.Find(Provider->GetName()))
 			{
 				if (*FoundCurrentUpdateId != Provider->GetLastUpdateId())
 				{
diff --git a/Engine/Plugins/AudioInsights/Source/AudioInsights/Public/AudioInsightsStyle.h b/Engine/Plugins/AudioInsights/Source/AudioInsights/Public/AudioInsightsStyle.h
--- a/Engine/Plugins/AudioInsights/Source/AudioInsights/Public/AudioInsightsStyle.h
+++ b/Engine/Plugins/AudioInsights/Source/AudioInsights/Public/AudioInsightsStyle.h
@@ -174,6 +174,9 @@ namespace UE::Audio::Insights
 			Set("SoundDashboard.PinnedColor", FLinearColor(0.9f, 0.9f, 0.9f));
 			Set("SoundDashboard.HiddenColor", FLinearColor(0.4f, 0.4f, 0.4f));
 
+			// Text colors
+			Set("SoundDashboard.TimingOutTextColor", FLinearColor(1.0f, 1.0f, 1.0f, 0.25f));
+
 			// Table style
 			Set("SoundDashboard.TableViewRow", FTableRowStyle(FAppStyle::Get().GetWidgetStyle<FTableRowStyle>("TableView.Row"))
 				.SetOddRowBackgroundBrush(FSlateColorBrush(FStyleColors::Recessed))
diff --git a/Engine/Plugins/AudioInsights/Source/AudioInsights/Public/AudioInsightsTraceProviderBase.h b/Engine/Plugins/AudioInsights/Source/AudioInsights/Public/AudioInsightsTraceProviderBase.h
--- a/Engine/Plugins/AudioInsights/Source/AudioInsights/Public/AudioInsightsTraceProviderBase.h
+++ b/Engine/Plugins/AudioInsights/Source/AudioInsights/Public/AudioInsightsTraceProviderBase.h
@@ -43,6 +43,11 @@ namespace UE::Audio::Insights
 			return true;
 		}
 
+		virtual bool ProcessManuallyUpdatedEntries()
+		{ 
+			return false;
+		};
+
 		uint64 GetLastUpdateId() const
 		{
 			return LastUpdateId;
@@ -53,6 +58,16 @@ namespace UE::Audio::Insights
 			return GetLastMessageId() == LastUpdateId;
 		}
 
+		bool ShouldForceUpdate() const
+		{
+			return bForceUpdate;
+		}
+
+		void ResetShouldForceUpdate()
+		{
+			bForceUpdate = false;
+		}
+
 #if !WITH_EDITOR
 		virtual void InitSessionCachedMessages(TraceServices::IAnalysisSession& InSession) {}
 		virtual void OnTimingViewTimeMarkerChanged(double TimeMarker) { ++LastMessageId; };
@@ -87,6 +102,7 @@ namespace UE::Audio::Insights
 		}
 
 		uint64 LastUpdateId = 0;
+		bool bForceUpdate = false;
 
 	private:
 		std::atomic<uint64> LastMessageId { 0 };
@@ -115,6 +131,13 @@ namespace UE::Audio::Insights
 					ProcessMessages();
 				}
 				LastUpdateId = GetLastMessageId();
+
+				const bool bShouldForceUpdate = ProcessManuallyUpdatedEntries();
+				if (bShouldForceUpdate)
+				{
+					bForceUpdate = true;
+				}
+
 				return true;
 			});
 
diff --git a/Engine/Plugins/AudioInsights/Source/AudioInsights/Public/Messages/SoundTraceMessages.h b/Engine/Plugins/AudioInsights/Source/AudioInsights/Public/Messages/SoundTraceMessages.h
--- a/Engine/Plugins/AudioInsights/Source/AudioInsights/Public/Messages/SoundTraceMessages.h
+++ b/Engine/Plugins/AudioInsights/Source/AudioInsights/Public/Messages/SoundTraceMessages.h
@@ -16,6 +16,8 @@
 
 namespace UE::Audio::Insights
 {
+	static constexpr double INVALID_TIMEOUT = -1.0;
+
 	enum class ESoundDashboardEntryType : uint8
 	{
 		None,
@@ -155,6 +157,7 @@ namespace UE::Audio::Insights
 		::Audio::FDeviceId DeviceId = INDEX_NONE;
 		uint32 PlayOrder = INDEX_NONE;
 		double Timestamp = 0.0;
+		double TimeoutTimestamp = INVALID_TIMEOUT;
 		FString Name;
 		ESoundDashboardEntryType EntryType = ESoundDashboardEntryType::None;
 		bool bShouldForceExpandChildren = true;
diff --git a/Engine/Plugins/AudioInsights/Source/AudioInsights/Public/Views/TreeDashboardViewFactory.h b/Engine/Plugins/AudioInsights/Source/AudioInsights/Public/Views/TreeDashboardViewFactory.h
--- a/Engine/Plugins/AudioInsights/Source/AudioInsights/Public/Views/TreeDashboardViewFactory.h
+++ b/Engine/Plugins/AudioInsights/Source/AudioInsights/Public/Views/TreeDashboardViewFactory.h
@@ -73,6 +73,7 @@ namespace UE::Audio::Insights
 		{
 			const TFunction<FText(const IDashboardDataTreeViewEntry&)> GetDisplayValue;
 			const TFunction<FName(const IDashboardDataTreeViewEntry&)> GetIconName;
+			const TFunction<FSlateColor(const IDashboardDataTreeViewEntry&)> GetTextColorValue;
 		};
 
 		virtual TSharedPtr<SWidget> GetFilterBarWidget();
@@ -172,6 +173,7 @@ namespace UE::Audio::Insights
 		EColumnSortMode::Type GetColumnSortMode(const FName InColumnId) const;
 		void RequestSort();
 		void OnColumnSortModeChanged(const EColumnSortPriority::Type InSortPriority, const FName& InColumnId, const EColumnSortMode::Type InSortMode);
+		void HandleRecursiveExpansion(TSharedPtr<IDashboardDataTreeViewEntry> Item, bool bIsItemExpanded);
 
 		TSharedPtr<SSearchBox> SearchBoxWidget;
 		FText SearchBoxFilterText;
