Home Screen - Removed Asset category
- Sounds are now nested in just their category
- Updated "number of children" logic on category headers

diff --git a/Engine/Plugins/AudioInsights/Source/AudioInsights/Private/Providers/SoundTraceProvider.cpp b/Engine/Plugins/AudioInsights/Source/AudioInsights/Private/Providers/SoundTraceProvider.cpp
--- a/Engine/Plugins/AudioInsights/Source/AudioInsights/Private/Providers/SoundTraceProvider.cpp
+++ b/Engine/Plugins/AudioInsights/Source/AudioInsights/Private/Providers/SoundTraceProvider.cpp
@@ -76,33 +76,9 @@ namespace UE::Audio::Insights
 			OutDataPoints.Push({ Msg.Timestamp, DataPoint });
 		}
 
-		//////////////////////////////////////////////////////////////////////////////////////////////////////////////
-		// Finds an Asset entry inside a Category entry, which will hold a collection of Active Sound instances that are all the same Asset
-		void GetAssetEntry(const TSharedPtr<FSoundDashboardEntry>* CategoryEntry, const FString& AssetName, TSharedPtr<FSoundDashboardEntry>& OutAssetEntry)
-		{
-			OutAssetEntry = nullptr;
-			if (CategoryEntry == nullptr || !CategoryEntry->IsValid())
-			{
-				return;
-			}
-
-			TSharedPtr<IDashboardDataTreeViewEntry>* AssetEntry = CategoryEntry->Get()->Children.FindByPredicate([&AssetName](TSharedPtr<IDashboardDataTreeViewEntry> ChildEntry)
-			{
-				FSoundDashboardEntry& NewEntry = CastEntry(*ChildEntry.Get());
-				return NewEntry.Name == AssetName;
-			});
-
-			if (AssetEntry == nullptr || !AssetEntry->IsValid())
-			{
-				return;
-			}
-
-			OutAssetEntry = StaticCastSharedPtr<FSoundDashboardEntry>(*AssetEntry);
-		}
-
 		//////////////////////////////////////////////////////////////////////////////////////////////////////////////
 		// Finds an Active Sound entry inside a Category entry
-		void FindActiveSoundEntryInCategory(const TSharedPtr<FSoundDashboardEntry>* CategoryEntry, const FString& AssetName, const uint32 ActiveSoundPlayOrder, TSharedPtr<FSoundDashboardEntry>& OutActiveSoundEntry)
+		void FindActiveSoundEntryInCategory(const TSharedPtr<FSoundDashboardEntry>* CategoryEntry, const uint32 ActiveSoundPlayOrder, TSharedPtr<FSoundDashboardEntry>& OutActiveSoundEntry)
 		{
 			OutActiveSoundEntry = nullptr;
 			if (CategoryEntry == nullptr || !CategoryEntry->IsValid())
@@ -110,18 +86,7 @@ namespace UE::Audio::Insights
 				return;
 			}
 
-			TSharedPtr<IDashboardDataTreeViewEntry>* AssetEntry = CategoryEntry->Get()->Children.FindByPredicate([&AssetName](TSharedPtr<IDashboardDataTreeViewEntry> ChildEntry)
-			{
-				FSoundDashboardEntry& NewEntry = CastEntry(*ChildEntry.Get());
-				return NewEntry.Name == AssetName;
-			});
-
-			if (AssetEntry == nullptr || !AssetEntry->IsValid())
-			{
-				return;
-			}
-
-			TSharedPtr<IDashboardDataTreeViewEntry>* ActiveSoundEntry = AssetEntry->Get()->Children.FindByPredicate([ActiveSoundPlayOrder](TSharedPtr<IDashboardDataTreeViewEntry> ChildEntry)
+			TSharedPtr<IDashboardDataTreeViewEntry>* ActiveSoundEntry = CategoryEntry->Get()->Children.FindByPredicate([ActiveSoundPlayOrder](TSharedPtr<IDashboardDataTreeViewEntry> ChildEntry)
 			{
 				FSoundDashboardEntry& NewEntry = CastEntry(*ChildEntry.Get());
 				return NewEntry.PlayOrder == ActiveSoundPlayOrder;
@@ -515,7 +480,15 @@ namespace UE::Audio::Insights
 			if (OutActiveSoundEntry && OutActiveSoundEntry->IsValid() && (*OutActiveSoundEntry)->Timestamp < Msg.Timestamp)
 			{
 				static constexpr double DashboardTimeoutTime = 3.0;
-				(*OutActiveSoundEntry)->TimeoutTimestamp = FPlatformTime::Seconds() + DashboardTimeoutTime;
+				const double TimeoutTimestamp = FPlatformTime::Seconds() + DashboardTimeoutTime;
+
+				(*OutActiveSoundEntry)->TimeoutTimestamp = TimeoutTimestamp;
+				for (TSharedPtr<IDashboardDataTreeViewEntry> ChildEntry : (*OutActiveSoundEntry)->Children)
+				{
+					FSoundDashboardEntry& SoundWaveEntry = FSoundTraceProviderPrivate::CastEntry(*ChildEntry.Get());
+					SoundWaveEntry.TimeoutTimestamp = TimeoutTimestamp;
+				}
+
 				EntriesTimingOut.Push(SoundMessageIDs{ Msg.DeviceId, Msg.PlayOrder });
 			}
 		});
@@ -576,28 +549,7 @@ namespace UE::Audio::Insights
 					CategoryEntry->Name = FSoundTraceProviderPrivate::GetEntryTypeDisplayName(Msg.EntryType);
 				}
 
-				TSharedPtr<IDashboardDataTreeViewEntry>* AssetEntry = CategoryEntry->Children.FindByPredicate([&Msg](TSharedPtr<IDashboardDataTreeViewEntry> ChildEntry)
-				{
-					FSoundDashboardEntry& NewEntry = FSoundTraceProviderPrivate::CastEntry(*ChildEntry.Get());
-					return NewEntry.Name == Msg.Name;
-				});
-
-				TSharedPtr<FSoundDashboardEntry> NewAssetEntry = nullptr;
-				if (AssetEntry == nullptr || !AssetEntry->IsValid())
-				{
-					NewAssetEntry = MakeShared<FSoundDashboardEntry>();
-					NewAssetEntry->DeviceId = Msg.DeviceId;
-					NewAssetEntry->EntryType = Msg.EntryType;
-					NewAssetEntry->Name = Msg.Name;
-
-					CategoryEntry->Children.Add(NewAssetEntry);
-				}
-				else
-				{
-					NewAssetEntry = StaticCastSharedPtr<FSoundDashboardEntry>(*AssetEntry);
-				}
-
-				TSharedPtr<IDashboardDataTreeViewEntry>* SoundEntry = NewAssetEntry->Children.FindByPredicate([&Msg](TSharedPtr<IDashboardDataTreeViewEntry> ChildEntry)
+				TSharedPtr<IDashboardDataTreeViewEntry>* SoundEntry = CategoryEntry->Children.FindByPredicate([&Msg](TSharedPtr<IDashboardDataTreeViewEntry> ChildEntry)
 				{
 					FSoundDashboardEntry& NewEntry = FSoundTraceProviderPrivate::CastEntry(*ChildEntry.Get());
 					return NewEntry.PlayOrder == Msg.PlayOrder;
@@ -609,9 +561,9 @@ namespace UE::Audio::Insights
 					NewEntry->DeviceId = Msg.DeviceId;
 					NewEntry->PlayOrder = Msg.PlayOrder;
 
-					NewAssetEntry->Children.Add(NewEntry);
+					CategoryEntry->Children.Add(NewEntry);
 
-					ActiveSoundToEntryKeysMap.Add(Msg.PlayOrder, SoundEntryKeys{ Msg.EntryType, Msg.Name });
+					ActiveSoundToEntryKeysMap.Add(Msg.PlayOrder, SoundEntryKeys{ Msg.EntryType });
 
 					ReturnedSoundEntry = &NewEntry;
 				}
@@ -626,6 +578,13 @@ namespace UE::Audio::Insights
 						{
 							return Msg.DeviceId == SoundMessageIDs.DeviceId && Msg.PlayOrder == SoundMessageIDs.PlayOrder;
 						});
+
+						// If this entry has SoundWaves, they will replay with new play order IDs, so we need to clear the old ones out
+						NewEntry->Children.RemoveAll([](TSharedPtr<IDashboardDataTreeViewEntry> ChildEntry) 
+						{
+							FSoundDashboardEntry& SoundWaveEntry = FSoundTraceProviderPrivate::CastEntry(*ChildEntry.Get());
+							return SoundWaveEntry.TimeoutTimestamp != INVALID_TIMEOUT;
+						});
 					}
 					ReturnedSoundEntry = &NewEntry;
 				}
@@ -652,7 +611,7 @@ namespace UE::Audio::Insights
 
 		TSharedPtr<FSoundDashboardEntry>* CategoryEntry = FindDeviceEntry(DeviceId, SoundEntryKeys->EntryType);
 
-		FSoundTraceProviderPrivate::FindActiveSoundEntryInCategory(CategoryEntry, SoundEntryKeys->AssetName, ActiveSoundPlayOrder, OutActiveSoundEntry);
+		FSoundTraceProviderPrivate::FindActiveSoundEntryInCategory(CategoryEntry, ActiveSoundPlayOrder, OutActiveSoundEntry);
 	}
 
 	void FSoundTraceProvider::RemoveActiveSoundEntry(const ::Audio::FDeviceId DeviceId, TSharedPtr<FSoundDashboardEntry> OutActiveSoundEntry)
@@ -671,25 +630,16 @@ namespace UE::Audio::Insights
 		}
 
 		TSharedPtr<FSoundDashboardEntry>* CategoryEntry = FindDeviceEntry(DeviceId, SoundEntryKeys->EntryType);
-
-		// First remove the entry from the asset name
-		TSharedPtr<FSoundDashboardEntry> AssetEntry = nullptr;
-		FSoundTraceProviderPrivate::GetAssetEntry(CategoryEntry, SoundEntryKeys->AssetName, AssetEntry);
-		if (AssetEntry == nullptr || !AssetEntry.IsValid())
+		if (CategoryEntry == nullptr || !CategoryEntry->IsValid())
 		{
 			return;
 		}
 
-		AssetEntry->Children.Remove(OutActiveSoundEntry);
+		(*CategoryEntry)->Children.Remove(OutActiveSoundEntry);
 
-		if (AssetEntry->Children.Num() == 0)
+		if ((*CategoryEntry)->Children.Num() == 0)
 		{
-			(*CategoryEntry)->Children.Remove(AssetEntry);
-
-			if ((*CategoryEntry)->Children.Num() == 0)
-			{
-				RemoveDeviceEntry(DeviceId, OutActiveSoundEntry->EntryType);
-			}
+			RemoveDeviceEntry(DeviceId, OutActiveSoundEntry->EntryType);
 		}
 
 		ActiveSoundToEntryKeysMap.Remove(ActiveSoundPlayOrder);
@@ -753,14 +703,14 @@ namespace UE::Audio::Insights
 
 			if (!FSoundTraceProviderPrivate::EntryTypeHasSoundWaveEntries(SoundWaveStartCachedMessage.EntryType))
 			{
-				return;
+				continue;
 			}
 
 			GetActiveSoundEntryFromIDs(SoundWaveStartCachedMessage.ActiveSoundPlayOrder, SoundWaveStartCachedMessage.DeviceId, ParentSoundEntryReturn);
 
 			if (!ParentSoundEntryReturn->IsValid())
 			{
-				return;
+				continue;
 			}
 
 			TSharedPtr<FSoundDashboardEntry> ChildEntry = MakeShared<FSoundDashboardEntry>();
diff --git a/Engine/Plugins/AudioInsights/Source/AudioInsights/Private/Providers/SoundTraceProvider.h b/Engine/Plugins/AudioInsights/Source/AudioInsights/Private/Providers/SoundTraceProvider.h
--- a/Engine/Plugins/AudioInsights/Source/AudioInsights/Private/Providers/SoundTraceProvider.h
+++ b/Engine/Plugins/AudioInsights/Source/AudioInsights/Private/Providers/SoundTraceProvider.h
@@ -10,9 +10,8 @@
  *	
  *	"Sounds Dashboard" displays active audio in the following tree structure:
  *	- Category											(i.e. The EntryType of the sound asset, e.g. SoundCue, MetaSound etc)
- *		- Asset											(i.e. Sound instances from the same asset are grouped here)
- *			- Active Sound								(i.e. The Active Sound instance playing)
- *				- SoundWave Instances (optional)		(i.e. If an entry type can have multiple SoundWaves playing simultaneously (e.g. SoundCue)
+ *		- Active Sound									(i.e. The Active Sound instance playing)
+ *			- SoundWave Instances (optional)			(i.e. If an entry type can have multiple SoundWaves playing simultaneously (e.g. SoundCue)
  *														, individual Soundwaves are shown here)
  */
 namespace UE::Audio::Insights
@@ -52,7 +51,6 @@ namespace UE::Audio::Insights
 		struct SoundEntryKeys
 		{
 			const ESoundDashboardEntryType EntryType;
-			const FString AssetName;
 		};
 
 		struct SoundMessageIDs
diff --git a/Engine/Plugins/AudioInsights/Source/AudioInsights/Private/Views/SoundDashboardViewFactory.cpp b/Engine/Plugins/AudioInsights/Source/AudioInsights/Private/Views/SoundDashboardViewFactory.cpp
--- a/Engine/Plugins/AudioInsights/Source/AudioInsights/Private/Views/SoundDashboardViewFactory.cpp
+++ b/Engine/Plugins/AudioInsights/Source/AudioInsights/Private/Views/SoundDashboardViewFactory.cpp
@@ -373,7 +373,18 @@ namespace UE::Audio::Insights
 						const uint32 TotalNumChildren = Algo::Accumulate(InRowData->Children, 0,
 							[](uint32 Accum, TSharedPtr<IDashboardDataTreeViewEntry> InChild)
 							{
-								return Accum + (InChild ? InChild->Children.Num() : 0);
+								if (InChild.IsValid())
+								{
+									const FSoundDashboardEntry& SoundDashboardEntry = FSoundDashboardViewFactoryPrivate::CastEntry(*InChild.Get());
+
+									if (SoundDashboardEntry.TimeoutTimestamp == INVALID_TIMEOUT)
+									{
+										const int32 NumNestedChildren = InChild->Children.Num();
+										return Accum + (NumNestedChildren > 0 ? NumNestedChildren : 1);
+									}
+								}
+
+								return Accum;
 							});
 
 						return FText::FromString("(" + FString::FromInt(TotalNumChildren) + ")");
