[Audio Insights] Fixed text conflict from LOCTEXT macro

diff --git a/Engine/Plugins/AudioInsights/Source/AudioInsights/Private/Providers/SoundTraceProvider.cpp b/Engine/Plugins/AudioInsights/Source/AudioInsights/Private/Providers/SoundTraceProvider.cpp
--- a/Engine/Plugins/AudioInsights/Source/AudioInsights/Private/Providers/SoundTraceProvider.cpp
+++ b/Engine/Plugins/AudioInsights/Source/AudioInsights/Private/Providers/SoundTraceProvider.cpp
@@ -22,7 +22,7 @@ namespace UE::Audio::Insights
 		const FString SoundWaveDisplayName			{ LOCTEXT("AudioDashboard_Sounds_EntryType_SoundWaves", "Sound Waves").ToString() };
 		const FString ProceduralSourceDisplayName	{ LOCTEXT("AudioDashboard_Sounds_EntryType_ProceduralSources", "Procedural Sources").ToString() };
 		const FString SoundCueTemplateName			{ LOCTEXT("AudioDashboard_Sounds_EntryType_SoundCueTemplate", "Sound Cue Templates").ToString() };
-		const FString UncategorizedSoundName		{ LOCTEXT("AudioDashboard_Sounds_EntryType_SoundCueTemplate", "Others").ToString() };
+		const FString UncategorizedSoundName		{ LOCTEXT("AudioDashboard_Sounds_EntryType_UncategorizedSound", "Others").ToString() };
 
 		FSoundDashboardEntry& CastEntry(IDashboardDataTreeViewEntry& InData)
 		{
