[Audio Insights] Sound Dashboard: Added logic for category filtering + adjusted logic for text filtering since it only was taking into account the root nodes + fixed small typo (using SoundCue instead of SoundCueTemplate)

diff --git a/Engine/Plugins/AudioInsights/Source/AudioInsights/Private/Views/SoundDashboardViewFactory.cpp b/Engine/Plugins/AudioInsights/Source/AudioInsights/Private/Views/SoundDashboardViewFactory.cpp
--- a/Engine/Plugins/AudioInsights/Source/AudioInsights/Private/Views/SoundDashboardViewFactory.cpp
+++ b/Engine/Plugins/AudioInsights/Source/AudioInsights/Private/Views/SoundDashboardViewFactory.cpp
@@ -5,6 +5,7 @@
 #include "AudioInsightsModule.h"
 #include "AudioInsightsStyle.h"
 #include "Messages/SoundTraceMessages.h"
+#include "Misc/EnumClassFlags.h"
 #include "Providers/SoundTraceProvider.h"
 #include "SoundDashboardCommands.h"
 #include "Views/SAudioFilterBar.h"
@@ -27,6 +28,11 @@ namespace UE::Audio::Insights
 			return static_cast<const FSoundDashboardEntry&>(InData);
 		};
 
+		FSoundDashboardEntry& CastEntry(IDashboardDataTreeViewEntry& InData)
+		{
+			return static_cast<FSoundDashboardEntry&>(InData);
+		};
+
 		float GetLastEntryArrayValue(const ::Audio::TCircularAudioBuffer<FDataPoint>& InDataPoints)
 		{
 			if (InDataPoints.Num() > 0)
@@ -37,6 +43,59 @@ namespace UE::Audio::Insights
 
 			return 0.0f;
 		};
+
+		bool SetFilteredVisibility(IDashboardDataTreeViewEntry& InEntry, const FString& InFilterString)
+		{
+			FSoundDashboardEntry& SoundEntry = FSoundDashboardViewFactoryPrivate::CastEntry(InEntry);
+
+			bool bEntryMatchesTextFilter = SoundEntry.GetDisplayName().ToString().Contains(InFilterString);
+
+			if (bEntryMatchesTextFilter)
+			{
+				SoundEntry.bIsVisible = true;
+			}
+			else
+			{
+				bool bChildMatchesTextFilter = false;
+
+				for (const TSharedPtr<IDashboardDataTreeViewEntry>& SoundEntryChild : SoundEntry.Children)
+				{
+					if (SoundEntryChild.IsValid() && SetFilteredVisibility(*SoundEntryChild, InFilterString))
+					{
+						bChildMatchesTextFilter = true;
+						break;
+					}
+				}
+
+				SoundEntry.bIsVisible = bChildMatchesTextFilter;
+			}
+
+			return SoundEntry.bIsVisible;
+		}
+
+		void ResetVisibility(IDashboardDataTreeViewEntry& InEntry)
+		{
+			FSoundDashboardEntry& SoundEntry = FSoundDashboardViewFactoryPrivate::CastEntry(InEntry);
+
+			SoundEntry.bIsVisible = true;
+
+			for (const TSharedPtr<IDashboardDataTreeViewEntry>& SoundEntryChild : SoundEntry.Children)
+			{
+				if (SoundEntryChild.IsValid())
+				{
+					ResetVisibility(*SoundEntryChild);
+				}
+			}
+		}
+
+		constexpr ESoundDashboardFilterFlags AllFilterFlags = 
+			ESoundDashboardFilterFlags::MetaSound        |
+			ESoundDashboardFilterFlags::SoundCue         |
+			ESoundDashboardFilterFlags::ProceduralSource |
+			ESoundDashboardFilterFlags::SoundWave        |
+			ESoundDashboardFilterFlags::SoundCueTemplate |
+			ESoundDashboardFilterFlags::Pinned           |
+			ESoundDashboardFilterFlags::Hidden;
 	} // namespace FSoundDashboardViewFactoryPrivate
 
 
@@ -213,7 +272,7 @@ namespace UE::Audio::Insights
 				FilterCategory),
 
 				MakeShared<FSoundDashboardFilter>(
-				ESoundDashboardFilterFlags::SoundCue,
+				ESoundDashboardFilterFlags::SoundCueTemplate,
 				"SoundCueTemplate",
 				LOCTEXT("SoundDashboard_Filter_SoundCueTemplateNameText", "Sound Cue Template"),
 				"AudioInsights.Icon.SoundDashboard.SoundCue",
@@ -503,14 +562,77 @@ namespace UE::Audio::Insights
 			];
 	}
 
+	TSharedRef<ITableRow> FSoundDashboardViewFactory::OnGenerateRow(TSharedPtr<IDashboardDataTreeViewEntry> Item, const TSharedRef<STableViewBase>& OwnerTable)
+	{
+		return SNew(SRowWidget, OwnerTable, Item, AsShared())
+			.Visibility_Lambda([Item]()
+			{
+				const FSoundDashboardEntry& SoundCategoryEntry = FSoundDashboardViewFactoryPrivate::CastEntry(*Item);
+
+				return SoundCategoryEntry.bIsVisible ? EVisibility::Visible : EVisibility::Hidden; 
+			});
+	}
+
 	void FSoundDashboardViewFactory::ProcessEntries(FTraceTreeDashboardViewFactory::EProcessReason InReason)
 	{
-		FTraceTreeDashboardViewFactory::FilterEntries<FSoundTraceProvider>([&FilterString = GetSearchFilterText().ToString()](const IDashboardDataTreeViewEntry& Entry)
+		using namespace FSoundDashboardViewFactoryPrivate;
+
+		// Filter by category
+		FTraceTreeDashboardViewFactory::FilterEntries<FSoundTraceProvider>([this](const IDashboardDataTreeViewEntry& InEntry)
 		{
-			const FSoundDashboardEntry& SoundsEntry = FSoundDashboardViewFactoryPrivate::CastEntry(Entry);
+			const bool bAllCategoryFiltersEnabled = SelectedFilterFlags == ESoundDashboardFilterFlags::None || EnumHasAllFlags(SelectedFilterFlags, FSoundDashboardViewFactoryPrivate::AllFilterFlags);
+			if (bAllCategoryFiltersEnabled)
+			{
+				return true;
+			}
+
+			const FSoundDashboardEntry& SoundCategoryEntry = FSoundDashboardViewFactoryPrivate::CastEntry(InEntry);
+
+			switch (SoundCategoryEntry.EntryType)
+			{
+				case ESoundDashboardEntryType::MetaSound:
+					return EnumHasAnyFlags(SelectedFilterFlags, ESoundDashboardFilterFlags::MetaSound);
+				case ESoundDashboardEntryType::SoundCue:
+					return EnumHasAnyFlags(SelectedFilterFlags, ESoundDashboardFilterFlags::SoundCue);
+				case ESoundDashboardEntryType::ProceduralSource:
+					return EnumHasAnyFlags(SelectedFilterFlags, ESoundDashboardFilterFlags::ProceduralSource);
+				case ESoundDashboardEntryType::SoundWave:
+					return EnumHasAnyFlags(SelectedFilterFlags, ESoundDashboardFilterFlags::SoundWave);
+				case ESoundDashboardEntryType::SoundCueTemplate:
+					return EnumHasAnyFlags(SelectedFilterFlags, ESoundDashboardFilterFlags::SoundCueTemplate);
+				default:
+					break;
+			}
 
-			return !SoundsEntry.GetDisplayName().ToString().Contains(FilterString);
+			return false;
 		});
+
+		// Filter by text
+		const FString FilterString = GetSearchFilterText().ToString();
+		const bool bFilterHasText  = !FilterString.IsEmpty();
+
+		for (const TSharedPtr<IDashboardDataTreeViewEntry>& CategoryEntry : DataViewEntries)
+		{
+			if (!CategoryEntry.IsValid())
+			{
+				continue;
+			}
+
+			for (const TSharedPtr<IDashboardDataTreeViewEntry>& Entry : CategoryEntry->Children)
+			{
+				if (Entry.IsValid())
+				{
+					if (bFilterHasText)
+					{
+						SetFilteredVisibility(*Entry, FilterString);
+					}
+					else
+					{
+						ResetVisibility(*Entry);
+					}
+				}
+			}
+		}
 	}
 
 	TSharedPtr<SWidget> FSoundDashboardViewFactory::OnConstructContextMenu()
diff --git a/Engine/Plugins/AudioInsights/Source/AudioInsights/Private/Views/TreeDashboardViewFactory.cpp b/Engine/Plugins/AudioInsights/Source/AudioInsights/Private/Views/TreeDashboardViewFactory.cpp
--- a/Engine/Plugins/AudioInsights/Source/AudioInsights/Private/Views/TreeDashboardViewFactory.cpp
+++ b/Engine/Plugins/AudioInsights/Source/AudioInsights/Private/Views/TreeDashboardViewFactory.cpp
@@ -116,10 +116,7 @@ namespace UE::Audio::Insights
 					SAssignNew(FilteredEntriesListView, STreeView<TSharedPtr<IDashboardDataTreeViewEntry>>)
 					.HeaderRow(MakeHeaderRowWidget())
 					.TreeItemsSource(&DataViewEntries)
-					.OnGenerateRow_Lambda([this](TSharedPtr<IDashboardDataTreeViewEntry> Item, const TSharedRef<STableViewBase>& OwnerTable)
-					{
-						return SNew(SRowWidget, OwnerTable, Item, AsShared());
-					})
+					.OnGenerateRow(this, &FTraceTreeDashboardViewFactory::OnGenerateRow)
 					.OnContextMenuOpening(this, &FTraceTreeDashboardViewFactory::OnConstructContextMenu)
 					.SelectionMode(ESelectionMode::Multi)
 					.OnSelectionChanged(this, &FTraceTreeDashboardViewFactory::OnSelectionChanged)
@@ -203,6 +200,11 @@ namespace UE::Audio::Insights
 		return HeaderRowWidget.ToSharedRef();
 	}
 
+	TSharedRef<ITableRow> FTraceTreeDashboardViewFactory::OnGenerateRow(TSharedPtr<IDashboardDataTreeViewEntry> Item, const TSharedRef<STableViewBase>& OwnerTable)
+	{
+		return SNew(SRowWidget, OwnerTable, Item, AsShared());
+	}
+
 	TSharedPtr<SWidget> FTraceTreeDashboardViewFactory::OnConstructContextMenu()
 	{
 		// To be optionally implemented by derived classes
diff --git a/Engine/Plugins/AudioInsights/Source/AudioInsights/Public/Messages/SoundTraceMessages.h b/Engine/Plugins/AudioInsights/Source/AudioInsights/Public/Messages/SoundTraceMessages.h
--- a/Engine/Plugins/AudioInsights/Source/AudioInsights/Public/Messages/SoundTraceMessages.h
+++ b/Engine/Plugins/AudioInsights/Source/AudioInsights/Public/Messages/SoundTraceMessages.h
@@ -160,7 +160,9 @@ namespace UE::Audio::Insights
 		double TimeoutTimestamp = INVALID_TIMEOUT;
 		FString Name;
 		ESoundDashboardEntryType EntryType = ESoundDashboardEntryType::None;
+
 		bool bShouldForceExpandChildren = true;
+		bool bIsVisible = true;
 
 		::Audio::TCircularAudioBuffer<FDataPoint> PriorityDataPoints;
 		::Audio::TCircularAudioBuffer<FDataPoint> DistanceDataPoints;
diff --git a/Engine/Plugins/AudioInsights/Source/AudioInsights/Public/Views/SoundDashboardViewFactory.h b/Engine/Plugins/AudioInsights/Source/AudioInsights/Public/Views/SoundDashboardViewFactory.h
--- a/Engine/Plugins/AudioInsights/Source/AudioInsights/Public/Views/SoundDashboardViewFactory.h
+++ b/Engine/Plugins/AudioInsights/Source/AudioInsights/Public/Views/SoundDashboardViewFactory.h
@@ -76,6 +76,8 @@ namespace UE::Audio::Insights
 		TSharedRef<SWidget> GenerateWidgetForRootColumn(const TSharedRef<FTraceTreeDashboardViewFactory::SRowWidget>& InRowWidget, const TSharedRef<IDashboardDataTreeViewEntry>& InRowData, const FName& InColumn, const FText& InValueText);
 		virtual TSharedRef<SWidget> GenerateWidgetForColumn(const TSharedRef<FTraceTreeDashboardViewFactory::SRowWidget>& InRowWidget, const TSharedRef<IDashboardDataTreeViewEntry>& InRowData, const FName& InColumn) override;
 		
+		virtual TSharedRef<ITableRow> OnGenerateRow(TSharedPtr<IDashboardDataTreeViewEntry> Item, const TSharedRef<STableViewBase>& OwnerTable) override;
+
 		virtual TSharedPtr<SWidget> OnConstructContextMenu() override;
 
 		virtual const TMap<FName, FTraceTreeDashboardViewFactory::FHeaderRowColumnData>& GetHeaderRowColumns() const override;
diff --git a/Engine/Plugins/AudioInsights/Source/AudioInsights/Public/Views/TreeDashboardViewFactory.h b/Engine/Plugins/AudioInsights/Source/AudioInsights/Public/Views/TreeDashboardViewFactory.h
--- a/Engine/Plugins/AudioInsights/Source/AudioInsights/Public/Views/TreeDashboardViewFactory.h
+++ b/Engine/Plugins/AudioInsights/Source/AudioInsights/Public/Views/TreeDashboardViewFactory.h
@@ -83,6 +83,8 @@ namespace UE::Audio::Insights
 
 		virtual TSharedRef<SWidget> GenerateWidgetForColumn(const TSharedRef<FTraceTreeDashboardViewFactory::SRowWidget>& InRowWidget, const TSharedRef<IDashboardDataTreeViewEntry>& InRowData, const FName& Column) = 0;
 		
+		virtual TSharedRef<ITableRow> OnGenerateRow(TSharedPtr<IDashboardDataTreeViewEntry> Item, const TSharedRef<STableViewBase>& OwnerTable);
+
 		virtual TSharedPtr<SWidget> OnConstructContextMenu();
 
 		virtual void OnSelectionChanged(TSharedPtr<IDashboardDataTreeViewEntry> SelectedItem, ESelectInfo::Type SelectInfo);
@@ -106,7 +108,7 @@ namespace UE::Audio::Insights
 #endif // WITH_EDITOR
 
 		template<typename TableProviderType>
-		bool FilterEntries(TFunctionRef<bool(const IDashboardDataTreeViewEntry&)> IsFiltered)
+		bool FilterEntries(TFunctionRef<bool(const IDashboardDataTreeViewEntry&)> InPredicate)
 		{
 			const TSharedPtr<const TableProviderType> Provider = FindProvider<const TableProviderType>();
 			if (Provider.IsValid())
@@ -121,20 +123,13 @@ namespace UE::Audio::Insights
 						return StaticCastSharedPtr<IDashboardDataTreeViewEntry>(Pair.Value);
 					};
 
-					const auto FilterEntry = [this, &IsFiltered](const typename TableProviderType::FEntryPair& Pair)
+					const auto FilterEntry = [this, &InPredicate](const typename TableProviderType::FEntryPair& Pair)
 					{
-						return !IsFiltered((const IDashboardDataTreeViewEntry&)(*Pair.Value));
+						return InPredicate(*Pair.Value);
 					};
 
-					if (SearchBoxFilterText.IsEmpty())
-					{
-						Algo::Transform(*DeviceData, DataViewEntries, TransformEntry);
-					}
-					else
-					{
-						Algo::TransformIf(*DeviceData, DataViewEntries, FilterEntry, TransformEntry);
-					}
-					
+					Algo::TransformIf(*DeviceData, DataViewEntries, FilterEntry, TransformEntry);
+
 					// Sort list
 					RequestSort();
 
