[Audio Insights] Sound Dashboard: Misc. style adjustments

diff --git a/Engine/Plugins/AudioInsights/Source/AudioInsights/Private/Views/SoundDashboardViewFactory.cpp b/Engine/Plugins/AudioInsights/Source/AudioInsights/Private/Views/SoundDashboardViewFactory.cpp
--- a/Engine/Plugins/AudioInsights/Source/AudioInsights/Private/Views/SoundDashboardViewFactory.cpp
+++ b/Engine/Plugins/AudioInsights/Source/AudioInsights/Private/Views/SoundDashboardViewFactory.cpp
@@ -360,10 +360,10 @@ namespace UE::Audio::Insights
 	{
 		const FColumnData& ColumnData = GetColumns()[InColumn];
 
-		const FName IconName = ColumnData.GetIconName.IsSet() ? ColumnData.GetIconName(InRowData.Get()) : NAME_None;
-
 		if (InColumn == "Name")
 		{
+			const FName IconName = ColumnData.GetIconName.IsSet() ? ColumnData.GetIconName(InRowData.Get()) : NAME_None;
+
 			// Custom color for root item icons
 			const auto GetIconColor = [](const TSharedPtr<IDashboardDataTreeViewEntry>& InEntry)
 			{
@@ -392,8 +392,6 @@ namespace UE::Audio::Insights
 				}
 			};
 
-			const FSlateFontInfo BoldFontInfo = FCoreStyle::GetDefaultFontStyle("Bold", 10);
-
 			return SNew(SHorizontalBox)
 				// Tree expander arrow
 				+ SHorizontalBox::Slot()
@@ -416,9 +414,9 @@ namespace UE::Audio::Insights
 				.AutoWidth()
 				[
 					SNew(STextBlock)
-						.Font(BoldFontInfo)
-						.Text(InValueText)
-						.ColorAndOpacity(FSlateColor::UseForeground())
+					.Font(FCoreStyle::Get().GetFontStyle("BoldFont"))
+					.Text(InValueText)
+					.ColorAndOpacity(FSlateColor::UseForeground())
 				]
 				// Number of children text
 				+ SHorizontalBox::Slot()
@@ -483,7 +481,6 @@ namespace UE::Audio::Insights
 
 		const FText ValueText  = ColumnData.GetDisplayValue.IsSet() ? ColumnData.GetDisplayValue(InRowData.Get()) : FText::GetEmpty();
 		const FName& ValueIcon = ColumnData.GetIconName.IsSet() ? ColumnData.GetIconName(InRowData.Get()) : NAME_None;
-		const FSlateColor ValueTextColor = ColumnData.GetTextColorValue.IsSet() ? ColumnData.GetTextColorValue(InRowData.Get()) : FSlateColor::UseForeground();
 
 		if (ValueText.IsEmpty() && ValueIcon.IsNone())
 		{
@@ -676,8 +673,8 @@ namespace UE::Audio::Insights
 					.IconName         = "AudioInsights.Icon.SoundDashboard.Mute",
 					.bShowDisplayName = false,
 					.bDefaultHidden   = false,
-					.FillWidth        = 0.04f,
-					.Alignment        = EHorizontalAlignment::HAlign_Left
+					.FillWidth        = 0.05f,
+					.Alignment        = EHorizontalAlignment::HAlign_Center
 				}
 			},
 			{
@@ -687,8 +684,8 @@ namespace UE::Audio::Insights
 					.IconName         = "AudioInsights.Icon.SoundDashboard.Solo",
 					.bShowDisplayName = false,
 					.bDefaultHidden   = false,
-					.FillWidth        = 0.04f,
-					.Alignment        = EHorizontalAlignment::HAlign_Left
+					.FillWidth        = 0.05f,
+					.Alignment        = EHorizontalAlignment::HAlign_Center
 				}
 			},
 #endif
@@ -699,7 +696,7 @@ namespace UE::Audio::Insights
 					.IconName         = NAME_None,
 					.bShowDisplayName = true,
 					.bDefaultHidden   = false,
-					.FillWidth        = 0.52f,
+					.FillWidth        = 0.5f,
 					.Alignment        = EHorizontalAlignment::HAlign_Left
 				}
 			},
diff --git a/Engine/Plugins/AudioInsights/Source/AudioInsights/Private/Views/TreeDashboardViewFactory.cpp b/Engine/Plugins/AudioInsights/Source/AudioInsights/Private/Views/TreeDashboardViewFactory.cpp
--- a/Engine/Plugins/AudioInsights/Source/AudioInsights/Private/Views/TreeDashboardViewFactory.cpp
+++ b/Engine/Plugins/AudioInsights/Source/AudioInsights/Private/Views/TreeDashboardViewFactory.cpp
@@ -23,7 +23,7 @@ namespace UE::Audio::Insights
 		Data = InData;
 		Factory = InFactory;
 
-		FSuperRowType::FArguments Args = FSuperRowType::FArguments().Style(&FSlateStyle::Get().GetWidgetStyle<FTableRowStyle>("SoundDashboard.TableViewRow"));
+		FSuperRowType::FArguments Args = FSuperRowType::FArguments().Style(&FSlateStyle::Get().GetWidgetStyle<FTableRowStyle>("TreeDashboard.TableViewRow"));
 
 		SMultiColumnTableRow<TSharedPtr<IDashboardDataTreeViewEntry>>::Construct(Args, InOwnerTable);
 	}
@@ -172,7 +172,7 @@ namespace UE::Audio::Insights
 					// Icon (optional)
 					+ SHorizontalBox::Slot()
 					.AutoWidth()
-					.Padding(ColumnData.IconName != NAME_None ? 2.0f : 0.0f, 3.0f)
+					.Padding(ColumnData.IconName != NAME_None ? 4.0f : 0.0f, 3.0f)
 					[
 						ColumnData.IconName != NAME_None
 							? SNew(SImage).Image(FSlateStyle::Get().GetBrush(ColumnData.IconName)) 
diff --git a/Engine/Plugins/AudioInsights/Source/AudioInsights/Public/AudioInsightsStyle.h b/Engine/Plugins/AudioInsights/Source/AudioInsights/Public/AudioInsightsStyle.h
--- a/Engine/Plugins/AudioInsights/Source/AudioInsights/Public/AudioInsightsStyle.h
+++ b/Engine/Plugins/AudioInsights/Source/AudioInsights/Public/AudioInsightsStyle.h
@@ -122,8 +122,11 @@ namespace UE::Audio::Insights
 			SetContentRoot(PluginsDir / TEXT("AudioInsights/Content"));
 			SetCoreContentRoot(PluginsDir / TEXT("Slate"));
 
+			/* Audio Insights */
+			// Colors
 			Set("AudioInsights.Analyzers.BackgroundColor", FLinearColor(0.0075f, 0.0075f, 0.0075f, 1.0f));
 
+			// Icons
 			const FVector2D Icon16(16.0f, 16.0f);
 			const FVector2D Icon20(20.0f, 20.0f);
 			const FVector2D Icon24(24.0f, 24.0f);
@@ -149,6 +152,17 @@ namespace UE::Audio::Insights
 
 			Set("AudioInsights.Thumbnail", new IMAGE_BRUSH_SVG(TEXT("Icons/audio_insights"), Icon16));
 
+
+			/* Tree Dashboard */
+			// Table style
+			Set("TreeDashboard.TableViewRow", FTableRowStyle(FAppStyle::Get().GetWidgetStyle<FTableRowStyle>("TableView.Row"))
+				.SetOddRowBackgroundBrush(FSlateColorBrush(FStyleColors::Recessed))
+				.SetEvenRowBackgroundBrush(FSlateColorBrush(FStyleColors::Background))
+				.SetOddRowBackgroundHoveredBrush(FSlateColorBrush(FStyleColors::SelectHover))
+				.SetEvenRowBackgroundHoveredBrush(FSlateColorBrush(FStyleColors::SelectHover))
+			);
+
+
 			/* Sound Dashboard */
 			// Icons
 			Set("AudioInsights.Icon.SoundDashboard.Browse",           new IMAGE_BRUSH_SVG(TEXT("Icons/SoundDashboard/browse"),           Icon16));
@@ -164,6 +178,7 @@ namespace UE::Audio::Insights
 			Set("AudioInsights.Icon.SoundDashboard.SoundCue",         new IMAGE_BRUSH_SVG(TEXT("Icons/SoundDashboard/soundcue"),         Icon16));
 			Set("AudioInsights.Icon.SoundDashboard.SoundWave",        new IMAGE_BRUSH_SVG(TEXT("Icons/SoundDashboard/soundwave"),        Icon16));
 			Set("AudioInsights.Icon.SoundDashboard.Tab",              new IMAGE_BRUSH_SVG(TEXT("Icons/SoundDashboard/tab"),              Icon16));
+			Set("AudioInsights.Icon.SoundDashboard.Transparent",      new IMAGE_BRUSH_SVG(TEXT("Icons/SoundDashboard/transparent"),      Icon16));
 			
 			// Category colors
 			Set("SoundDashboard.MetaSoundColor", FLinearColor(0.008f, 0.76f, 0.078f));
@@ -177,14 +192,6 @@ namespace UE::Audio::Insights
 			// Text colors
 			Set("SoundDashboard.TimingOutTextColor", FLinearColor(1.0f, 1.0f, 1.0f, 0.25f));
 
-			// Table style
-			Set("SoundDashboard.TableViewRow", FTableRowStyle(FAppStyle::Get().GetWidgetStyle<FTableRowStyle>("TableView.Row"))
-				.SetOddRowBackgroundBrush(FSlateColorBrush(FStyleColors::Recessed))
-				.SetEvenRowBackgroundBrush(FSlateColorBrush(FStyleColors::Background))
-				.SetOddRowBackgroundHoveredBrush(FSlateColorBrush(FStyleColors::SelectHover))
-				.SetEvenRowBackgroundHoveredBrush(FSlateColorBrush(FStyleColors::SelectHover))
-			);
-
 			FSlateStyleRegistry::RegisterSlateStyle(*this);
 		}
 	};
