[Audio Insights] Sound Dashboard: Mute/Solo implementation

diff --git a/Engine/Plugins/AudioInsights/Source/AudioInsights/Private/Views/SoundDashboardViewFactory.cpp b/Engine/Plugins/AudioInsights/Source/AudioInsights/Private/Views/SoundDashboardViewFactory.cpp
--- a/Engine/Plugins/AudioInsights/Source/AudioInsights/Private/Views/SoundDashboardViewFactory.cpp
+++ b/Engine/Plugins/AudioInsights/Source/AudioInsights/Private/Views/SoundDashboardViewFactory.cpp
@@ -88,6 +88,155 @@ namespace UE::Audio::Insights
 			}
 		}
 
+#if WITH_EDITOR
+		void SetMuteSolo(const IDashboardDataTreeViewEntry& InEntry, const EMuteSoloMode InMuteSoloMode, const bool bInOnOff)
+		{
+#if ENABLE_AUDIO_DEBUG
+			if (FAudioDeviceManager* AudioDeviceManager = FAudioDeviceManager::Get())
+			{
+				::Audio::FAudioDebugger& AudioDebugger = AudioDeviceManager->GetDebugger();
+
+				const FSoundDashboardEntry& SoundEntry = FSoundDashboardViewFactoryPrivate::CastEntry(InEntry);
+
+				const FName SoundAssetDisplayName(SoundEntry.GetDisplayName().ToString());
+
+				const bool bIsSoundCueType = SoundEntry.EntryType == ESoundDashboardEntryType::SoundCue || SoundEntry.EntryType == ESoundDashboardEntryType::SoundCueTemplate;
+
+				if (!bIsSoundCueType)
+				{
+					switch (InMuteSoloMode)
+					{
+						case EMuteSoloMode::Mute:
+							AudioDebugger.SetMuteSoundWave(SoundAssetDisplayName, bInOnOff);
+							break;
+
+						case EMuteSoloMode::Solo:
+							AudioDebugger.SetSoloSoundWave(SoundAssetDisplayName, bInOnOff);
+							break;
+
+						default:
+							break;
+					}
+				}
+
+				for (const TSharedPtr<IDashboardDataTreeViewEntry>& SoundEntryChild : SoundEntry.Children)
+				{
+					if (SoundEntryChild.IsValid())
+					{
+						SetMuteSolo(*SoundEntryChild, InMuteSoloMode, bInOnOff);
+					}
+				}
+			}
+#endif // ENABLE_AUDIO_DEBUG
+		}
+
+		void ToggleMuteSolo(const IDashboardDataTreeViewEntry& InEntry, const EMuteSoloMode InMuteSoloMode)
+		{
+#if ENABLE_AUDIO_DEBUG
+			if (FAudioDeviceManager* AudioDeviceManager = FAudioDeviceManager::Get())
+			{
+				::Audio::FAudioDebugger& AudioDebugger = AudioDeviceManager->GetDebugger();
+
+				const FSoundDashboardEntry& SoundEntry = FSoundDashboardViewFactoryPrivate::CastEntry(InEntry);
+
+				const FName SoundAssetDisplayName(SoundEntry.GetDisplayName().ToString());
+
+				const bool bIsSoundCueType = SoundEntry.EntryType == ESoundDashboardEntryType::SoundCue || SoundEntry.EntryType == ESoundDashboardEntryType::SoundCueTemplate;
+
+				if (!bIsSoundCueType)
+				{
+					switch (InMuteSoloMode)
+					{
+						case EMuteSoloMode::Mute:
+							AudioDebugger.ToggleMuteSoundWave(SoundAssetDisplayName);
+							break;
+
+						case EMuteSoloMode::Solo:
+							AudioDebugger.ToggleSoloSoundWave(SoundAssetDisplayName);
+							break;
+
+						default:
+							break;
+					}
+				}
+
+				for (const TSharedPtr<IDashboardDataTreeViewEntry>& SoundEntryChild : SoundEntry.Children)
+				{
+					if (SoundEntryChild.IsValid())
+					{
+						ToggleMuteSolo(*SoundEntryChild, InMuteSoloMode);
+					}
+				}
+			}
+#endif // ENABLE_AUDIO_DEBUG
+		}
+
+		bool IsMuteSolo(const IDashboardDataTreeViewEntry& InEntry, const bool bInCheckChildren, const EMuteSoloMode InMuteSoloMode)
+		{
+#if ENABLE_AUDIO_DEBUG
+			if (FAudioDeviceManager* AudioDeviceManager = FAudioDeviceManager::Get())
+			{
+				::Audio::FAudioDebugger& AudioDebugger = AudioDeviceManager->GetDebugger();
+
+				const FSoundDashboardEntry& SoundEntry = FSoundDashboardViewFactoryPrivate::CastEntry(InEntry);
+
+				const FName SoundAssetDisplayName(SoundEntry.GetDisplayName().ToString());
+
+				const bool bIsSoundCueType = SoundEntry.EntryType == ESoundDashboardEntryType::SoundCue || SoundEntry.EntryType == ESoundDashboardEntryType::SoundCueTemplate;
+
+				if (!bIsSoundCueType)
+				{
+					switch (InMuteSoloMode)
+					{
+						case EMuteSoloMode::Mute:
+						{
+							if (AudioDebugger.IsMuteSoundWave(SoundAssetDisplayName))
+							{
+								return true;
+							}
+
+							break;
+						}
+
+						case EMuteSoloMode::Solo:
+						{
+							if (AudioDebugger.IsSoloSoundWave(SoundAssetDisplayName))
+							{
+								return true;
+							}
+
+							break;
+						}
+
+						default:
+							break;
+					}
+				}
+
+				if (bInCheckChildren)
+				{
+					uint32 NumChildrenMuteSolo = 0;
+
+					for (const TSharedPtr<IDashboardDataTreeViewEntry>& SoundEntryChild : SoundEntry.Children)
+					{
+						if (SoundEntryChild.IsValid() && IsMuteSolo(*SoundEntryChild, true /*bInCheckChildren*/, InMuteSoloMode))
+						{
+							++NumChildrenMuteSolo;
+						}
+					}
+
+					const bool bAllChildrenMuteSolo = !SoundEntry.Children.IsEmpty() && NumChildrenMuteSolo == SoundEntry.Children.Num();
+
+					return bAllChildrenMuteSolo;
+				}
+			}
+#endif // ENABLE_AUDIO_DEBUG
+
+			return false;
+		}
+#endif // WITH_EDITOR
+
+
 		constexpr ESoundDashboardFilterFlags AllFilterFlags = 
 			ESoundDashboardFilterFlags::MetaSound        |
 			ESoundDashboardFilterFlags::SoundCue         |
@@ -159,6 +308,8 @@ namespace UE::Audio::Insights
 #if WITH_EDITOR
 	TSharedRef<SWidget> FSoundDashboardViewFactory::MakeMuteSoloWidget()
 	{
+		using namespace FSoundDashboardViewFactoryPrivate;
+
 		const auto CreateButtonContentWidget = [](const FName& InIconName = FName(), const FText& InLabel = FText::GetEmpty(), const FName& InTextStyle = TEXT("ButtonText"))
         {
         	TSharedRef<SHorizontalBox> ButtonContainerWidget = SNew(SHorizontalBox);
@@ -208,7 +359,15 @@ namespace UE::Audio::Insights
 			[
 				SNew(SButton)
 				.ButtonStyle(&FAppStyle::Get().GetWidgetStyle<FButtonStyle>("SimpleButton"))
-				.OnClicked(this, &FSoundDashboardViewFactory::MuteSelected)
+				.OnClicked_Lambda([this]()
+				{
+					if (FilteredEntriesListView.IsValid())
+					{
+						ToggleMuteSoloEntries(FilteredEntriesListView->GetSelectedItems(), EMuteSoloMode::Mute);
+					}
+
+					return FReply::Handled();
+				})
 				[
 					CreateButtonContentWidget("AudioInsights.Icon.SoundDashboard.Mute", LOCTEXT("SoundDashboard_MuteButtonText", "Mute Selected"), "SmallButtonText")
 				]
@@ -220,7 +379,15 @@ namespace UE::Audio::Insights
 			[
 				SNew(SButton)
 				.ButtonStyle(&FAppStyle::Get().GetWidgetStyle<FButtonStyle>("SimpleButton"))
-				.OnClicked(this, &FSoundDashboardViewFactory::SoloSelected)
+				.OnClicked_Lambda([this]()
+				{
+					if (FilteredEntriesListView.IsValid())
+					{
+						ToggleMuteSoloEntries(FilteredEntriesListView->GetSelectedItems(), EMuteSoloMode::Solo);
+					}
+
+					return FReply::Handled();
+				})
 				[
 					CreateButtonContentWidget("AudioInsights.Icon.SoundDashboard.Solo", LOCTEXT("SoundDashboard_SoloButtonText", "Solo Selected"), "SmallButtonText")
 				]
@@ -448,35 +615,28 @@ namespace UE::Audio::Insights
 					})
 				];
 		}
-		else if (InColumn == "Mute" || InColumn == "Solo")
-		{
-			return SNew(SVerticalBox)
-				// Icon
-				+ SVerticalBox::Slot()
-				.Padding(2.0f, 2.0f, 2.0f, 2.0f)
-				.AutoHeight()
-				[
-					SNew(SImage)
-					.Image_Lambda([this, InRowData, InColumn]()
-					{
-						const FColumnData& ColumnData = GetColumns()[InColumn];
-						const FName IconName = ColumnData.GetIconName.IsSet() ? ColumnData.GetIconName(InRowData.Get()) : NAME_None;
-
-						return IconName != NAME_None ? FSlateStyle::Get().GetBrush(IconName) : nullptr;
-					})
-				];
-		}
 
 		return SNullWidget::NullWidget;
 	}
 
-	bool FSoundDashboardViewFactory::IsRootItem(const TSharedPtr<IDashboardDataTreeViewEntry>& InEntry) const
+	bool FSoundDashboardViewFactory::IsRootItem(const TSharedRef<IDashboardDataTreeViewEntry>& InEntry) const
 	{
 		return FilteredEntriesListView.IsValid() && FilteredEntriesListView->GetRootItems().Contains(InEntry);
 	}
 
+	bool FSoundDashboardViewFactory::EntryCanHaveChildren(const TSharedRef<IDashboardDataTreeViewEntry>& InEntry) const
+	{
+		const FSoundDashboardEntry& SoundEntry = FSoundDashboardViewFactoryPrivate::CastEntry(*InEntry);
+
+		const bool bIsSoundCueType = SoundEntry.EntryType == ESoundDashboardEntryType::SoundCue || SoundEntry.EntryType == ESoundDashboardEntryType::SoundCueTemplate;
+
+		return IsRootItem(InEntry) || bIsSoundCueType;
+	}
+
 	TSharedRef<SWidget> FSoundDashboardViewFactory::GenerateWidgetForColumn(const TSharedRef<FTraceTreeDashboardViewFactory::SRowWidget>& InRowWidget, const TSharedRef<IDashboardDataTreeViewEntry>& InRowData, const FName& InColumn)
 	{
+		using namespace FSoundDashboardViewFactoryPrivate;
+
 		const FColumnData& ColumnData = GetColumns()[InColumn];
 
 		const FText ValueText  = ColumnData.GetDisplayValue.IsSet() ? ColumnData.GetDisplayValue(InRowData.Get()) : FText::GetEmpty();
@@ -487,76 +647,103 @@ namespace UE::Audio::Insights
 			return SNullWidget::NullWidget;
 		}
 
-		if (IsRootItem(InRowData))
+#if WITH_EDITOR
+		if (InColumn == "Mute")
 		{
-			return GenerateWidgetForRootColumn(InRowWidget, InRowData, InColumn, ValueText);
+			return CreateMuteSoloButton(InRowData, InColumn, 
+				[this](const TArray<TSharedPtr<IDashboardDataTreeViewEntry>>& InEntries)
+				{
+					ToggleMuteSoloEntries(InEntries, EMuteSoloMode::Mute);
+				},
+				[](const IDashboardDataTreeViewEntry& InEntry, const bool bInCheckChildren)
+				{
+					return IsMuteSolo(InEntry, bInCheckChildren, EMuteSoloMode::Mute);
+				});
 		}
-
-		const FName IconName = ColumnData.GetIconName.IsSet() ? ColumnData.GetIconName(InRowData.Get()) : NAME_None;
-
-		return SNew(SHorizontalBox)
-			// Tree expander arrow (only for Name column)
-			+ SHorizontalBox::Slot()
-			.AutoWidth()
-			[
-				InColumn == "Name" ? SNew(SExpanderArrow, InRowWidget) : SNullWidget::NullWidget
-			]
-			// Icon (optional)
-			+ SHorizontalBox::Slot()
-			.Padding(2.0f, 2.0f, 2.0f, 2.0f)
-			.AutoWidth()
-			[
-				SNew(SImage)
-				.Image_Lambda([this, InRowData, InColumn]()
+		else if (InColumn == "Solo")
+		{
+			return CreateMuteSoloButton(InRowData, InColumn, 
+				[this](const TArray<TSharedPtr<IDashboardDataTreeViewEntry>>& InEntries)
 				{
-					const FColumnData& ColumnData = GetColumns()[InColumn];
-					const FName IconName = ColumnData.GetIconName.IsSet() ? ColumnData.GetIconName(InRowData.Get()) : NAME_None;
-
-					return IconName != NAME_None ? FSlateStyle::Get().GetBrush(IconName) : nullptr;
-				})
-				.ColorAndOpacity_Lambda([this, InRowData, InColumn]()
+					ToggleMuteSoloEntries(InEntries, EMuteSoloMode::Solo);
+				}, 
+				[](const IDashboardDataTreeViewEntry& InEntry, const bool bInCheckChildren)
 				{
-					const FColumnData& ColumnData = GetColumns()[InColumn];
-					const FSlateColor TextColor = ColumnData.GetTextColorValue.IsSet() ? ColumnData.GetTextColorValue(InRowData.Get()) : FSlateColor::UseForeground();
+					return IsMuteSolo(InEntry, bInCheckChildren, EMuteSoloMode::Solo);
+				});
+		}
+		else
+#endif // WITH_EDITOR
+		{
+			if (IsRootItem(InRowData))
+			{
+				return GenerateWidgetForRootColumn(InRowWidget, InRowData, InColumn, ValueText);
+			}
 
-					return TextColor;
-				})
-			]
-			// Text (optional)
-			+ SHorizontalBox::Slot()
-			.Padding(10.0f, 2.0f, 2.0f, 2.0f)
-			.AutoWidth()
-			[
-				SNew(STextBlock)
-				.Text_Lambda([this, InRowData, InColumn]()
-				{
-					const FColumnData& ColumnData = GetColumns()[InColumn];
-					const FText ValueText = ColumnData.GetDisplayValue.IsSet() ? ColumnData.GetDisplayValue(InRowData.Get()) : FText::GetEmpty();
+			return SNew(SHorizontalBox)
+				// Tree expander arrow (only for Name column)
+				+ SHorizontalBox::Slot()
+				.AutoWidth()
+				[
+					InColumn == "Name" ? SNew(SExpanderArrow, InRowWidget) : SNullWidget::NullWidget
+				]
+				// Icon (optional)
+				+ SHorizontalBox::Slot()
+				.Padding(2.0f, 2.0f, 2.0f, 2.0f)
+				.AutoWidth()
+				[
+					SNew(SImage)
+					.Image_Lambda([this, InRowData, InColumn]()
+					{
+						const FColumnData& ColumnData = GetColumns()[InColumn];
+						const FName IconName = ColumnData.GetIconName.IsSet() ? ColumnData.GetIconName(InRowData.Get()) : NAME_None;
+						
+						return IconName != NAME_None ? FSlateStyle::Get().GetBrush(IconName) : nullptr;
+					})
+					.ColorAndOpacity_Lambda([this, InRowData, InColumn]()
+					{
+						const FColumnData& ColumnData = GetColumns()[InColumn];
+						const FSlateColor TextColor = ColumnData.GetTextColorValue.IsSet() ? ColumnData.GetTextColorValue(InRowData.Get()) : FSlateColor::UseForeground();
+						
+						return TextColor;
+					})
+				]
+				// Text (optional)
+				+ SHorizontalBox::Slot()
+				.Padding(10.0f, 2.0f, 2.0f, 2.0f)
+				.AutoWidth()
+				[
+					SNew(STextBlock)
+					.Text_Lambda([this, InRowData, InColumn]()
+					{
+						const FColumnData& ColumnData = GetColumns()[InColumn];
+						const FText ValueText = ColumnData.GetDisplayValue.IsSet() ? ColumnData.GetDisplayValue(InRowData.Get()) : FText::GetEmpty();
 
-					return ValueText;
-				})
-				.ColorAndOpacity_Lambda([this, InRowData, InColumn]()
-				{
-					const FColumnData& ColumnData = GetColumns()[InColumn];
-					const FSlateColor TextColor = ColumnData.GetTextColorValue.IsSet() ? ColumnData.GetTextColorValue(InRowData.Get()) : FSlateColor::UseForeground();
+						return ValueText;
+					})
+					.ColorAndOpacity_Lambda([this, InRowData, InColumn]()
+					{
+						const FColumnData& ColumnData = GetColumns()[InColumn];
+						const FSlateColor TextColor = ColumnData.GetTextColorValue.IsSet() ? ColumnData.GetTextColorValue(InRowData.Get()) : FSlateColor::UseForeground();
 
-					return TextColor;
-				})
-			]
-			// Number of children text (if is not leaf item)
-			+ SHorizontalBox::Slot()
-			.Padding(6.0f, 2.0f, 0.0f, 2.0f)
-			.AutoWidth()
-			[
-				InColumn == "Name" && !InRowData->Children.IsEmpty()
-					? SNew(STextBlock)
-					  .ColorAndOpacity(FSlateColor(FColor::White.WithAlpha(128)))
-					  .Text_Lambda([InRowData]()
-					  {
-					      return FText::FromString("(" + FString::FromInt(InRowData->Children.Num()) + ")");
-					  })
-					: SNullWidget::NullWidget
-			];
+						return TextColor;
+					})
+				]
+				// Number of children text (if it is not leaf item)
+				+ SHorizontalBox::Slot()
+				.Padding(6.0f, 2.0f, 0.0f, 2.0f)
+				.AutoWidth()
+				[
+					InColumn == "Name" && !InRowData->Children.IsEmpty()
+						? SNew(STextBlock)
+						  .ColorAndOpacity(FSlateColor(FColor::White.WithAlpha(128)))
+						  .Text_Lambda([InRowData]()
+						  {
+							  return FText::FromString("(" + FString::FromInt(InRowData->Children.Num()) + ")");
+						  })
+						: SNullWidget::NullWidget
+				];
+		}
 	}
 
 	TSharedRef<ITableRow> FSoundDashboardViewFactory::OnGenerateRow(TSharedPtr<IDashboardDataTreeViewEntry> Item, const TSharedRef<STableViewBase>& OwnerTable)
@@ -635,9 +822,13 @@ namespace UE::Audio::Insights
 	TSharedPtr<SWidget> FSoundDashboardViewFactory::OnConstructContextMenu()
 	{
 		// Skip right click for root items
-		if (FilteredEntriesListView.IsValid() && FilteredEntriesListView->GetSelectedItems().Num() == 1 && IsRootItem(FilteredEntriesListView->GetSelectedItems()[0]))
+		if (FilteredEntriesListView.IsValid() && FilteredEntriesListView->GetSelectedItems().Num() == 1)
 		{
-			return SNullWidget::NullWidget;
+			const TSharedPtr<IDashboardDataTreeViewEntry> SelectedEntry = FilteredEntriesListView->GetSelectedItems()[0];
+			if (IsRootItem(SelectedEntry.ToSharedRef()))
+			{
+				return SNullWidget::NullWidget;
+			}
 		}
 
 		const FSoundDashboardCommands& Commands = FSoundDashboardCommands::Get();
@@ -772,7 +963,7 @@ namespace UE::Audio::Insights
 				{
 					.GetIconName = [this](const IDashboardDataTreeViewEntry& InData)
 					{
-						return IsMuted(InData) ? FName("AudioInsights.Icon.SoundDashboard.Mute") : NAME_None;
+						return "AudioInsights.Icon.SoundDashboard.Mute";
 					}
 				}
 			},
@@ -781,7 +972,7 @@ namespace UE::Audio::Insights
 				{
 					.GetIconName = [this](const IDashboardDataTreeViewEntry& InData)
 					{
-						return IsSoloed(InData) ? FName("AudioInsights.Icon.SoundDashboard.Solo") : NAME_None;
+						return "AudioInsights.Icon.SoundDashboard.Solo";
 					}
 				}
 			},
@@ -1098,112 +1289,65 @@ namespace UE::Audio::Insights
 	}
 
 #if WITH_EDITOR
-	FReply FSoundDashboardViewFactory::MuteSelected()
+	TSharedRef<SWidget> FSoundDashboardViewFactory::CreateMuteSoloButton(const TSharedRef<IDashboardDataTreeViewEntry>& InRowData, 
+		const FName& InColumn, 
+		TFunction<void(const TArray<TSharedPtr<IDashboardDataTreeViewEntry>>&)> MuteSoloToggleFunc, 
+		TFunctionRef<bool(const IDashboardDataTreeViewEntry&, const bool)> IsMuteSoloFunc)
 	{
-#if ENABLE_AUDIO_DEBUG
-		FAudioDeviceManager* AudioDeviceManager = FAudioDeviceManager::Get();
-
-		if (AudioDeviceManager && FilteredEntriesListView.IsValid())
-		{
-			const auto MuteSound = [&AudioDeviceManager](const TSharedPtr<IDashboardDataTreeViewEntry>& SelectedItem)
-			{
-				::Audio::FAudioDebugger AudioDebugger = AudioDeviceManager->GetDebugger();
-
-				const FSoundDashboardEntry& SoundAssetDashboardEntry = *StaticCastSharedPtr<FSoundDashboardEntry>(SelectedItem).Get();
-
-				const FName SoundAssetDisplayName = FName(SoundAssetDashboardEntry.GetDisplayName().ToString());
-
-				if (SoundAssetDashboardEntry.EntryType == ESoundDashboardEntryType::SoundCue
-					|| SoundAssetDashboardEntry.EntryType == ESoundDashboardEntryType::SoundCueTemplate)
-				{
-					AudioDebugger.SetMuteSoundCue(SoundAssetDisplayName, true /* bInOnOff */);
-				}
-				else
-				{
-					AudioDebugger.SetMuteSoundWave(SoundAssetDisplayName, true /* bInOnOff */);
-				}
-			};
-
-			for (const TSharedPtr<IDashboardDataTreeViewEntry>& SelectedItem : FilteredEntriesListView->GetSelectedItems())
-			{
-				if (SelectedItem.IsValid())
+		return SNew(SHorizontalBox)
+			+ SHorizontalBox::Slot()
+			.AutoWidth()
+			.Padding(2.0f, 0.0f, 0.0f, 0.0f)
+			[
+				SNew(SCheckBox)
+				.HAlign(EHorizontalAlignment::HAlign_Center)
+				.Style(&FSlateStyle::Get().GetWidgetStyle<FCheckBoxStyle>("SoundDashboard.MuteSoloButton"))
+				.IsChecked(ECheckBoxState::Unchecked)
+				.OnCheckStateChanged_Lambda([InRowData, MuteSoloToggleFunc](ECheckBoxState NewState)
 				{
-					MuteSound(SelectedItem);
-				}
-			}
-		}
-#endif // ENABLE_AUDIO_DEBUG
+					MuteSoloToggleFunc({ InRowData });
+				})
+				[
+					SNew(SImage)
+					.Image_Lambda([this, InRowData, InColumn, IsMuteSoloFunc]()
+					{
+						const FColumnData& ColumnData = GetColumns()[InColumn];
+						const FName IconName = IsMuteSoloFunc(*InRowData, EntryCanHaveChildren(InRowData) /*bInCheckChildren*/) && ColumnData.GetIconName.IsSet() ? ColumnData.GetIconName(InRowData.Get()) : "AudioInsights.Icon.SoundDashboard.Transparent";
 
-		return FReply::Handled();
+						return FSlateStyle::Get().GetBrush(IconName);
+					})
+				]
+			];
 	}
 
-	FReply FSoundDashboardViewFactory::SoloSelected()
+	void FSoundDashboardViewFactory::ToggleMuteSoloEntries(const TArray<TSharedPtr<IDashboardDataTreeViewEntry>>& InEntries, const EMuteSoloMode InMuteSoloMode)
 	{
-#if ENABLE_AUDIO_DEBUG
-		FAudioDeviceManager* AudioDeviceManager = FAudioDeviceManager::Get();
+		using namespace FSoundDashboardViewFactoryPrivate;
 
-		if (AudioDeviceManager && FilteredEntriesListView.IsValid())
+		for (const TSharedPtr<IDashboardDataTreeViewEntry>& Entry : InEntries)
 		{
-			const auto SoloSound = [&AudioDeviceManager](const TSharedPtr<IDashboardDataTreeViewEntry>& SelectedItem)
+			if (!Entry.IsValid())
 			{
-				::Audio::FAudioDebugger AudioDebugger = AudioDeviceManager->GetDebugger();
-
-				const FSoundDashboardEntry& SoundAssetDashboardEntry = *StaticCastSharedPtr<FSoundDashboardEntry>(SelectedItem).Get();
+				continue;
+			}
 
-				const FName SoundAssetDisplayName = FName(SoundAssetDashboardEntry.GetDisplayName().ToString());
+			if (EntryCanHaveChildren(Entry.ToSharedRef()))
+			{
+				const bool bAreChildrenMuteSolo = FSoundDashboardViewFactoryPrivate::IsMuteSolo(*Entry, true /*bInCheckChildren*/, InMuteSoloMode);
 
-				if (SoundAssetDashboardEntry.EntryType == ESoundDashboardEntryType::SoundCue
-					|| SoundAssetDashboardEntry.EntryType == ESoundDashboardEntryType::SoundCueTemplate)
+				for (const TSharedPtr<IDashboardDataTreeViewEntry>& ChildEntry : Entry->Children)
 				{
-					AudioDebugger.SetSoloSoundCue(SoundAssetDisplayName, true /* bInOnOff */);
-				}
-				else
-				{
-					AudioDebugger.SetSoloSoundWave(SoundAssetDisplayName, true /* bInOnOff */);
+					if (Entry.IsValid())
+					{
+						FSoundDashboardViewFactoryPrivate::SetMuteSolo(*ChildEntry, InMuteSoloMode, !bAreChildrenMuteSolo);
+					}
 				}
-			};
-
-			for (const TSharedPtr<IDashboardDataTreeViewEntry>& SelectedItem : FilteredEntriesListView->GetSelectedItems())
+			}
+			else
 			{
-				if (SelectedItem.IsValid())
-				{
-					SoloSound(SelectedItem);
-				}
+				FSoundDashboardViewFactoryPrivate::ToggleMuteSolo(*Entry, InMuteSoloMode);
 			}
 		}
-#endif // ENABLE_AUDIO_DEBUG
-
-		return FReply::Handled();
-	}
-
-	bool FSoundDashboardViewFactory::IsSoloed(const IDashboardDataTreeViewEntry& InData) const
-	{
-#if ENABLE_AUDIO_DEBUG
-		if (FAudioDeviceManager* AudioDeviceManager = FAudioDeviceManager::Get())
-		{
-			const FSoundDashboardEntry& SoundAssetDashboardEntry = FSoundDashboardViewFactoryPrivate::CastEntry(InData);
-			const FName SoundAssetName(SoundAssetDashboardEntry.GetDisplayName().ToString());
-
-			return AudioDeviceManager->GetDebugger().IsSoloSoundWave(SoundAssetName);
-		}
-#endif // ENABLE_AUDIO_DEBUG
-
-		return false;
-	}
-
-	bool FSoundDashboardViewFactory::IsMuted(const IDashboardDataTreeViewEntry& InData) const
-	{
-#if ENABLE_AUDIO_DEBUG
-		if (FAudioDeviceManager* AudioDeviceManager = FAudioDeviceManager::Get())
-		{
-			const FSoundDashboardEntry& SoundAssetDashboardEntry = FSoundDashboardViewFactoryPrivate::CastEntry(InData);
-			const FName SoundAssetName(SoundAssetDashboardEntry.GetDisplayName().ToString());
-
-			return AudioDeviceManager->GetDebugger().IsMuteSoundWave(SoundAssetName);
-		}
-#endif // ENABLE_AUDIO_DEBUG
-
-		return false;
 	}
 
 	TArray<TObjectPtr<UObject>> FSoundDashboardViewFactory::GetSelectedEditableAssets() const
diff --git a/Engine/Plugins/AudioInsights/Source/AudioInsights/Public/AudioInsightsStyle.h b/Engine/Plugins/AudioInsights/Source/AudioInsights/Public/AudioInsightsStyle.h
--- a/Engine/Plugins/AudioInsights/Source/AudioInsights/Public/AudioInsightsStyle.h
+++ b/Engine/Plugins/AudioInsights/Source/AudioInsights/Public/AudioInsightsStyle.h
@@ -192,6 +192,15 @@ namespace UE::Audio::Insights
 			// Text colors
 			Set("SoundDashboard.TimingOutTextColor", FLinearColor(1.0f, 1.0f, 1.0f, 0.25f));
 
+			// Mute/Solo buttons style
+			const FSlateRoundedBoxBrush RoundedWhiteAlphaBrush(FLinearColor(1.0f, 1.0f, 1.0f, 0.1f), 5.0f);
+
+			Set("SoundDashboard.MuteSoloButton", FCheckBoxStyle(FAppStyle::Get().GetWidgetStyle<FCheckBoxStyle>("TransparentCheckBox"))
+				.SetPadding(FMargin(5.0f, 2.0f, 5.0f, 2.0f))
+				.SetCheckedHoveredImage(RoundedWhiteAlphaBrush)
+				.SetUncheckedHoveredImage(RoundedWhiteAlphaBrush)
+			);
+
 			FSlateStyleRegistry::RegisterSlateStyle(*this);
 		}
 	};
diff --git a/Engine/Plugins/AudioInsights/Source/AudioInsights/Public/Views/SoundDashboardViewFactory.h b/Engine/Plugins/AudioInsights/Source/AudioInsights/Public/Views/SoundDashboardViewFactory.h
--- a/Engine/Plugins/AudioInsights/Source/AudioInsights/Public/Views/SoundDashboardViewFactory.h
+++ b/Engine/Plugins/AudioInsights/Source/AudioInsights/Public/Views/SoundDashboardViewFactory.h
@@ -23,6 +23,14 @@ namespace UE::Audio::Insights
 
 	ENUM_CLASS_FLAGS(ESoundDashboardFilterFlags);
 
+#if WITH_EDITOR
+	enum class EMuteSoloMode : uint8
+	{
+		Mute,
+		Solo
+	};
+#endif // WITH_EDITOR
+
 	class FSoundDashboardFilter : public FGenericFilter<ESoundDashboardFilterFlags>
 	{
 	public:
@@ -71,7 +79,8 @@ namespace UE::Audio::Insights
 		virtual TSharedPtr<SWidget> GetFilterBarWidget() override;
 		virtual TSharedPtr<SWidget> GetFilterBarButtonWidget() override;
 
-		bool IsRootItem(const TSharedPtr<IDashboardDataTreeViewEntry>& InEntry) const;
+		bool IsRootItem(const TSharedRef<IDashboardDataTreeViewEntry>& InEntry) const;
+		bool EntryCanHaveChildren(const TSharedRef<IDashboardDataTreeViewEntry>& InEntry) const;
 
 		TSharedRef<SWidget> GenerateWidgetForRootColumn(const TSharedRef<FTraceTreeDashboardViewFactory::SRowWidget>& InRowWidget, const TSharedRef<IDashboardDataTreeViewEntry>& InRowData, const FName& InColumn, const FText& InValueText);
 		virtual TSharedRef<SWidget> GenerateWidgetForColumn(const TSharedRef<FTraceTreeDashboardViewFactory::SRowWidget>& InRowWidget, const TSharedRef<IDashboardDataTreeViewEntry>& InRowData, const FName& InColumn) override;
@@ -88,11 +97,12 @@ namespace UE::Audio::Insights
 #if WITH_EDITOR
 		TSharedRef<SWidget> MakeMuteSoloWidget();
 
-		FReply MuteSelected();
-		FReply SoloSelected();
+		TSharedRef<SWidget> CreateMuteSoloButton(const TSharedRef<IDashboardDataTreeViewEntry>& InRowData,
+			const FName& InColumn,
+			TFunction<void(const TArray<TSharedPtr<IDashboardDataTreeViewEntry>>& /*InEntries*/)> MuteSoloToggleFunc,
+			TFunctionRef<bool(const IDashboardDataTreeViewEntry& /*InEntry*/, const bool /*bInCheckChildren*/)> IsMuteSoloFunc);
 
-		bool IsSoloed(const IDashboardDataTreeViewEntry& InData) const;
-		bool IsMuted(const IDashboardDataTreeViewEntry& InData) const;
+		void ToggleMuteSoloEntries(const TArray<TSharedPtr<IDashboardDataTreeViewEntry>>& InEntries, const EMuteSoloMode InMuteSoloMode);
 
 		TArray<TObjectPtr<UObject>> GetSelectedEditableAssets() const;
 #endif // WITH_EDITOR
