[Audio Insights] Sound dashboard: Comment the Hide category logic since it has not been fully implemented yet (potentially for UE5.7)

diff --git a/Engine/Plugins/AudioInsights/Source/AudioInsights/Private/SoundDashboardCommands.cpp b/Engine/Plugins/AudioInsights/Source/AudioInsights/Private/SoundDashboardCommands.cpp
--- a/Engine/Plugins/AudioInsights/Source/AudioInsights/Private/SoundDashboardCommands.cpp
+++ b/Engine/Plugins/AudioInsights/Source/AudioInsights/Private/SoundDashboardCommands.cpp
@@ -20,7 +20,8 @@ namespace UE::Audio::Insights
 		UI_COMMAND(Pin, "Pin", "Pins the selected sound in the Pinned category.", EUserInterfaceActionType::Button, FInputChord(EKeys::P, EModifierKey::Control));
 		UI_COMMAND(Browse, "Browse To Asset", "Browses to the selected sound asset in the content browser.", EUserInterfaceActionType::Button, FInputChord(EKeys::B, EModifierKey::Control));
 		UI_COMMAND(Edit, "Edit", "Opens the selected sound for edit.", EUserInterfaceActionType::Button, FInputChord(EKeys::E, EModifierKey::Control));
-		UI_COMMAND(Hide, "Hide", "Hides the selected sound.", EUserInterfaceActionType::Button, FInputChord(EKeys::H, EModifierKey::Control));
+		// @TODO UE-250399: Hide category pending to implement
+		//UI_COMMAND(Hide, "Hide", "Hides the selected sound.", EUserInterfaceActionType::Button, FInputChord(EKeys::H, EModifierKey::Control));
 	}
 } // namespace UE::Audio::Insights
 
diff --git a/Engine/Plugins/AudioInsights/Source/AudioInsights/Private/SoundDashboardCommands.h b/Engine/Plugins/AudioInsights/Source/AudioInsights/Private/SoundDashboardCommands.h
--- a/Engine/Plugins/AudioInsights/Source/AudioInsights/Private/SoundDashboardCommands.h
+++ b/Engine/Plugins/AudioInsights/Source/AudioInsights/Private/SoundDashboardCommands.h
@@ -18,12 +18,14 @@ namespace UE::Audio::Insights
 		TSharedPtr<const FUICommandInfo> GetPinCommand() const    { return Pin; }
 		TSharedPtr<const FUICommandInfo> GetBrowseCommand() const { return Browse; }
 		TSharedPtr<const FUICommandInfo> GetEditCommand() const   { return Edit; }
-		TSharedPtr<const FUICommandInfo> GetHideCommand() const   { return Hide; }
+		// @TODO UE-250399: Hide category pending to implement
+		//TSharedPtr<const FUICommandInfo> GetHideCommand() const   { return Hide; }
 		
 	private:
 		TSharedPtr<FUICommandInfo> Pin;
 		TSharedPtr<FUICommandInfo> Browse;
 		TSharedPtr<FUICommandInfo> Edit;
-		TSharedPtr<FUICommandInfo> Hide;
+		// @TODO UE-250399: Hide category pending to implement
+		//TSharedPtr<FUICommandInfo> Hide;
 	};
 } // namespace UE::Audio::Insights
diff --git a/Engine/Plugins/AudioInsights/Source/AudioInsights/Private/Views/SoundDashboardViewFactory.cpp b/Engine/Plugins/AudioInsights/Source/AudioInsights/Private/Views/SoundDashboardViewFactory.cpp
--- a/Engine/Plugins/AudioInsights/Source/AudioInsights/Private/Views/SoundDashboardViewFactory.cpp
+++ b/Engine/Plugins/AudioInsights/Source/AudioInsights/Private/Views/SoundDashboardViewFactory.cpp
@@ -302,7 +302,9 @@ namespace UE::Audio::Insights
 		CommandList->MapAction(Commands.GetBrowseCommand(), FExecuteAction::CreateRaw(this, &FSoundDashboardViewFactory::BrowseSoundAsset));
 		CommandList->MapAction(Commands.GetEditCommand(),   FExecuteAction::CreateRaw(this, &FSoundDashboardViewFactory::OpenSoundAsset));
 #endif // WITH_EDITOR
-		CommandList->MapAction(Commands.GetHideCommand(), FExecuteAction::CreateRaw(this, &FSoundDashboardViewFactory::HideSound));
+
+		// @TODO UE-250399: Hide category pending to implement
+		//CommandList->MapAction(Commands.GetHideCommand(), FExecuteAction::CreateRaw(this, &FSoundDashboardViewFactory::HideSound));
 	}
 
 #if WITH_EDITOR
@@ -456,14 +458,15 @@ namespace UE::Audio::Insights
 				FSlateStyle::Get().GetColor("SoundDashboard.PinnedColor"),
 				FilterCategory),
 
-				MakeShared<FSoundDashboardFilter>(
-				ESoundDashboardFilterFlags::Hidden,
-				"Hidden",
-				LOCTEXT("SoundDashboard_Filter_HiddenNameText", "Hidden"),
-				"AudioInsights.Icon.SoundDashboard.Hide",
-				FText::GetEmpty(),
-				FSlateStyle::Get().GetColor("SoundDashboard.HiddenColor"),
-				FilterCategory)
+				// @TODO UE-250399: Hide category pending to implement
+				//MakeShared<FSoundDashboardFilter>(
+				//ESoundDashboardFilterFlags::Hidden,
+				//"Hidden",
+				//LOCTEXT("SoundDashboard_Filter_HiddenNameText", "Hidden"),
+				//"AudioInsights.Icon.SoundDashboard.Hide",
+				//FText::GetEmpty(),
+				//FSlateStyle::Get().GetColor("SoundDashboard.HiddenColor"),
+				//FilterCategory)
 			};
 
 			SAssignNew(SoundsFilterBar, SAudioFilterBar<ESoundDashboardFilterFlags>)
@@ -844,7 +847,9 @@ namespace UE::Audio::Insights
 			MenuBuilder.AddMenuEntry(Commands.GetBrowseCommand(), NAME_None, TAttribute<FText>(), TAttribute<FText>(), FSlateStyle::Get().CreateIcon("AudioInsights.Icon.SoundDashboard.Browse"));
 			MenuBuilder.AddMenuEntry(Commands.GetEditCommand(),   NAME_None, TAttribute<FText>(), TAttribute<FText>(), FSlateStyle::Get().CreateIcon("AudioInsights.Icon.SoundDashboard.Edit"));
 #endif // WITH_EDITOR
-			MenuBuilder.AddMenuEntry(Commands.GetHideCommand(), NAME_None, TAttribute<FText>(), TAttribute<FText>(), FSlateStyle::Get().CreateIcon("AudioInsights.Icon.SoundDashboard.Hide"));
+
+			// @TODO UE-250399: Hide category pending to implement
+			//MenuBuilder.AddMenuEntry(Commands.GetHideCommand(), NAME_None, TAttribute<FText>(), TAttribute<FText>(), FSlateStyle::Get().CreateIcon("AudioInsights.Icon.SoundDashboard.Hide"));
 		}
 
 		MenuBuilder.EndSection();
@@ -1420,10 +1425,11 @@ namespace UE::Audio::Insights
 	}
 #endif // WITH_EDITOR
 
-	void FSoundDashboardViewFactory::HideSound()
-	{
-		// @TODO UE-220744: pending to implement
-	}
+	// @TODO UE-250399: Hide category pending to implement
+	//void FSoundDashboardViewFactory::HideSound()
+	//{
+	//
+	//}
 } // namespace UE::Audio::Insights
 
 #undef LOCTEXT_NAMESPACE
diff --git a/Engine/Plugins/AudioInsights/Source/AudioInsights/Public/AudioInsightsStyle.h b/Engine/Plugins/AudioInsights/Source/AudioInsights/Public/AudioInsightsStyle.h
--- a/Engine/Plugins/AudioInsights/Source/AudioInsights/Public/AudioInsightsStyle.h
+++ b/Engine/Plugins/AudioInsights/Source/AudioInsights/Public/AudioInsightsStyle.h
@@ -168,7 +168,8 @@ namespace UE::Audio::Insights
 			Set("AudioInsights.Icon.SoundDashboard.Browse",           new IMAGE_BRUSH_SVG(TEXT("Icons/SoundDashboard/browse"),           Icon16));
 			Set("AudioInsights.Icon.SoundDashboard.Edit",             new IMAGE_BRUSH_SVG(TEXT("Icons/SoundDashboard/edit"),             Icon16));
 			Set("AudioInsights.Icon.SoundDashboard.Filter",           new IMAGE_BRUSH_SVG(TEXT("Icons/SoundDashboard/filter"),           Icon16));
-			Set("AudioInsights.Icon.SoundDashboard.Hide",             new IMAGE_BRUSH_SVG(TEXT("Icons/SoundDashboard/hide"),             Icon16));
+			// @TODO UE-250399: Hide category pending to implement
+			//Set("AudioInsights.Icon.SoundDashboard.Hide",             new IMAGE_BRUSH_SVG(TEXT("Icons/SoundDashboard/hide"),             Icon16));
 			Set("AudioInsights.Icon.SoundDashboard.Info",             new IMAGE_BRUSH_SVG(TEXT("Icons/SoundDashboard/info"),             Icon16));
 			Set("AudioInsights.Icon.SoundDashboard.MetaSound",        new IMAGE_BRUSH_SVG(TEXT("Icons/SoundDashboard/metasound"),        Icon16));
 			Set("AudioInsights.Icon.SoundDashboard.Mute",             new IMAGE_BRUSH_SVG(TEXT("Icons/SoundDashboard/mute"),             Icon16));
diff --git a/Engine/Plugins/AudioInsights/Source/AudioInsights/Public/Views/SoundDashboardViewFactory.h b/Engine/Plugins/AudioInsights/Source/AudioInsights/Public/Views/SoundDashboardViewFactory.h
--- a/Engine/Plugins/AudioInsights/Source/AudioInsights/Public/Views/SoundDashboardViewFactory.h
+++ b/Engine/Plugins/AudioInsights/Source/AudioInsights/Public/Views/SoundDashboardViewFactory.h
@@ -113,7 +113,8 @@ namespace UE::Audio::Insights
 		void BrowseSoundAsset() const;
 		void OpenSoundAsset() const;
 #endif
-		void HideSound();
+		// @TODO UE-250399: Hide category pending to implement
+		//void HideSound();
 
 		TSharedPtr<FUICommandList> CommandList;
 		TSharedPtr<SWidget> SoundsFilterBar;
