Home Screen - Pinned Sounds & Sorting of columns
- Allows user to pin entires to the top of the sounds dashboard
- Pinned entires are sorted seperately to other entries
- Updated sorting of columns to be recursive. Items are sorted within categories.

diff --git a/Engine/Plugins/AudioInsights/Source/AudioInsights/Private/Views/SoundDashboardViewFactory.cpp b/Engine/Plugins/AudioInsights/Source/AudioInsights/Private/Views/SoundDashboardViewFactory.cpp
--- a/Engine/Plugins/AudioInsights/Source/AudioInsights/Private/Views/SoundDashboardViewFactory.cpp
+++ b/Engine/Plugins/AudioInsights/Source/AudioInsights/Private/Views/SoundDashboardViewFactory.cpp
@@ -21,6 +21,8 @@
 
 namespace UE::Audio::Insights
 {
+	/////////////////////////////////////////////////////////////////////////////////////////
+	// FSoundDashboardViewFactoryPrivate
 	namespace FSoundDashboardViewFactoryPrivate
 	{
 		const FSoundDashboardEntry& CastEntry(const IDashboardDataTreeViewEntry& InData)
@@ -245,9 +247,149 @@ namespace UE::Audio::Insights
 			ESoundDashboardFilterFlags::SoundCueTemplate |
 			ESoundDashboardFilterFlags::Pinned           |
 			ESoundDashboardFilterFlags::Hidden;
+		const FText FiltersName = LOCTEXT("SoundDashboard_Filter_CategoryText", "Filters");
+		const FText FiltersTooltip = LOCTEXT("CurveFiltersToolTip", "Filters what kind of sounds types can be displayed.");
+
+		const FText MetaSoundCategoryName = LOCTEXT("SoundDashboard_Filter_MetaSoundNameText", "MetaSound");
+		const FText SoundCueCategoryName = LOCTEXT("SoundDashboard_Filter_SoundCueNameText", "Sound Cue");
+		const FText ProceduralSourceCategoryName = LOCTEXT("SoundDashboard_Filter_ProceduralSourceNameText", "Procedural Source");
+		const FText SoundWaveCategoryName = LOCTEXT("SoundDashboard_Filter_SoundWaveNameText", "Sound Wave");
+		const FText SoundCueTemplateCategoryName = LOCTEXT("SoundDashboard_Filter_SoundCueTemplateNameText", "Sound Cue Template");
+		const FText PinnedCategoryName = LOCTEXT("SoundDashboard_Filter_PinnedNameText", "Pinned");
+		const FText HiddenCategoryName = LOCTEXT("SoundDashboard_Filter_HiddenNameText", "Hidden");
 	} // namespace FSoundDashboardViewFactoryPrivate
 
+	/////////////////////////////////////////////////////////////////////////////////////////
+	// FPinnedSoundEntryWrapperPrivate
+	namespace FPinnedSoundEntryWrapperPrivate
+	{
+		bool CanBeDeleted(const TSharedPtr<FPinnedSoundEntryWrapper>& Entry)
+		{
+			return !Entry->EntryIsValid() || (Entry->bIsCategoryEntry && Entry->PinnedWrapperChildren.IsEmpty());
+		}
+
+		void CopyDataToPinnedEntry(FSoundDashboardEntry& PinnedEntry, const FSoundDashboardEntry& OriginalEntry)
+		{
+			// Only copy data that has possibly changed from the other entry
+			PinnedEntry.TimeoutTimestamp = OriginalEntry.TimeoutTimestamp;
+
+			PinnedEntry.bShouldForceExpandChildren = OriginalEntry.bShouldForceExpandChildren;
+			PinnedEntry.bIsVisible = OriginalEntry.bIsVisible;
+
+			// Just copy the last entry in the buffer rather than the whole buffer
+			auto AddLastValueInBuffer = [](::Audio::TCircularAudioBuffer<FDataPoint>& To, const ::Audio::TCircularAudioBuffer<FDataPoint>& From)
+			{
+				if (From.Num() == 0)
+				{
+					return;
+				}
+
+				if (To.Num() > 0)
+				{
+					To.Pop(1);
+				}
+
+				const ::Audio::DisjointedArrayView<const FDataPoint> DataPointsDisjointedArrayView = From.PeekInPlace(From.Num());
+				To.Push(DataPointsDisjointedArrayView.FirstBuffer.Last());
+			};
+
+			AddLastValueInBuffer(PinnedEntry.PriorityDataPoints, OriginalEntry.PriorityDataPoints);
+			AddLastValueInBuffer(PinnedEntry.DistanceDataPoints, OriginalEntry.DistanceDataPoints);
+			AddLastValueInBuffer(PinnedEntry.AmplitudeDataPoints, OriginalEntry.AmplitudeDataPoints);
+			AddLastValueInBuffer(PinnedEntry.VolumeDataPoints, OriginalEntry.VolumeDataPoints);
+			AddLastValueInBuffer(PinnedEntry.PitchDataPoints, OriginalEntry.PitchDataPoints);
+		};
+	} // namespace FPinnedSoundEntryWrapperPrivate
+
+	/////////////////////////////////////////////////////////////////////////////////////////
+	// FPinnedSoundEntryWrapper
+	FPinnedSoundEntryWrapper::FPinnedSoundEntryWrapper(const TSharedPtr<IDashboardDataTreeViewEntry>& OriginalEntry, const bool bIsCategory)
+		: bIsCategoryEntry(bIsCategory)
+		, OriginalDataEntry(OriginalEntry)
+	{
+		if (!OriginalEntry.IsValid())
+		{
+			return;
+		}
+
+		// Take a deep copy of the original entry to add to the pinned section of the dashboard
+		// We need deep copies of any children too
+		PinnedSectionEntry = MakeShared<FSoundDashboardEntry>(FSoundDashboardViewFactoryPrivate::CastEntry(*OriginalEntry));
+		PinnedSectionEntry->Children.Empty();
+
+		FSoundDashboardEntry& PinnedSectionSoundEntry = FSoundDashboardViewFactoryPrivate::CastEntry(*PinnedSectionEntry);
+		PinnedSectionSoundEntry.PinnedEntryType = FSoundDashboardEntry::EPinnedEntryType::PinnedCopy;
+
+		for (const TSharedPtr<IDashboardDataTreeViewEntry>& Child : OriginalEntry->Children)
+		{
+			AddChildEntry(Child.ToSharedRef(), false /*bIsCategory*/);
+		}
+	}
+
+	TSharedPtr<FPinnedSoundEntryWrapper> FPinnedSoundEntryWrapper::AddChildEntry(const TSharedPtr<IDashboardDataTreeViewEntry> Child, const bool bIsCategory)
+	{
+		TSharedPtr<FPinnedSoundEntryWrapper> NewChild = MakeShared<FPinnedSoundEntryWrapper>(Child, bIsCategory);
+
+		FSoundDashboardEntry& NewChildSound = FSoundDashboardViewFactoryPrivate::CastEntry(*NewChild->GetPinnedSectionEntry());
+		NewChildSound.PinnedEntryType = FSoundDashboardEntry::EPinnedEntryType::PinnedCopy;
+
+		PinnedWrapperChildren.Add(NewChild);
+		PinnedSectionEntry->Children.Add(NewChild->GetPinnedSectionEntry());
+
+		return NewChild;
+	}
 
+	void FPinnedSoundEntryWrapper::UpdateParams()
+	{
+		// If we lose our handle to the original entry, we should stop updating
+		if (!EntryIsValid())
+		{
+			OriginalDataEntry.Reset();
+			PinnedSectionEntry.Reset();
+
+			return;
+		}
+
+		// Only non-category entries have data to update
+		if (OriginalDataEntry.IsValid())
+		{
+			FSoundDashboardEntry& Pinned = FSoundDashboardViewFactoryPrivate::CastEntry(*PinnedSectionEntry);
+			const FSoundDashboardEntry& Original = FSoundDashboardViewFactoryPrivate::CastEntry(*OriginalDataEntry.Pin());
+
+			FPinnedSoundEntryWrapperPrivate::CopyDataToPinnedEntry(Pinned, Original);
+		}
+
+		for (const TSharedPtr<FPinnedSoundEntryWrapper>& Child : PinnedWrapperChildren)
+		{
+			Child->UpdateParams();
+		}
+	}
+
+	void FPinnedSoundEntryWrapper::CleanUp()
+	{
+		for (int Index = PinnedWrapperChildren.Num() - 1; Index >= 0; --Index)
+		{
+			const TSharedPtr<FPinnedSoundEntryWrapper> Child = PinnedWrapperChildren[Index];
+			if (FPinnedSoundEntryWrapperPrivate::CanBeDeleted(Child))
+			{
+				TSharedPtr<IDashboardDataTreeViewEntry> DataEntry = Child->GetPinnedSectionEntry();
+				PinnedSectionEntry->Children.Remove(DataEntry);
+				PinnedWrapperChildren.Remove(Child);
+			}
+			else
+			{
+				Child->CleanUp();
+			}
+		}
+	}
+
+	bool FPinnedSoundEntryWrapper::EntryIsValid() const
+	{
+		return PinnedSectionEntry.IsValid() && (OriginalDataEntry.IsValid() || bIsCategoryEntry);
+	}
+
+	/////////////////////////////////////////////////////////////////////////////////////////
+	// FSoundDashboardViewFactory
 	FSoundDashboardViewFactory::FSoundDashboardViewFactory()
 	{
 		FTraceModule& AudioInsightsTraceModule = static_cast<FTraceModule&>(FAudioInsightsModule::GetChecked().GetTraceModule());
@@ -399,15 +541,17 @@ namespace UE::Audio::Insights
 
 	TSharedPtr<SWidget> FSoundDashboardViewFactory::GetFilterBarWidget()
 	{
+		using namespace FSoundDashboardViewFactoryPrivate;
+
 		if (!SoundsFilterBar.IsValid())
 		{
-			const TSharedPtr<FFilterCategory> FilterCategory = MakeShared<FFilterCategory>(LOCTEXT("SoundDashboard_Filter_CategoryText", "Filters"), LOCTEXT("CurveFiltersToolTip", "Filters what kind of sounds types can be displayed."));
+			const TSharedPtr<FFilterCategory> FilterCategory = MakeShared<FFilterCategory>(FiltersName, FiltersTooltip);
 			
 			TArray<TSharedRef<FFilterBase<ESoundDashboardFilterFlags>>> Filters{
 				MakeShared<FSoundDashboardFilter>(
 				ESoundDashboardFilterFlags::MetaSound,
 				"MetaSound",
-				LOCTEXT("SoundDashboard_Filter_MetaSoundNameText", "MetaSound"),
+				MetaSoundCategoryName,
 				"AudioInsights.Icon.SoundDashboard.MetaSound",
 				FText::GetEmpty(),
 				FSlateStyle::Get().GetColor("SoundDashboard.MetaSoundColor"),
@@ -416,7 +560,7 @@ namespace UE::Audio::Insights
 				MakeShared<FSoundDashboardFilter>(
 				ESoundDashboardFilterFlags::SoundCue,
 				"SoundCue",
-				LOCTEXT("SoundDashboard_Filter_SoundCueNameText", "Sound Cue"),
+				SoundCueCategoryName,
 				"AudioInsights.Icon.SoundDashboard.SoundCue",
 				FText::GetEmpty(),
 				FSlateStyle::Get().GetColor("SoundDashboard.SoundCueColor"),
@@ -425,7 +569,7 @@ namespace UE::Audio::Insights
 				MakeShared<FSoundDashboardFilter>(
 				ESoundDashboardFilterFlags::ProceduralSource,
 				"ProceduralSource",
-				LOCTEXT("SoundDashboard_Filter_ProceduralSourceNameText", "Procedural Source"),
+				ProceduralSourceCategoryName,
 				"AudioInsights.Icon.SoundDashboard.ProceduralSource",
 				FText::GetEmpty(),
 				FSlateStyle::Get().GetColor("SoundDashboard.ProceduralSourceColor"),
@@ -434,7 +578,7 @@ namespace UE::Audio::Insights
 				MakeShared<FSoundDashboardFilter>(
 				ESoundDashboardFilterFlags::SoundWave,
 				"SoundWave",
-				LOCTEXT("SoundDashboard_Filter_SoundWaveNameText", "Sound Wave"),
+				SoundWaveCategoryName,
 				"AudioInsights.Icon.SoundDashboard.SoundWave",
 				FText::GetEmpty(),
 				FSlateStyle::Get().GetColor("SoundDashboard.SoundWaveColor"),
@@ -443,7 +587,7 @@ namespace UE::Audio::Insights
 				MakeShared<FSoundDashboardFilter>(
 				ESoundDashboardFilterFlags::SoundCueTemplate,
 				"SoundCueTemplate",
-				LOCTEXT("SoundDashboard_Filter_SoundCueTemplateNameText", "Sound Cue Template"),
+				SoundCueTemplateCategoryName,
 				"AudioInsights.Icon.SoundDashboard.SoundCue",
 				FText::GetEmpty(),
 				FSlateStyle::Get().GetColor("SoundDashboard.SoundCueTemplateColor"),
@@ -452,11 +596,11 @@ namespace UE::Audio::Insights
 				MakeShared<FSoundDashboardFilter>(
 				ESoundDashboardFilterFlags::Pinned,
 				"Pinned",
-				LOCTEXT("SoundDashboard_Filter_PinnedNameText", "Pinned"),
+				PinnedCategoryName,
 				"AudioInsights.Icon.SoundDashboard.Pin",
 				FText::GetEmpty(),
 				FSlateStyle::Get().GetColor("SoundDashboard.PinnedColor"),
-				FilterCategory),
+				FilterCategory)
 
 				// @TODO UE-250399: Hide category pending to implement
 				//MakeShared<FSoundDashboardFilter>(
@@ -537,7 +681,7 @@ namespace UE::Audio::Insights
 			// Custom color for root item icons
 			const auto GetIconColor = [](const TSharedPtr<IDashboardDataTreeViewEntry>& InEntry)
 			{
-				const FSoundDashboardEntry& SoundDashboardEntry = FSoundDashboardViewFactoryPrivate::CastEntry(*InEntry.Get());
+				const FSoundDashboardEntry& SoundDashboardEntry = FSoundDashboardViewFactoryPrivate::CastEntry(*InEntry);
 
 				switch (SoundDashboardEntry.EntryType)
 				{
@@ -556,6 +700,9 @@ namespace UE::Audio::Insights
 				case ESoundDashboardEntryType::SoundCueTemplate:
 					return FSlateColor(FSlateStyle::Get().GetColor("SoundDashboard.SoundCueTemplateColor"));
 
+				case ESoundDashboardEntryType::Pinned:
+					return FSlateColor(FSlateStyle::Get().GetColor("SoundDashboard.PinnedColor"));
+
 				case ESoundDashboardEntryType::None:
 				default:
 					return FSlateColor(FColor::White);
@@ -602,7 +749,7 @@ namespace UE::Audio::Insights
 							{
 								if (InChild.IsValid())
 								{
-									const FSoundDashboardEntry& SoundDashboardEntry = FSoundDashboardViewFactoryPrivate::CastEntry(*InChild.Get());
+									const FSoundDashboardEntry& SoundDashboardEntry = FSoundDashboardViewFactoryPrivate::CastEntry(*InChild);
 
 									if (SoundDashboardEntry.TimeoutTimestamp == INVALID_TIMEOUT)
 									{
@@ -633,7 +780,17 @@ namespace UE::Audio::Insights
 
 		const bool bIsSoundCueType = SoundEntry.EntryType == ESoundDashboardEntryType::SoundCue || SoundEntry.EntryType == ESoundDashboardEntryType::SoundCueTemplate;
 
-		return IsRootItem(InEntry) || bIsSoundCueType;
+		return IsCategoryItem(InEntry) || bIsSoundCueType;
+	}
+
+	bool FSoundDashboardViewFactory::IsCategoryItem(const TSharedRef<IDashboardDataTreeViewEntry>& InEntry) const
+	{
+		if (IsRootItem(InEntry))
+		{
+			return true;
+		}
+
+		return PinnedItemEntries.IsValid() && PinnedItemEntries->GetPinnedSectionEntry().IsValid() && PinnedItemEntries->GetPinnedSectionEntry()->Children.Contains(InEntry);
 	}
 
 	TSharedRef<SWidget> FSoundDashboardViewFactory::GenerateWidgetForColumn(const TSharedRef<FTraceTreeDashboardViewFactory::SRowWidget>& InRowWidget, const TSharedRef<IDashboardDataTreeViewEntry>& InRowData, const FName& InColumn)
@@ -756,7 +913,7 @@ namespace UE::Audio::Insights
 			{
 				const FSoundDashboardEntry& SoundCategoryEntry = FSoundDashboardViewFactoryPrivate::CastEntry(*Item);
 
-				return SoundCategoryEntry.bIsVisible ? EVisibility::Visible : EVisibility::Hidden; 
+				return SoundCategoryEntry.bIsVisible && SoundCategoryEntry.PinnedEntryType != FSoundDashboardEntry::EPinnedEntryType::HiddenOriginalEntry ? EVisibility::Visible : EVisibility::Hidden;
 			});
 	}
 
@@ -824,11 +981,11 @@ namespace UE::Audio::Insights
 
 	TSharedPtr<SWidget> FSoundDashboardViewFactory::OnConstructContextMenu()
 	{
-		// Skip right click for root items
+		// Skip right click for category items
 		if (FilteredEntriesListView.IsValid() && FilteredEntriesListView->GetSelectedItems().Num() == 1)
 		{
 			const TSharedPtr<IDashboardDataTreeViewEntry> SelectedEntry = FilteredEntriesListView->GetSelectedItems()[0];
-			if (IsRootItem(SelectedEntry.ToSharedRef()))
+			if (IsCategoryItem(SelectedEntry.ToSharedRef()))
 			{
 				return SNullWidget::NullWidget;
 			}
@@ -1007,6 +1164,8 @@ namespace UE::Audio::Insights
 								return FName("AudioInsights.Icon.SoundDashboard.SoundWave");
 							case ESoundDashboardEntryType::SoundCueTemplate:
 								return FName("AudioInsights.Icon.SoundDashboard.SoundCue");
+							case ESoundDashboardEntryType::Pinned:
+								return FName("AudioInsights.Icon.SoundDashboard.Pin");
 							case ESoundDashboardEntryType::None:
 							default:
 								break;
@@ -1153,144 +1312,153 @@ namespace UE::Audio::Insights
 	{
 		using namespace FSoundDashboardViewFactoryPrivate;
 
-		if (SortByColumn == "Name")
+		UpdatePinnedSection();
+
+		auto SortByName = [](const FSoundDashboardEntry& First, const FSoundDashboardEntry& Second)
 		{
-			if (SortMode == EColumnSortMode::Ascending)
-			{
-				DataViewEntries.Sort([](const TSharedPtr<IDashboardDataTreeViewEntry>& A, const TSharedPtr<IDashboardDataTreeViewEntry>& B)
-				{
-					const FSoundDashboardEntry& AData = CastEntry(*A.Get());
-					const FSoundDashboardEntry& BData = CastEntry(*B.Get());
+			return First.GetDisplayName().CompareToCaseIgnored(Second.GetDisplayName()) <= 0;
+		};
 
-					return AData.GetDisplayName().CompareToCaseIgnored(BData.GetDisplayName()) < 0;
-				});
-			}
-			else if (SortMode == EColumnSortMode::Descending)
-			{
-				DataViewEntries.Sort([](const TSharedPtr<IDashboardDataTreeViewEntry>& A, const TSharedPtr<IDashboardDataTreeViewEntry>& B)
-				{
-					const FSoundDashboardEntry& AData = CastEntry(*A.Get());
-					const FSoundDashboardEntry& BData = CastEntry(*B.Get());
+		auto SortByPriority = [](const FSoundDashboardEntry& First, const FSoundDashboardEntry& Second)
+		{
+			return GetLastEntryArrayValue(First.PriorityDataPoints) <= GetLastEntryArrayValue(Second.PriorityDataPoints);
+		};
 
-					return BData.GetDisplayName().CompareToCaseIgnored(AData.GetDisplayName()) < 0;
-				});
-			}
-		}
-		else if (SortByColumn == "Priority")
+		auto SortByDistance = [](const FSoundDashboardEntry& First, const FSoundDashboardEntry& Second)
 		{
-			if (SortMode == EColumnSortMode::Ascending)
-			{
-				DataViewEntries.Sort([](const TSharedPtr<IDashboardDataTreeViewEntry>& A, const TSharedPtr<IDashboardDataTreeViewEntry>& B)
-				{
-					const FSoundDashboardEntry& AData = CastEntry(*A.Get());
-					const FSoundDashboardEntry& BData = CastEntry(*B.Get());
+			return GetLastEntryArrayValue(First.DistanceDataPoints) <= GetLastEntryArrayValue(Second.DistanceDataPoints);
+		};
 
-					return GetLastEntryArrayValue(AData.PriorityDataPoints) < GetLastEntryArrayValue(BData.PriorityDataPoints);
-				});
-			}
-			else if (SortMode == EColumnSortMode::Descending)
-			{
-				DataViewEntries.Sort([](const TSharedPtr<IDashboardDataTreeViewEntry>& A, const TSharedPtr<IDashboardDataTreeViewEntry>& B)
-				{
-					const FSoundDashboardEntry& AData = CastEntry(*A.Get());
-					const FSoundDashboardEntry& BData = CastEntry(*B.Get());
+		auto SortByAmplitude = [](const FSoundDashboardEntry& First, const FSoundDashboardEntry& Second)
+		{
+			return GetLastEntryArrayValue(First.AmplitudeDataPoints) <= GetLastEntryArrayValue(Second.AmplitudeDataPoints);
+		};
 
-					return GetLastEntryArrayValue(BData.PriorityDataPoints) < GetLastEntryArrayValue(AData.PriorityDataPoints);
-				});
-			}
+		auto SortByVolume = [](const FSoundDashboardEntry& First, const FSoundDashboardEntry& Second)
+		{
+			return GetLastEntryArrayValue(First.VolumeDataPoints) <= GetLastEntryArrayValue(Second.VolumeDataPoints);
+		};
+
+		auto SortByPitch = [](const FSoundDashboardEntry& First, const FSoundDashboardEntry& Second)
+		{
+			return GetLastEntryArrayValue(First.PitchDataPoints) <= GetLastEntryArrayValue(Second.PitchDataPoints);
+		};
+
+		if (SortByColumn == "Name")
+		{
+			SortByPredicate(SortByName, false/* bSortByChildren*/);
+		}
+		else if (SortByColumn == "Priority")
+		{
+			SortByPredicate(SortByPriority);
 		}
 		else if (SortByColumn == "Distance")
 		{
-			if (SortMode == EColumnSortMode::Ascending)
-			{
-				DataViewEntries.Sort([](const TSharedPtr<IDashboardDataTreeViewEntry>& A, const TSharedPtr<IDashboardDataTreeViewEntry>& B)
-				{
-					const FSoundDashboardEntry& AData = CastEntry(*A.Get());
-					const FSoundDashboardEntry& BData = CastEntry(*B.Get());
+			SortByPredicate(SortByDistance);
+		}
+		else if (SortByColumn == "Amplitude")
+		{
+			SortByPredicate(SortByAmplitude);
+		}
+		else if (SortByColumn == "Volume")
+		{
+			SortByPredicate(SortByVolume);
+		}
+		else if (SortByColumn == "Pitch")
+		{
+			SortByPredicate(SortByPitch);
+		}
 
-					return GetLastEntryArrayValue(AData.DistanceDataPoints) < GetLastEntryArrayValue(BData.DistanceDataPoints);
-				});
-			}
-			else if (SortMode == EColumnSortMode::Descending)
-			{
-				DataViewEntries.Sort([](const TSharedPtr<IDashboardDataTreeViewEntry>& A, const TSharedPtr<IDashboardDataTreeViewEntry>& B)
-				{
-					const FSoundDashboardEntry& AData = CastEntry(*A.Get());
-					const FSoundDashboardEntry& BData = CastEntry(*B.Get());
+		FullTree.Reset();
 
-					return GetLastEntryArrayValue(BData.DistanceDataPoints) < GetLastEntryArrayValue(AData.DistanceDataPoints);
-				});
-			}
+		if (PinnedItemEntries.IsValid())
+		{
+			FullTree.Add(PinnedItemEntries->GetPinnedSectionEntry());
 		}
-		else if (SortByColumn == "Amplitude")
+		
+		FullTree.Append(DataViewEntries);
+	}
+
+	bool FSoundDashboardViewFactory::ResetTreeData()
+	{
+		bool bDataReset = false;
+		if (!DataViewEntries.IsEmpty())
 		{
-			if (SortMode == EColumnSortMode::Ascending)
-			{
-				DataViewEntries.Sort([](const TSharedPtr<IDashboardDataTreeViewEntry>& A, const TSharedPtr<IDashboardDataTreeViewEntry>& B)
-				{
-					const FSoundDashboardEntry& AData = CastEntry(*A.Get());
-					const FSoundDashboardEntry& BData = CastEntry(*B.Get());
+			DataViewEntries.Empty();
+			bDataReset = true;
+		}
 
-					return GetLastEntryArrayValue(AData.AmplitudeDataPoints) < GetLastEntryArrayValue(BData.AmplitudeDataPoints);
-				});
-			}
-			else if (SortMode == EColumnSortMode::Descending)
-			{
-				DataViewEntries.Sort([](const TSharedPtr<IDashboardDataTreeViewEntry>& A, const TSharedPtr<IDashboardDataTreeViewEntry>& B)
-				{
-					const FSoundDashboardEntry& AData = CastEntry(*A.Get());
-					const FSoundDashboardEntry& BData = CastEntry(*B.Get());
+		if (PinnedItemEntries.IsValid())
+		{
+			PinnedItemEntries.Reset();
+			bDataReset = true;
+		}
 
-					return GetLastEntryArrayValue(BData.AmplitudeDataPoints) < GetLastEntryArrayValue(AData.AmplitudeDataPoints);
-				});
-			}
+		if (!FullTree.IsEmpty())
+		{
+			FullTree.Empty();
+			bDataReset = true;
 		}
-		else if (SortByColumn == "Volume")
+
+		return bDataReset;
+	}
+
+	void FSoundDashboardViewFactory::RecursiveSort(TArray<TSharedPtr<IDashboardDataTreeViewEntry>>& OutTree, TFunctionRef<bool(const FSoundDashboardEntry&, const FSoundDashboardEntry&)> Predicate, const bool bSortByChildren)
+	{
+		using namespace FSoundDashboardViewFactoryPrivate;
+
+		for (TSharedPtr<IDashboardDataTreeViewEntry>& Entry : OutTree)
 		{
-			if (SortMode == EColumnSortMode::Ascending)
+			if (Entry->Children.Num() > 0)
 			{
-				DataViewEntries.Sort([](const TSharedPtr<IDashboardDataTreeViewEntry>& A, const TSharedPtr<IDashboardDataTreeViewEntry>& B)
-				{
-					const FSoundDashboardEntry& AData = CastEntry(*A.Get());
-					const FSoundDashboardEntry& BData = CastEntry(*B.Get());
-
-					return GetLastEntryArrayValue(AData.VolumeDataPoints) < GetLastEntryArrayValue(BData.VolumeDataPoints);
-				});
+				RecursiveSort(Entry->Children, Predicate, bSortByChildren);
 			}
-			else if (SortMode == EColumnSortMode::Descending)
-			{
-				DataViewEntries.Sort([](const TSharedPtr<IDashboardDataTreeViewEntry>& A, const TSharedPtr<IDashboardDataTreeViewEntry>& B)
-				{
-					const FSoundDashboardEntry& AData = CastEntry(*A.Get());
-					const FSoundDashboardEntry& BData = CastEntry(*B.Get());
+		}
 
-					return GetLastEntryArrayValue(BData.VolumeDataPoints) < GetLastEntryArrayValue(AData.VolumeDataPoints);
-				});
+		for (const TSharedPtr<IDashboardDataTreeViewEntry>& Entry : OutTree)
+		{
+			const FSoundDashboardEntry& EntryData = CastEntry(*Entry);
+			if (IsCategoryItem(Entry.ToSharedRef()))
+			{
+				return;
 			}
 		}
-		else if (SortByColumn == "Pitch")
+
+		auto SortDashboardEntries = [this, bSortByChildren](const TSharedPtr<IDashboardDataTreeViewEntry>& First, const TSharedPtr<IDashboardDataTreeViewEntry>& Second, TFunctionRef<bool(const FSoundDashboardEntry&, const FSoundDashboardEntry&)> Predicate)
 		{
-			if (SortMode == EColumnSortMode::Ascending)
-			{
-				DataViewEntries.Sort([](const TSharedPtr<IDashboardDataTreeViewEntry>& A, const TSharedPtr<IDashboardDataTreeViewEntry>& B)
-				{
-					const FSoundDashboardEntry& AData = CastEntry(*A.Get());
-					const FSoundDashboardEntry& BData = CastEntry(*B.Get());
+			const TSharedPtr<IDashboardDataTreeViewEntry>& FirstEntryToSortBy = bSortByChildren && (First->Children.Num() > 0) ? First->Children[0] : First;
+			const TSharedPtr<IDashboardDataTreeViewEntry>& SecondtEntryToSortBy = bSortByChildren && (Second->Children.Num() > 0) ? Second->Children[0] : Second;
 
-					return GetLastEntryArrayValue(AData.PitchDataPoints) < GetLastEntryArrayValue(BData.PitchDataPoints);
-				});
-			}
-			else if (SortMode == EColumnSortMode::Descending)
+			const FSoundDashboardEntry& FirstData = CastEntry(*FirstEntryToSortBy);
+			const FSoundDashboardEntry& SecondData = CastEntry(*SecondtEntryToSortBy);
+
+			return Predicate(FirstData, SecondData);
+		};
+
+		if (SortMode == EColumnSortMode::Ascending)
+		{
+			OutTree.Sort([&SortDashboardEntries, &Predicate](const TSharedPtr<IDashboardDataTreeViewEntry>& A, const TSharedPtr<IDashboardDataTreeViewEntry>& B)
 			{
-				DataViewEntries.Sort([](const TSharedPtr<IDashboardDataTreeViewEntry>& A, const TSharedPtr<IDashboardDataTreeViewEntry>& B)
-				{
-					const FSoundDashboardEntry& AData = CastEntry(*A.Get());
-					const FSoundDashboardEntry& BData = CastEntry(*B.Get());
+				return SortDashboardEntries(A, B, Predicate);
+			});
+		}
+		else if (SortMode == EColumnSortMode::Descending)
+		{
+			OutTree.Sort([&SortDashboardEntries, &Predicate](const TSharedPtr<IDashboardDataTreeViewEntry>& A, const TSharedPtr<IDashboardDataTreeViewEntry>& B)
+			{
+				return SortDashboardEntries(B, A, Predicate);
+			});
+		}
+	}
 
-					return GetLastEntryArrayValue(BData.PitchDataPoints) < GetLastEntryArrayValue(AData.PitchDataPoints);
-				});
-			}
+	void FSoundDashboardViewFactory::SortByPredicate(TFunctionRef<bool(const FSoundDashboardEntry&, const FSoundDashboardEntry&)> Predicate, const bool bSortByChildren /*= true*/)
+	{
+		if (PinnedItemEntries.IsValid())
+		{
+			RecursiveSort(PinnedItemEntries->GetPinnedSectionEntry()->Children, Predicate, bSortByChildren);
 		}
+
+		RecursiveSort(DataViewEntries, Predicate, bSortByChildren);
 	}
 
 #if WITH_EDITOR
@@ -1399,7 +1567,148 @@ namespace UE::Audio::Insights
 
 	void FSoundDashboardViewFactory::PinSound()
 	{
-		// @TODO UE-220744: pending to implement
+		if (FilteredEntriesListView.IsValid())
+		{
+			const TArray<TSharedPtr<IDashboardDataTreeViewEntry>> SelectedItems = FilteredEntriesListView->GetSelectedItems();
+
+			for (TSharedPtr<IDashboardDataTreeViewEntry>& Entry : DataViewEntries)
+			{
+				PinSelectedItems(Entry, SelectedItems);
+			}
+		}
+	}
+
+	void FSoundDashboardViewFactory::PinSelectedItems(const TSharedPtr<IDashboardDataTreeViewEntry>& Entry, const TArray<TSharedPtr<IDashboardDataTreeViewEntry>>& SelectedItems)
+	{
+		for (const TSharedPtr<IDashboardDataTreeViewEntry>& Child : Entry->Children)
+		{
+			if (SelectedItems.Contains(Child))
+			{
+				if (IsRootItem(Entry.ToSharedRef()))
+				{
+					MarkBranchAsPinned(Child);
+					CreatePinnedEntry(Child);
+				}
+				else
+				{
+					MarkBranchAsPinned(Entry);
+					CreatePinnedEntry(Entry);
+				}
+			}
+			else
+			{
+				PinSelectedItems(Child, SelectedItems);
+			}
+		}
+	}
+
+	void FSoundDashboardViewFactory::MarkBranchAsPinned(const TSharedPtr<IDashboardDataTreeViewEntry> Entry)
+	{
+		FSoundDashboardEntry& SoundEntry = FSoundDashboardViewFactoryPrivate::CastEntry(*Entry);
+		SoundEntry.PinnedEntryType = FSoundDashboardEntry::EPinnedEntryType::HiddenOriginalEntry;
+
+		for (TSharedPtr<IDashboardDataTreeViewEntry> Child : Entry->Children)
+		{
+			MarkBranchAsPinned(Child);
+		}
+	}
+
+	void FSoundDashboardViewFactory::InitPinnedItemEntries()
+	{
+		if (PinnedItemEntries.IsValid())
+		{
+			return;
+		}
+
+		TSharedPtr<FSoundDashboardEntry> PinnedCategory = MakeShared<FSoundDashboardEntry>();
+		PinnedCategory->Name = FSoundDashboardViewFactoryPrivate::PinnedCategoryName.ToString();
+		PinnedCategory->EntryType = ESoundDashboardEntryType::Pinned;
+		PinnedCategory->PinnedEntryType = FSoundDashboardEntry::EPinnedEntryType::PinnedCopy;
+
+		PinnedItemEntries = MakeShared<FPinnedSoundEntryWrapper>(PinnedCategory, true /*bIsCategoryEntry*/);
+	}
+
+	void FSoundDashboardViewFactory::CreatePinnedEntry(TSharedPtr<IDashboardDataTreeViewEntry> Entry)
+	{
+		using namespace FSoundDashboardViewFactoryPrivate;
+
+		// If we have at least one entry that is pinned, ensure the pinned section has been created
+		// The pinned area will delete itself once empty
+		InitPinnedItemEntries();
+
+		const FSoundDashboardEntry& SoundEntry = CastEntry(*Entry);
+
+		// Check if category is already in the list, if so we need to merge
+		bool bFoundExistingCategory = false;
+		for (TSharedPtr<FPinnedSoundEntryWrapper> PinnedCategoryEntry : PinnedItemEntries->PinnedWrapperChildren)
+		{
+			if (!PinnedCategoryEntry.IsValid())
+			{
+				continue;
+			}
+
+			const FSoundDashboardEntry& PinnedCategorySoundDashboardEntry = CastEntry(*PinnedCategoryEntry->GetPinnedSectionEntry());
+
+			if (PinnedCategorySoundDashboardEntry.EntryType == SoundEntry.EntryType)
+			{
+				bFoundExistingCategory = true;
+
+				const TSharedPtr<IDashboardDataTreeViewEntry>* FoundExisingEntry = PinnedCategorySoundDashboardEntry.Children.FindByPredicate([&SoundEntry](const TSharedPtr<IDashboardDataTreeViewEntry> PinnedEntry)
+				{
+					if (!PinnedEntry.IsValid())
+					{
+						return false;
+					}
+
+					const FSoundDashboardEntry& ExistingPinnedEntry = CastEntry(*PinnedEntry);
+					return ExistingPinnedEntry.PlayOrder == SoundEntry.PlayOrder;
+				});
+
+				// If we didn't find this entry already inside the pinned area, add it here
+				if (FoundExisingEntry == nullptr)
+				{
+					PinnedCategoryEntry->AddChildEntry(Entry, false /*bIsCategory*/);
+				}
+				break;
+			}
+		}
+
+		// if we haven't found an existing pinned category, create a new one and add this item
+		if (!bFoundExistingCategory)
+		{
+			for (const TSharedPtr<IDashboardDataTreeViewEntry>& DataCategoryEntry : DataViewEntries)
+			{
+				if (SoundEntry.EntryType == CastEntry(*DataCategoryEntry).EntryType)
+				{
+					TSharedPtr<FSoundDashboardEntry> NewPinnedCategoryEntry = MakeShared<FSoundDashboardEntry>(CastEntry(*DataCategoryEntry));
+					NewPinnedCategoryEntry->Children.Reset();
+
+					TSharedPtr<FPinnedSoundEntryWrapper> PinnedCategory = PinnedItemEntries->AddChildEntry(NewPinnedCategoryEntry, true /*bIsCategory*/);
+					PinnedCategory->AddChildEntry(Entry, false /*bIsCategory*/);
+
+					break;
+				}
+			}
+		}
+	}
+
+	void FSoundDashboardViewFactory::UpdatePinnedSection()
+	{
+		if (!PinnedItemEntries.IsValid())
+		{
+			return;
+		}
+
+		PinnedItemEntries->CleanUp();
+
+		if (FPinnedSoundEntryWrapperPrivate::CanBeDeleted(PinnedItemEntries))
+		{
+			PinnedItemEntries.Reset();
+		}
+		else
+		{
+			PinnedItemEntries->UpdateParams();
+		}
 	}
 
 #if WITH_EDITOR
diff --git a/Engine/Plugins/AudioInsights/Source/AudioInsights/Private/Views/TreeDashboardViewFactory.cpp b/Engine/Plugins/AudioInsights/Source/AudioInsights/Private/Views/TreeDashboardViewFactory.cpp
--- a/Engine/Plugins/AudioInsights/Source/AudioInsights/Private/Views/TreeDashboardViewFactory.cpp
+++ b/Engine/Plugins/AudioInsights/Source/AudioInsights/Private/Views/TreeDashboardViewFactory.cpp
@@ -115,7 +115,7 @@ namespace UE::Audio::Insights
 				[
 					SAssignNew(FilteredEntriesListView, STreeView<TSharedPtr<IDashboardDataTreeViewEntry>>)
 					.HeaderRow(MakeHeaderRowWidget())
-					.TreeItemsSource(&DataViewEntries)
+					.TreeItemsSource(&GetTreeItemsSource())
 					.OnGenerateRow(this, &FTraceTreeDashboardViewFactory::OnGenerateRow)
 					.OnContextMenuOpening(this, &FTraceTreeDashboardViewFactory::OnConstructContextMenu)
 					.SelectionMode(ESelectionMode::Multi)
diff --git a/Engine/Plugins/AudioInsights/Source/AudioInsights/Public/Messages/SoundTraceMessages.h b/Engine/Plugins/AudioInsights/Source/AudioInsights/Public/Messages/SoundTraceMessages.h
--- a/Engine/Plugins/AudioInsights/Source/AudioInsights/Public/Messages/SoundTraceMessages.h
+++ b/Engine/Plugins/AudioInsights/Source/AudioInsights/Public/Messages/SoundTraceMessages.h
@@ -25,7 +25,8 @@ namespace UE::Audio::Insights
 		SoundCue,
 		ProceduralSource,
 		SoundWave,
-		SoundCueTemplate
+		SoundCueTemplate,
+		Pinned
 	};
 
 	// Trace messages
@@ -154,13 +155,20 @@ namespace UE::Audio::Insights
 		virtual bool ShouldInitExpandChildren() const override { return bShouldForceExpandChildren; }
 		virtual void ResetShouldInitExpandChildren() override { bShouldForceExpandChildren = false; }
 
+		enum class EPinnedEntryType : uint8
+		{
+			None = 0,
+			PinnedCopy,
+			HiddenOriginalEntry
+		};
+
 		::Audio::FDeviceId DeviceId = INDEX_NONE;
 		uint32 PlayOrder = INDEX_NONE;
 		double Timestamp = 0.0;
 		double TimeoutTimestamp = INVALID_TIMEOUT;
 		FString Name;
 		ESoundDashboardEntryType EntryType = ESoundDashboardEntryType::None;
-
+		EPinnedEntryType PinnedEntryType = EPinnedEntryType::None;
 		bool bShouldForceExpandChildren = true;
 		bool bIsVisible = true;
 
diff --git a/Engine/Plugins/AudioInsights/Source/AudioInsights/Public/Views/SoundDashboardViewFactory.h b/Engine/Plugins/AudioInsights/Source/AudioInsights/Public/Views/SoundDashboardViewFactory.h
--- a/Engine/Plugins/AudioInsights/Source/AudioInsights/Public/Views/SoundDashboardViewFactory.h
+++ b/Engine/Plugins/AudioInsights/Source/AudioInsights/Public/Views/SoundDashboardViewFactory.h
@@ -9,6 +9,8 @@ class FUICommandList;
 
 namespace UE::Audio::Insights
 {
+	class FSoundDashboardEntry;
+
 	enum class ESoundDashboardFilterFlags : uint32
 	{
 		None = 0,
@@ -60,6 +62,35 @@ namespace UE::Audio::Insights
 		bool bIsActive = false;
 	};
 
+	/**
+	* Helper class for pinned items in the dashboard tree
+	*	- Contains a weak handle to the original entry (OriginalDataEntry) which is updated from the trace provider
+	*	- Copies updated params to PinnedSectionEntry for display
+	*/
+	class AUDIOINSIGHTS_API FPinnedSoundEntryWrapper
+	{
+	public:
+		FPinnedSoundEntryWrapper() = delete;
+		FPinnedSoundEntryWrapper(const TSharedPtr<IDashboardDataTreeViewEntry>& OriginalEntry, const bool bIsCategory);
+
+		TSharedPtr<IDashboardDataTreeViewEntry> GetPinnedSectionEntry() const { return PinnedSectionEntry; }
+
+		TSharedPtr<FPinnedSoundEntryWrapper> AddChildEntry(const TSharedPtr<IDashboardDataTreeViewEntry> Child, const bool bIsCategory);
+
+		void UpdateParams();
+
+		void CleanUp();
+
+		bool EntryIsValid() const;
+
+		TArray<TSharedPtr<FPinnedSoundEntryWrapper>> PinnedWrapperChildren;
+		const bool bIsCategoryEntry = false;
+
+	private:
+		TSharedPtr<IDashboardDataTreeViewEntry> PinnedSectionEntry;
+		TWeakPtr<IDashboardDataTreeViewEntry> OriginalDataEntry;
+	};
+
 	class AUDIOINSIGHTS_API FSoundDashboardViewFactory : public FTraceTreeDashboardViewFactory
 	{
 	public:
@@ -81,6 +112,7 @@ namespace UE::Audio::Insights
 
 		bool IsRootItem(const TSharedRef<IDashboardDataTreeViewEntry>& InEntry) const;
 		bool EntryCanHaveChildren(const TSharedRef<IDashboardDataTreeViewEntry>& InEntry) const;
+		bool IsCategoryItem(const TSharedRef<IDashboardDataTreeViewEntry>& InEntry) const;
 
 		TSharedRef<SWidget> GenerateWidgetForRootColumn(const TSharedRef<FTraceTreeDashboardViewFactory::SRowWidget>& InRowWidget, const TSharedRef<IDashboardDataTreeViewEntry>& InRowData, const FName& InColumn, const FText& InValueText);
 		virtual TSharedRef<SWidget> GenerateWidgetForColumn(const TSharedRef<FTraceTreeDashboardViewFactory::SRowWidget>& InRowWidget, const TSharedRef<IDashboardDataTreeViewEntry>& InRowData, const FName& InColumn) override;
@@ -93,6 +125,11 @@ namespace UE::Audio::Insights
 		virtual const TMap<FName, FTraceTreeDashboardViewFactory::FColumnData>& GetColumns() const override;
 
 		virtual void SortTable() override;
+		virtual TArray<TSharedPtr<IDashboardDataTreeViewEntry>>& GetTreeItemsSource() override { return FullTree; }
+		virtual bool ResetTreeData() override;
+
+		void RecursiveSort(TArray<TSharedPtr<IDashboardDataTreeViewEntry>>& OutTree, TFunctionRef<bool(const FSoundDashboardEntry&, const FSoundDashboardEntry&)> Predicate, const bool bSortByChildren);
+		void SortByPredicate(TFunctionRef<bool(const FSoundDashboardEntry&, const FSoundDashboardEntry&)> Predicate, const bool bSortByChildren = true);
 
 #if WITH_EDITOR
 		TSharedRef<SWidget> MakeMuteSoloWidget();
@@ -108,6 +145,11 @@ namespace UE::Audio::Insights
 #endif // WITH_EDITOR
 
 		void PinSound();
+		void PinSelectedItems(const TSharedPtr<IDashboardDataTreeViewEntry>& Entry, const TArray<TSharedPtr<IDashboardDataTreeViewEntry>>& SelectedItems);
+		void MarkBranchAsPinned(const TSharedPtr<IDashboardDataTreeViewEntry> Entry);
+		void InitPinnedItemEntries();
+		void CreatePinnedEntry(TSharedPtr<IDashboardDataTreeViewEntry> Entry);
+		void UpdatePinnedSection();
 
 #if WITH_EDITOR
 		void BrowseSoundAsset() const;
@@ -116,6 +158,9 @@ namespace UE::Audio::Insights
 		// @TODO UE-250399: Hide category pending to implement
 		//void HideSound();
 
+		TSharedPtr<FPinnedSoundEntryWrapper> PinnedItemEntries;
+		TArray<TSharedPtr<IDashboardDataTreeViewEntry>> FullTree;
+
 		TSharedPtr<FUICommandList> CommandList;
 		TSharedPtr<SWidget> SoundsFilterBar;
 		TSharedPtr<SWidget> SoundsFilterBarButton;
diff --git a/Engine/Plugins/AudioInsights/Source/AudioInsights/Public/Views/TreeDashboardViewFactory.h b/Engine/Plugins/AudioInsights/Source/AudioInsights/Public/Views/TreeDashboardViewFactory.h
--- a/Engine/Plugins/AudioInsights/Source/AudioInsights/Public/Views/TreeDashboardViewFactory.h
+++ b/Engine/Plugins/AudioInsights/Source/AudioInsights/Public/Views/TreeDashboardViewFactory.h
@@ -100,6 +100,19 @@ namespace UE::Audio::Insights
 		virtual void ProcessEntries(EProcessReason Reason) = 0;
 		virtual void SortTable() = 0;
 
+		virtual TArray<TSharedPtr<IDashboardDataTreeViewEntry>>& GetTreeItemsSource() { return DataViewEntries; }
+
+		virtual bool ResetTreeData()
+		{
+			if (!DataViewEntries.IsEmpty())
+			{
+				DataViewEntries.Empty();
+				return true;
+			}
+
+			return false;
+		}
+
 		void Tick(float InElapsed);
 
 #if WITH_EDITOR
@@ -137,11 +150,7 @@ namespace UE::Audio::Insights
 				}
 				else
 				{
-					if (!DataViewEntries.IsEmpty())
-					{
-						DataViewEntries.Empty();
-						return true;
-					}
+					return ResetTreeData();
 				}
 			}
 
