[Audio Insights] Home Screen - update count next to category names to not include hidden pinned items. Hide the row if this count reaches zero.

diff --git a/Engine/Plugins/AudioInsights/Source/AudioInsights/Private/Views/SoundDashboardViewFactory.cpp b/Engine/Plugins/AudioInsights/Source/AudioInsights/Private/Views/SoundDashboardViewFactory.cpp
--- a/Engine/Plugins/AudioInsights/Source/AudioInsights/Private/Views/SoundDashboardViewFactory.cpp
+++ b/Engine/Plugins/AudioInsights/Source/AudioInsights/Private/Views/SoundDashboardViewFactory.cpp
@@ -90,6 +90,58 @@ namespace UE::Audio::Insights
 			}
 		}
 
+		bool HasPinEntryType(const IDashboardDataTreeViewEntry& InEntry, const FSoundDashboardEntry::EPinnedEntryType PinnedEntryType)
+		{
+			const FSoundDashboardEntry& SoundEntry = CastEntry(InEntry);
+			return SoundEntry.PinnedEntryType == PinnedEntryType;
+		}
+
+		int32 GetNumChildrenWithoutPinEntryType(const IDashboardDataTreeViewEntry& InEntry, const FSoundDashboardEntry::EPinnedEntryType ExcludedPinnedEntryType)
+		{
+			int32 NumChildrenWithoutType = 0;
+			for (const TSharedPtr<IDashboardDataTreeViewEntry>& Child : InEntry.Children)
+			{
+				if (!Child.IsValid())
+				{
+					continue;
+				}
+
+				if (!HasPinEntryType(*Child, ExcludedPinnedEntryType))
+				{
+					++NumChildrenWithoutType;
+				}
+			}
+
+			return NumChildrenWithoutType;
+		}
+
+		int32 CountNumChildren(const IDashboardDataTreeViewEntry& InEntry)
+		{
+			const uint32 TotalNumChildren = Algo::Accumulate(InEntry.Children, 0,
+				[](uint32 Accum, TSharedPtr<IDashboardDataTreeViewEntry> InChild)
+				{
+					if (InChild.IsValid())
+					{
+						const FSoundDashboardEntry& SoundDashboardEntry = FSoundDashboardViewFactoryPrivate::CastEntry(*InChild);
+
+						if (SoundDashboardEntry.TimeoutTimestamp == INVALID_TIMEOUT)
+						{
+							if (HasPinEntryType(*InChild, FSoundDashboardEntry::EPinnedEntryType::HiddenOriginalEntry))
+							{
+								return Accum;
+							}
+
+							const int32 NumNestedChildren = GetNumChildrenWithoutPinEntryType(*InChild, FSoundDashboardEntry::EPinnedEntryType::HiddenOriginalEntry);
+							return Accum + (NumNestedChildren > 0 ? NumNestedChildren : 1);
+						}
+					}
+
+					return Accum;
+				});
+
+			return TotalNumChildren;
+		}
+
 #if WITH_EDITOR
 		void SetMuteSolo(const IDashboardDataTreeViewEntry& InEntry, const EMuteSoloMode InMuteSoloMode, const bool bInOnOff)
 		{
@@ -672,6 +724,8 @@ namespace UE::Audio::Insights
 
 	TSharedRef<SWidget> FSoundDashboardViewFactory::GenerateWidgetForRootColumn(const TSharedRef<FTraceTreeDashboardViewFactory::SRowWidget>& InRowWidget, const TSharedRef<IDashboardDataTreeViewEntry>& InRowData, const FName& InColumn, const FText& InValueText)
 	{
+		using namespace FSoundDashboardViewFactoryPrivate;
+
 		const FColumnData& ColumnData = GetColumns()[InColumn];
 
 		if (InColumn == "Name")
@@ -744,23 +798,7 @@ namespace UE::Audio::Insights
 					.ColorAndOpacity(FSlateColor(FColor::White.WithAlpha(128)))
 					.Text_Lambda([InRowData]()
 					{
-						const uint32 TotalNumChildren = Algo::Accumulate(InRowData->Children, 0,
-							[](uint32 Accum, TSharedPtr<IDashboardDataTreeViewEntry> InChild)
-							{
-								if (InChild.IsValid())
-								{
-									const FSoundDashboardEntry& SoundDashboardEntry = FSoundDashboardViewFactoryPrivate::CastEntry(*InChild);
-
-									if (SoundDashboardEntry.TimeoutTimestamp == INVALID_TIMEOUT)
-									{
-										const int32 NumNestedChildren = InChild->Children.Num();
-										return Accum + (NumNestedChildren > 0 ? NumNestedChildren : 1);
-									}
-								}
-
-								return Accum;
-							});
-
+						const uint32 TotalNumChildren = CountNumChildren(*InRowData);
 						return FText::FromString("(" + FString::FromInt(TotalNumChildren) + ")");
 					})
 				];
@@ -899,7 +937,8 @@ namespace UE::Audio::Insights
 						  .ColorAndOpacity(FSlateColor(FColor::White.WithAlpha(128)))
 						  .Text_Lambda([InRowData]()
 						  {
-							  return FText::FromString("(" + FString::FromInt(InRowData->Children.Num()) + ")");
+							  const uint32 TotalNumChildren = CountNumChildren(*InRowData);
+							  return FText::FromString("(" + FString::FromInt(TotalNumChildren) + ")");
 						  })
 						: SNullWidget::NullWidget
 				];
@@ -908,12 +947,19 @@ namespace UE::Audio::Insights
 
 	TSharedRef<ITableRow> FSoundDashboardViewFactory::OnGenerateRow(TSharedPtr<IDashboardDataTreeViewEntry> Item, const TSharedRef<STableViewBase>& OwnerTable)
 	{
+		using namespace FSoundDashboardViewFactoryPrivate;
+
 		return SNew(SRowWidget, OwnerTable, Item, AsShared())
-			.Visibility_Lambda([Item]()
+			.Visibility_Lambda([this, Item]()
 			{
-				const FSoundDashboardEntry& SoundCategoryEntry = FSoundDashboardViewFactoryPrivate::CastEntry(*Item);
+				const FSoundDashboardEntry& SoundCategoryEntry = CastEntry(*Item);
+				const int32 NumUnpinnedChildren = GetNumChildrenWithoutPinEntryType(*Item, FSoundDashboardEntry::EPinnedEntryType::HiddenOriginalEntry);
+
+				const bool bRowShouldBeVisible = SoundCategoryEntry.bIsVisible 
+												&& SoundCategoryEntry.PinnedEntryType != FSoundDashboardEntry::EPinnedEntryType::HiddenOriginalEntry 
+												&& (!IsCategoryItem(Item.ToSharedRef()) || (NumUnpinnedChildren > 0));
 
-				return SoundCategoryEntry.bIsVisible && SoundCategoryEntry.PinnedEntryType != FSoundDashboardEntry::EPinnedEntryType::HiddenOriginalEntry ? EVisibility::Visible : EVisibility::Hidden;
+				return bRowShouldBeVisible ? EVisibility::Visible : EVisibility::Hidden;
 			});
 	}
 
