Home Screen - Add option to Unpin sounds in the sounds dashboard

diff --git a/Engine/Plugins/AudioInsights/Source/AudioInsights/Private/SoundDashboardCommands.cpp b/Engine/Plugins/AudioInsights/Source/AudioInsights/Private/SoundDashboardCommands.cpp
--- a/Engine/Plugins/AudioInsights/Source/AudioInsights/Private/SoundDashboardCommands.cpp
+++ b/Engine/Plugins/AudioInsights/Source/AudioInsights/Private/SoundDashboardCommands.cpp
@@ -18,6 +18,7 @@ namespace UE::Audio::Insights
 	void FSoundDashboardCommands::RegisterCommands()
 	{
 		UI_COMMAND(Pin, "Pin", "Pins the selected sound in the Pinned category.", EUserInterfaceActionType::Button, FInputChord(EKeys::P, EModifierKey::Control));
+		UI_COMMAND(Unpin, "Unpin", "Removes the selected sound from the Pinned category.", EUserInterfaceActionType::Button, FInputChord(EKeys::P, EModifierKey::Control));
 		UI_COMMAND(Browse, "Browse To Asset", "Browses to the selected sound asset in the content browser.", EUserInterfaceActionType::Button, FInputChord(EKeys::B, EModifierKey::Control));
 		UI_COMMAND(Edit, "Edit", "Opens the selected sound for edit.", EUserInterfaceActionType::Button, FInputChord(EKeys::E, EModifierKey::Control));
 		// @TODO UE-250399: Hide category pending to implement
diff --git a/Engine/Plugins/AudioInsights/Source/AudioInsights/Private/SoundDashboardCommands.h b/Engine/Plugins/AudioInsights/Source/AudioInsights/Private/SoundDashboardCommands.h
--- a/Engine/Plugins/AudioInsights/Source/AudioInsights/Private/SoundDashboardCommands.h
+++ b/Engine/Plugins/AudioInsights/Source/AudioInsights/Private/SoundDashboardCommands.h
@@ -16,6 +16,7 @@ namespace UE::Audio::Insights
 		virtual void RegisterCommands() override;
 
 		TSharedPtr<const FUICommandInfo> GetPinCommand() const    { return Pin; }
+		TSharedPtr<const FUICommandInfo> GetUnpinCommand() const  { return Unpin; }
 		TSharedPtr<const FUICommandInfo> GetBrowseCommand() const { return Browse; }
 		TSharedPtr<const FUICommandInfo> GetEditCommand() const   { return Edit; }
 		// @TODO UE-250399: Hide category pending to implement
@@ -23,6 +24,7 @@ namespace UE::Audio::Insights
 		
 	private:
 		TSharedPtr<FUICommandInfo> Pin;
+		TSharedPtr<FUICommandInfo> Unpin;
 		TSharedPtr<FUICommandInfo> Browse;
 		TSharedPtr<FUICommandInfo> Edit;
 		// @TODO UE-250399: Hide category pending to implement
diff --git a/Engine/Plugins/AudioInsights/Source/AudioInsights/Private/Views/SoundDashboardViewFactory.cpp b/Engine/Plugins/AudioInsights/Source/AudioInsights/Private/Views/SoundDashboardViewFactory.cpp
--- a/Engine/Plugins/AudioInsights/Source/AudioInsights/Private/Views/SoundDashboardViewFactory.cpp
+++ b/Engine/Plugins/AudioInsights/Source/AudioInsights/Private/Views/SoundDashboardViewFactory.cpp
@@ -435,6 +435,11 @@ namespace UE::Audio::Insights
 		}
 	}
 
+	void FPinnedSoundEntryWrapper::MarkToDelete()
+	{
+		OriginalDataEntry.Reset();
+	}
+
 	bool FPinnedSoundEntryWrapper::EntryIsValid() const
 	{
 		return PinnedSectionEntry.IsValid() && (OriginalDataEntry.IsValid() || bIsCategoryEntry);
@@ -492,6 +497,7 @@ namespace UE::Audio::Insights
 		const FSoundDashboardCommands& Commands = FSoundDashboardCommands::Get();
 
 		CommandList->MapAction(Commands.GetPinCommand(), FExecuteAction::CreateRaw(this, &FSoundDashboardViewFactory::PinSound));
+		CommandList->MapAction(Commands.GetUnpinCommand(), FExecuteAction::CreateRaw(this, &FSoundDashboardViewFactory::UnpinSound));
 #if WITH_EDITOR
 		CommandList->MapAction(Commands.GetBrowseCommand(), FExecuteAction::CreateRaw(this, &FSoundDashboardViewFactory::BrowseSoundAsset));
 		CommandList->MapAction(Commands.GetEditCommand(),   FExecuteAction::CreateRaw(this, &FSoundDashboardViewFactory::OpenSoundAsset));
@@ -1045,7 +1051,21 @@ namespace UE::Audio::Insights
 		MenuBuilder.BeginSection("SoundDashboardActions", LOCTEXT("SoundDashboard_Actions_HeaderText", "Sound Options"));
 
 		{
-			MenuBuilder.AddMenuEntry(Commands.GetPinCommand(), NAME_None, TAttribute<FText>(), TAttribute<FText>(), FSlateStyle::Get().CreateIcon("AudioInsights.Icon.SoundDashboard.Pin"));
+			MenuBuilder.AddMenuEntry(Commands.GetPinCommand()
+									, NAME_None
+									, TAttribute<FText>()
+									, TAttribute<FText>()
+									, FSlateStyle::Get().CreateIcon("AudioInsights.Icon.SoundDashboard.Pin")
+									, NAME_None
+									, TAttribute<EVisibility>::CreateLambda([this]() { return SelectionIncludesUnpinnedItem() ? EVisibility::Visible : EVisibility::Collapsed; }));
+			
+			MenuBuilder.AddMenuEntry(Commands.GetUnpinCommand()
+									, NAME_None
+									, TAttribute<FText>()
+									, TAttribute<FText>()
+									, FSlateStyle::Get().CreateIcon("AudioInsights.Icon.SoundDashboard.Pin")
+									, NAME_None
+									, TAttribute<EVisibility>::CreateLambda([this]() { return SelectionIncludesUnpinnedItem() ? EVisibility::Collapsed : EVisibility::Visible; }));
 #if WITH_EDITOR
 			MenuBuilder.AddMenuEntry(Commands.GetBrowseCommand(), NAME_None, TAttribute<FText>(), TAttribute<FText>(), FSlateStyle::Get().CreateIcon("AudioInsights.Icon.SoundDashboard.Browse"));
 			MenuBuilder.AddMenuEntry(Commands.GetEditCommand(),   NAME_None, TAttribute<FText>(), TAttribute<FText>(), FSlateStyle::Get().CreateIcon("AudioInsights.Icon.SoundDashboard.Edit"));
@@ -1621,9 +1641,53 @@ namespace UE::Audio::Insights
 			{
 				PinSelectedItems(Entry, SelectedItems);
 			}
+
+			FilteredEntriesListView->ClearSelection();
 		}
 	}
 
+	void FSoundDashboardViewFactory::UnpinSound()
+	{
+		if (FilteredEntriesListView.IsValid() && PinnedItemEntries.IsValid())
+		{
+			const TArray<TSharedPtr<IDashboardDataTreeViewEntry>> SelectedItems = FilteredEntriesListView->GetSelectedItems();
+
+			for (const TSharedPtr<FPinnedSoundEntryWrapper>& PinnedEntry : PinnedItemEntries->PinnedWrapperChildren)
+			{
+				UnpinSelectedItems(PinnedEntry, SelectedItems);
+			}
+
+			FilteredEntriesListView->ClearSelection();
+		}
+	}
+
+	bool FSoundDashboardViewFactory::SelectionIncludesUnpinnedItem() const
+	{
+		using namespace FSoundDashboardViewFactoryPrivate;
+
+		if (!FilteredEntriesListView.IsValid())
+		{
+			return false;
+		}
+
+		const TArray<TSharedPtr<IDashboardDataTreeViewEntry>> SelectedItems = FilteredEntriesListView->GetSelectedItems();
+		for (const TSharedPtr<IDashboardDataTreeViewEntry>& SelectedItem : SelectedItems)
+		{
+			if (!SelectedItem.IsValid())
+			{
+				continue;
+			}
+
+			const FSoundDashboardEntry& SelectedSoundEntry = CastEntry(*SelectedItem);
+			if (SelectedSoundEntry.PinnedEntryType == FSoundDashboardEntry::EPinnedEntryType::None)
+			{
+				return true;
+			}
+		}
+
+		return false;
+	}
+
 	void FSoundDashboardViewFactory::PinSelectedItems(const TSharedPtr<IDashboardDataTreeViewEntry>& Entry, const TArray<TSharedPtr<IDashboardDataTreeViewEntry>>& SelectedItems)
 	{
 		for (const TSharedPtr<IDashboardDataTreeViewEntry>& Child : Entry->Children)
@@ -1632,12 +1696,12 @@ namespace UE::Audio::Insights
 			{
 				if (IsRootItem(Entry.ToSharedRef()))
 				{
-					MarkBranchAsPinned(Child);
+					MarkBranchAsPinned(Child, true /*bIsPinned*/);
 					CreatePinnedEntry(Child);
 				}
 				else
 				{
-					MarkBranchAsPinned(Entry);
+					MarkBranchAsPinned(Entry, true /*bIsPinned*/);
 					CreatePinnedEntry(Entry);
 				}
 			}
@@ -1648,14 +1712,70 @@ namespace UE::Audio::Insights
 		}
 	}
 
-	void FSoundDashboardViewFactory::MarkBranchAsPinned(const TSharedPtr<IDashboardDataTreeViewEntry> Entry)
+	void FSoundDashboardViewFactory::UnpinSelectedItems(const TSharedPtr<FPinnedSoundEntryWrapper>& PinnedWrapperEntry, const TArray<TSharedPtr<IDashboardDataTreeViewEntry>>& SelectedItems)
+	{
+		using namespace FSoundDashboardViewFactoryPrivate;
+
+		// Run through all the pinned items in the dashboard and check if they are in the list of selected items.
+		// If the child of a non-category parent is selected, we will move the parent and all of it's children back to the unpinned section
+		for (const TSharedPtr<FPinnedSoundEntryWrapper>& PinnedWrapperChild : PinnedWrapperEntry->PinnedWrapperChildren)
+		{
+			const TSharedPtr<IDashboardDataTreeViewEntry> OriginalChildEntry = PinnedWrapperChild->GetOriginalDataEntry();
+			if (!OriginalChildEntry.IsValid())
+			{
+				continue;
+			}
+
+			const FSoundDashboardEntry& OriginalChildSoundEntry = CastEntry(*OriginalChildEntry);
+
+			const TSharedPtr<IDashboardDataTreeViewEntry>* FoundMatchingSelectedEntry = SelectedItems.FindByPredicate([&OriginalChildSoundEntry](const TSharedPtr<IDashboardDataTreeViewEntry> SelectedEntry)
+			{
+				if (!SelectedEntry.IsValid())
+				{
+					return false;
+				}
+
+				const FSoundDashboardEntry& SelectedSoundEntry = CastEntry(*SelectedEntry);
+				return OriginalChildSoundEntry.PlayOrder == SelectedSoundEntry.PlayOrder;
+			});
+
+			if (FoundMatchingSelectedEntry != nullptr)
+			{
+				if (IsCategoryItem(PinnedWrapperEntry->GetPinnedSectionEntry().ToSharedRef()))
+				{
+					// If the current parent pinned item is a category, move the child to the unpinned area and continue
+					MarkBranchAsPinned(OriginalChildEntry, false /*bIsPinned*/);
+					PinnedWrapperChild->MarkToDelete();
+				}
+				else
+				{
+					// If the current parent pinned item is not a category, move it and all of it's children to unpinned.
+					// There is no need to check the other children so break.
+					const TWeakPtr<IDashboardDataTreeViewEntry> OriginalPinnedEntry = PinnedWrapperEntry->GetOriginalDataEntry();
+					if (OriginalPinnedEntry.IsValid())
+					{
+						MarkBranchAsPinned(OriginalPinnedEntry.Pin(), false /*bIsPinned*/);
+						PinnedWrapperEntry->MarkToDelete();
+						break;
+					}
+				}
+			}
+			else
+			{
+				// If this child item is not selected, check it's children.
+				UnpinSelectedItems(PinnedWrapperChild, SelectedItems);
+			}
+		}
+	}
+
+	void FSoundDashboardViewFactory::MarkBranchAsPinned(const TSharedPtr<IDashboardDataTreeViewEntry> Entry, const bool bIsPinned)
 	{
 		FSoundDashboardEntry& SoundEntry = FSoundDashboardViewFactoryPrivate::CastEntry(*Entry);
-		SoundEntry.PinnedEntryType = FSoundDashboardEntry::EPinnedEntryType::HiddenOriginalEntry;
+		SoundEntry.PinnedEntryType = bIsPinned ? FSoundDashboardEntry::EPinnedEntryType::HiddenOriginalEntry : FSoundDashboardEntry::EPinnedEntryType::None;
 
 		for (TSharedPtr<IDashboardDataTreeViewEntry> Child : Entry->Children)
 		{
-			MarkBranchAsPinned(Child);
+			MarkBranchAsPinned(Child, bIsPinned);
 		}
 	}
 
diff --git a/Engine/Plugins/AudioInsights/Source/AudioInsights/Public/Views/SoundDashboardViewFactory.h b/Engine/Plugins/AudioInsights/Source/AudioInsights/Public/Views/SoundDashboardViewFactory.h
--- a/Engine/Plugins/AudioInsights/Source/AudioInsights/Public/Views/SoundDashboardViewFactory.h
+++ b/Engine/Plugins/AudioInsights/Source/AudioInsights/Public/Views/SoundDashboardViewFactory.h
@@ -74,12 +74,14 @@ namespace UE::Audio::Insights
 		FPinnedSoundEntryWrapper(const TSharedPtr<IDashboardDataTreeViewEntry>& OriginalEntry, const bool bIsCategory);
 
 		TSharedPtr<IDashboardDataTreeViewEntry> GetPinnedSectionEntry() const { return PinnedSectionEntry; }
+		TSharedPtr<IDashboardDataTreeViewEntry> GetOriginalDataEntry() const { return OriginalDataEntry.IsValid() ? OriginalDataEntry.Pin() : nullptr; }
 
 		TSharedPtr<FPinnedSoundEntryWrapper> AddChildEntry(const TSharedPtr<IDashboardDataTreeViewEntry> Child, const bool bIsCategory);
 
 		void UpdateParams();
 
 		void CleanUp();
+		void MarkToDelete();
 
 		bool EntryIsValid() const;
 
@@ -145,8 +147,13 @@ namespace UE::Audio::Insights
 #endif // WITH_EDITOR
 
 		void PinSound();
+		void UnpinSound();
+		bool SelectionIncludesUnpinnedItem() const;
+
 		void PinSelectedItems(const TSharedPtr<IDashboardDataTreeViewEntry>& Entry, const TArray<TSharedPtr<IDashboardDataTreeViewEntry>>& SelectedItems);
-		void MarkBranchAsPinned(const TSharedPtr<IDashboardDataTreeViewEntry> Entry);
+		void UnpinSelectedItems(const TSharedPtr<FPinnedSoundEntryWrapper>& PinnedWrapperEntry, const TArray<TSharedPtr<IDashboardDataTreeViewEntry>>& SelectedItems);
+
+		void MarkBranchAsPinned(const TSharedPtr<IDashboardDataTreeViewEntry> Entry, const bool bIsPinned);
 		void InitPinnedItemEntries();
 		void CreatePinnedEntry(TSharedPtr<IDashboardDataTreeViewEntry> Entry);
 		void UpdatePinnedSection();
