[Audio Insights] Sound dashboard: Small hidden category comment code adjustment

diff --git a/Engine/Plugins/AudioInsights/Source/AudioInsights/Private/Views/SoundDashboardViewFactory.cpp b/Engine/Plugins/AudioInsights/Source/AudioInsights/Private/Views/SoundDashboardViewFactory.cpp
--- a/Engine/Plugins/AudioInsights/Source/AudioInsights/Private/Views/SoundDashboardViewFactory.cpp
+++ b/Engine/Plugins/AudioInsights/Source/AudioInsights/Private/Views/SoundDashboardViewFactory.cpp
@@ -299,6 +299,7 @@ namespace UE::Audio::Insights
 			ESoundDashboardFilterFlags::SoundCueTemplate |
 			ESoundDashboardFilterFlags::Pinned           |
 			ESoundDashboardFilterFlags::Hidden;
+
 		const FText FiltersName = LOCTEXT("SoundDashboard_Filter_CategoryText", "Filters");
 		const FText FiltersTooltip = LOCTEXT("CurveFiltersToolTip", "Filters what kind of sounds types can be displayed.");
 
@@ -308,7 +309,8 @@ namespace UE::Audio::Insights
 		const FText SoundWaveCategoryName = LOCTEXT("SoundDashboard_Filter_SoundWaveNameText", "Sound Wave");
 		const FText SoundCueTemplateCategoryName = LOCTEXT("SoundDashboard_Filter_SoundCueTemplateNameText", "Sound Cue Template");
 		const FText PinnedCategoryName = LOCTEXT("SoundDashboard_Filter_PinnedNameText", "Pinned");
-		const FText HiddenCategoryName = LOCTEXT("SoundDashboard_Filter_HiddenNameText", "Hidden");
+		// @TODO UE-250399: Hide category pending to implement
+		//const FText HiddenCategoryName = LOCTEXT("SoundDashboard_Filter_HiddenNameText", "Hidden");
 	} // namespace FSoundDashboardViewFactoryPrivate
 
 	/////////////////////////////////////////////////////////////////////////////////////////
@@ -664,7 +666,7 @@ namespace UE::Audio::Insights
 				//MakeShared<FSoundDashboardFilter>(
 				//ESoundDashboardFilterFlags::Hidden,
 				//"Hidden",
-				//LOCTEXT("SoundDashboard_Filter_HiddenNameText", "Hidden"),
+				//HiddenCategoryName,
 				//"AudioInsights.Icon.SoundDashboard.Hide",
 				//FText::GetEmpty(),
 				//FSlateStyle::Get().GetColor("SoundDashboard.HiddenColor"),
