[Audio Insights] Submix and AudioBus dashboards: Bug fix where entries didn't show up in the dashboard doing this sequence: Open Audio Insights -> Start PIE -> Close Audio Insights -> Stop PIE -> Open Audio Insights

diff --git a/Engine/Plugins/AudioInsights/Source/AudioInsightsEditor/Private/Providers/AudioBusProvider.cpp b/Engine/Plugins/AudioInsights/Source/AudioInsightsEditor/Private/Providers/AudioBusProvider.cpp
--- a/Engine/Plugins/AudioInsights/Source/AudioInsightsEditor/Private/Providers/AudioBusProvider.cpp
+++ b/Engine/Plugins/AudioInsights/Source/AudioInsightsEditor/Private/Providers/AudioBusProvider.cpp
@@ -75,6 +75,11 @@ namespace UE::Audio::Insights
 		UpdateAudioBusAssetNames();
 	}
 
+	void FAudioBusProvider::RequestEntriesUpdate()
+	{
+		UpdateAudioBusAssetNames();
+	}
+
 	void FAudioBusProvider::OnTraceStarted(FTraceAuxiliary::EConnectionType TraceType, const FString& TraceDestination)
 	{
 #if UE_AUDIO_PROFILERTRACE_ENABLED
diff --git a/Engine/Plugins/AudioInsights/Source/AudioInsightsEditor/Private/Providers/AudioBusProvider.h b/Engine/Plugins/AudioInsights/Source/AudioInsightsEditor/Private/Providers/AudioBusProvider.h
--- a/Engine/Plugins/AudioInsights/Source/AudioInsightsEditor/Private/Providers/AudioBusProvider.h
+++ b/Engine/Plugins/AudioInsights/Source/AudioInsightsEditor/Private/Providers/AudioBusProvider.h
@@ -18,6 +18,8 @@ namespace UE::Audio::Insights
 
 		virtual UE::Trace::IAnalyzer* ConstructAnalyzer(TraceServices::IAnalysisSession& InSession) override;
 
+		void RequestEntriesUpdate();
+
 		DECLARE_MULTICAST_DELEGATE_OneParam(FOnAudioBusAssetAdded, const TWeakObjectPtr<UObject> /*Asset*/);
 		inline static FOnAudioBusAssetAdded OnAudioBusAssetAdded;
 
diff --git a/Engine/Plugins/AudioInsights/Source/AudioInsightsEditor/Private/Providers/SoundSubmixProvider.cpp b/Engine/Plugins/AudioInsights/Source/AudioInsightsEditor/Private/Providers/SoundSubmixProvider.cpp
--- a/Engine/Plugins/AudioInsights/Source/AudioInsightsEditor/Private/Providers/SoundSubmixProvider.cpp
+++ b/Engine/Plugins/AudioInsights/Source/AudioInsightsEditor/Private/Providers/SoundSubmixProvider.cpp
@@ -75,6 +75,11 @@ namespace UE::Audio::Insights
 		UpdateSubmixAssetNames();
 	}
 
+	void FSoundSubmixProvider::RequestEntriesUpdate()
+	{
+		UpdateSubmixAssetNames();
+	}
+
 	void FSoundSubmixProvider::OnTraceStarted(FTraceAuxiliary::EConnectionType TraceType, const FString& TraceDestination)
 	{
 #if UE_AUDIO_PROFILERTRACE_ENABLED
diff --git a/Engine/Plugins/AudioInsights/Source/AudioInsightsEditor/Private/Providers/SoundSubmixProvider.h b/Engine/Plugins/AudioInsights/Source/AudioInsightsEditor/Private/Providers/SoundSubmixProvider.h
--- a/Engine/Plugins/AudioInsights/Source/AudioInsightsEditor/Private/Providers/SoundSubmixProvider.h
+++ b/Engine/Plugins/AudioInsights/Source/AudioInsightsEditor/Private/Providers/SoundSubmixProvider.h
@@ -18,6 +18,8 @@ namespace UE::Audio::Insights
 
 		virtual UE::Trace::IAnalyzer* ConstructAnalyzer(TraceServices::IAnalysisSession& InSession) override;
 		
+		void RequestEntriesUpdate();
+
 		DECLARE_MULTICAST_DELEGATE_OneParam(FOnSubmixAssetAdded, const TWeakObjectPtr<UObject> /*Asset*/);
 		inline static FOnSubmixAssetAdded OnSubmixAssetAdded;
 
diff --git a/Engine/Plugins/AudioInsights/Source/AudioInsightsEditor/Private/Views/AudioBusesDashboardViewFactory.cpp b/Engine/Plugins/AudioInsights/Source/AudioInsightsEditor/Private/Views/AudioBusesDashboardViewFactory.cpp
--- a/Engine/Plugins/AudioInsights/Source/AudioInsightsEditor/Private/Views/AudioBusesDashboardViewFactory.cpp
+++ b/Engine/Plugins/AudioInsights/Source/AudioInsightsEditor/Private/Views/AudioBusesDashboardViewFactory.cpp
@@ -32,7 +32,7 @@ namespace UE::Audio::Insights
 
 		IAudioInsightsTraceModule& AudioInsightsTraceModule = IAudioInsightsEditorModule::GetChecked().GetTraceModule();
 
-		const TSharedPtr<FAudioBusProvider> AudioBusProvider = MakeShared<FAudioBusProvider>();
+		AudioBusProvider = MakeShared<FAudioBusProvider>();
 
 		AudioInsightsTraceModule.AddTraceProvider(AudioBusProvider);
 
@@ -286,6 +286,13 @@ namespace UE::Audio::Insights
 				FilteredEntriesListView->SetSelectionMode(ESelectionMode::Single);
 			}
 		}
+		else
+		{
+			if (AudioBusProvider.IsValid())
+			{
+				AudioBusProvider->RequestEntriesUpdate();
+			}
+		}
 
 		for (const auto& [AudioBus, bIsCheckedState] : AudioBusCheckboxCheckedStates)
 		{
diff --git a/Engine/Plugins/AudioInsights/Source/AudioInsightsEditor/Private/Views/AudioBusesDashboardViewFactory.h b/Engine/Plugins/AudioInsights/Source/AudioInsightsEditor/Private/Views/AudioBusesDashboardViewFactory.h
--- a/Engine/Plugins/AudioInsights/Source/AudioInsightsEditor/Private/Views/AudioBusesDashboardViewFactory.h
+++ b/Engine/Plugins/AudioInsights/Source/AudioInsightsEditor/Private/Views/AudioBusesDashboardViewFactory.h
@@ -9,6 +9,8 @@ class UAudioBus;
 
 namespace UE::Audio::Insights
 {
+	class FAudioBusProvider;
+
 	class FAudioBusesDashboardViewFactory : public FTraceObjectTableDashboardViewFactory
 	{
 	public:
@@ -50,6 +52,7 @@ namespace UE::Audio::Insights
 
 		void HandleOnAudioBusAssetListUpdated(const TWeakObjectPtr<UObject> InAsset);
 
+		TSharedPtr<FAudioBusProvider> AudioBusProvider;
 		TMap<const TWeakObjectPtr<UAudioBus>, bool> AudioBusCheckboxCheckedStates;
 
 		using FComboboxSelectionItem = TPair<EAudioBusTypeComboboxSelection, FText>;
diff --git a/Engine/Plugins/AudioInsights/Source/AudioInsightsEditor/Private/Views/SubmixesDashboardViewFactory.cpp b/Engine/Plugins/AudioInsights/Source/AudioInsightsEditor/Private/Views/SubmixesDashboardViewFactory.cpp
--- a/Engine/Plugins/AudioInsights/Source/AudioInsightsEditor/Private/Views/SubmixesDashboardViewFactory.cpp
+++ b/Engine/Plugins/AudioInsights/Source/AudioInsightsEditor/Private/Views/SubmixesDashboardViewFactory.cpp
@@ -30,7 +30,7 @@ namespace UE::Audio::Insights
 
 		IAudioInsightsTraceModule& AudioInsightsTraceModule = IAudioInsightsEditorModule::GetChecked().GetTraceModule();
 
-		const TSharedPtr<FSoundSubmixProvider> SoundSubmixProvider = MakeShared<FSoundSubmixProvider>();
+		SoundSubmixProvider = MakeShared<FSoundSubmixProvider>();
 
 		AudioInsightsTraceModule.AddTraceProvider(SoundSubmixProvider);
 
@@ -194,6 +194,13 @@ namespace UE::Audio::Insights
 				FilteredEntriesListView->SetSelectionMode(ESelectionMode::Single);
 			}
 		}
+		else
+		{
+			if (SoundSubmixProvider.IsValid())
+			{
+				SoundSubmixProvider->RequestEntriesUpdate();
+			}
+		}
 
 		for (const auto& [AudioBus, bIsCheckedState] : SubmixCheckboxCheckedStates)
 		{
diff --git a/Engine/Plugins/AudioInsights/Source/AudioInsightsEditor/Private/Views/SubmixesDashboardViewFactory.h b/Engine/Plugins/AudioInsights/Source/AudioInsightsEditor/Private/Views/SubmixesDashboardViewFactory.h
--- a/Engine/Plugins/AudioInsights/Source/AudioInsightsEditor/Private/Views/SubmixesDashboardViewFactory.h
+++ b/Engine/Plugins/AudioInsights/Source/AudioInsightsEditor/Private/Views/SubmixesDashboardViewFactory.h
@@ -10,6 +10,8 @@ class USoundSubmix;
 
 namespace UE::Audio::Insights
 {
+	class FSoundSubmixProvider;
+
 	class FSubmixesDashboardViewFactory : public FTraceObjectTableDashboardViewFactory
 	{
 	public:
@@ -44,6 +46,7 @@ namespace UE::Audio::Insights
 
 		void HandleOnSubmixAssetListUpdated(const TWeakObjectPtr<UObject> InAsset);
 
+		TSharedPtr<FSoundSubmixProvider> SoundSubmixProvider;
 		TMap<const TWeakObjectPtr<USoundSubmix>, bool> SubmixCheckboxCheckedStates;
 	};
 } // namespace UE::Audio::Insights
