[Audio Insights] Fix for assert in DynamicsProcessor after device swap

diff --git a/Engine/Plugins/AudioInsights/Source/AudioInsightsEditor/Private/Messages/SoundSubmixMessages.h b/Engine/Plugins/AudioInsights/Source/AudioInsightsEditor/Private/Messages/SoundSubmixMessages.h
--- a/Engine/Plugins/AudioInsights/Source/AudioInsightsEditor/Private/Messages/SoundSubmixMessages.h
+++ b/Engine/Plugins/AudioInsights/Source/AudioInsightsEditor/Private/Messages/SoundSubmixMessages.h
@@ -56,7 +56,6 @@ namespace UE::Audio::Insights
 		FSoundSubmixAssetDashboardEntry() = default;
 		virtual ~FSoundSubmixAssetDashboardEntry() = default;
 
-		TWeakObjectPtr<USoundSubmix> SoundSubmix;
 		uint32 SubmixId = INDEX_NONE;
 		bool bHasActivity = false;
 	};
diff --git a/Engine/Plugins/AudioInsights/Source/AudioInsightsEditor/Private/Providers/SoundSubmixProvider.cpp b/Engine/Plugins/AudioInsights/Source/AudioInsightsEditor/Private/Providers/SoundSubmixProvider.cpp
--- a/Engine/Plugins/AudioInsights/Source/AudioInsightsEditor/Private/Providers/SoundSubmixProvider.cpp
+++ b/Engine/Plugins/AudioInsights/Source/AudioInsightsEditor/Private/Providers/SoundSubmixProvider.cpp
@@ -100,15 +100,12 @@ namespace UE::Audio::Insights
 			const FAudioInsightsEditorModule AudioInsightsEditorModule = FAudioInsightsEditorModule::GetChecked();
 			const ::Audio::FDeviceId AudioDeviceId = AudioInsightsEditorModule.GetDeviceId();
 
-			TSharedPtr<FSoundSubmixAssetDashboardEntry> SoundSubmixAssetDashboardEntry = MakeShared<FSoundSubmixAssetDashboardEntry>();
+			TSharedPtr<FSoundSubmixAssetDashboardEntry>& SoundSubmixAssetDashboardEntry = SubmixDataViewEntries.Add_GetRef(MakeShared<FSoundSubmixAssetDashboardEntry>());
 			SoundSubmixAssetDashboardEntry->DeviceId    = AudioDeviceId;
 			SoundSubmixAssetDashboardEntry->Name        = InAssetData.GetObjectPathString();
-			SoundSubmixAssetDashboardEntry->SoundSubmix = Cast<USoundSubmix>(InAssetData.GetAsset());
-			SoundSubmixAssetDashboardEntry->SubmixId    = SoundSubmixAssetDashboardEntry->SoundSubmix->GetUniqueID();
+			SoundSubmixAssetDashboardEntry->SubmixId    = GetTypeHash(SoundSubmixAssetDashboardEntry->Name);
 
-			SubmixDataViewEntries.Add(MoveTemp(SoundSubmixAssetDashboardEntry));
-
-			OnSubmixAssetAdded.Broadcast(InAssetData.GetAsset());
+			OnSubmixAssetAdded.Broadcast(SoundSubmixAssetDashboardEntry->SubmixId);
 
 			bAssetEntriesNeedRefreshing = true;
 			++LastUpdateId;
@@ -128,11 +125,13 @@ namespace UE::Audio::Insights
 			const FAudioInsightsEditorModule AudioInsightsEditorModule = FAudioInsightsEditorModule::GetChecked();
 			const ::Audio::FDeviceId AudioDeviceId = AudioInsightsEditorModule.GetDeviceId();
 
-			RemoveDeviceEntry(AudioDeviceId, SubmixDataViewEntries[FoundSubmixAssetNameIndex]->SubmixId);
+			const uint32 SubmixId = SubmixDataViewEntries[FoundSubmixAssetNameIndex]->SubmixId;
+
+			RemoveDeviceEntry(AudioDeviceId, SubmixId);
 
 			SubmixDataViewEntries.RemoveAt(FoundSubmixAssetNameIndex);
 
-			OnSubmixAssetRemoved.Broadcast(InAssetData.GetAsset());
+			OnSubmixAssetRemoved.Broadcast(SubmixId);
 
 			bAssetEntriesNeedRefreshing = true;
 			++LastUpdateId;
@@ -181,7 +180,7 @@ namespace UE::Audio::Insights
 
 			for (const TSharedPtr<FSoundSubmixAssetDashboardEntry>& SubmixDataViewEntry : SubmixDataViewEntries)
 			{
-				if (SubmixDataViewEntry.IsValid() && SubmixDataViewEntry->SoundSubmix.IsValid())
+				if (SubmixDataViewEntry.IsValid())
 				{
 					UpdateDeviceEntry(AudioDeviceId, SubmixDataViewEntry->SubmixId, [&SubmixDataViewEntry](TSharedPtr<FSoundSubmixAssetDashboardEntry>& Entry)
 					{
diff --git a/Engine/Plugins/AudioInsights/Source/AudioInsightsEditor/Private/Providers/SoundSubmixProvider.h b/Engine/Plugins/AudioInsights/Source/AudioInsightsEditor/Private/Providers/SoundSubmixProvider.h
--- a/Engine/Plugins/AudioInsights/Source/AudioInsightsEditor/Private/Providers/SoundSubmixProvider.h
+++ b/Engine/Plugins/AudioInsights/Source/AudioInsightsEditor/Private/Providers/SoundSubmixProvider.h
@@ -20,10 +20,10 @@ namespace UE::Audio::Insights
 		
 		void RequestEntriesUpdate();
 
-		DECLARE_MULTICAST_DELEGATE_OneParam(FOnSubmixAssetAdded, const TWeakObjectPtr<UObject> /*Asset*/);
+		DECLARE_MULTICAST_DELEGATE_OneParam(FOnSubmixAssetAdded, const uint32 /*SubmixId*/);
 		inline static FOnSubmixAssetAdded OnSubmixAssetAdded;
 
-		DECLARE_MULTICAST_DELEGATE_OneParam(FOnSubmixAssetRemoved, const TWeakObjectPtr<UObject> /*Asset*/);
+		DECLARE_MULTICAST_DELEGATE_OneParam(FOnSubmixAssetRemoved, const uint32 /*SubmixId*/);
 		inline static FOnSubmixAssetRemoved OnSubmixAssetRemoved;
 
 		DECLARE_MULTICAST_DELEGATE(FOnSubmixAssetListUpdated);
diff --git a/Engine/Plugins/AudioInsights/Source/AudioInsightsEditor/Private/Views/AudioMetersDashboardViewFactory.cpp b/Engine/Plugins/AudioInsights/Source/AudioInsightsEditor/Private/Views/AudioMetersDashboardViewFactory.cpp
--- a/Engine/Plugins/AudioInsights/Source/AudioInsightsEditor/Private/Views/AudioMetersDashboardViewFactory.cpp
+++ b/Engine/Plugins/AudioInsights/Source/AudioInsightsEditor/Private/Views/AudioMetersDashboardViewFactory.cpp
@@ -58,12 +58,12 @@ namespace UE::Audio::Insights
 
 		if (!OnSubmixAssetRemovedHandle.IsValid())
 		{
-			OnSubmixAssetRemovedHandle = FSoundSubmixProvider::OnSubmixAssetRemoved.AddSP(this, &FAudioMetersDashboardViewFactory::HandleOnAssetRemoved);
+			OnSubmixAssetRemovedHandle = FSoundSubmixProvider::OnSubmixAssetRemoved.AddSP(this, &FAudioMetersDashboardViewFactory::HandleOnSubmixAssetRemoved);
 		}
 
 		if (!OnAudioBusAssetRemovedHandle.IsValid())
 		{
-			OnAudioBusAssetRemovedHandle = FAudioBusProvider::OnAudioBusAssetRemoved.AddSP(this, &FAudioMetersDashboardViewFactory::HandleOnAssetRemoved);
+			OnAudioBusAssetRemovedHandle = FAudioBusProvider::OnAudioBusAssetRemoved.AddSP(this, &FAudioMetersDashboardViewFactory::HandleOnAudioBusAssetRemoved);
 		}
 		
 		if (!MeterViewsScrollBox.IsValid())
@@ -79,31 +79,36 @@ namespace UE::Audio::Insights
 		return MeterViewsScrollBox.ToSharedRef();
 	}
 
-	void FAudioMetersDashboardViewFactory::HandleOnSubmixAssetInit(const bool bInIsChecked, const TWeakObjectPtr<USoundSubmix> InSoundSubmix)
+	void FAudioMetersDashboardViewFactory::HandleOnSubmixAssetInit(const bool bInIsChecked, const uint32 InSubmixId, const FString& InSubmixName)
 	{
-		if (!bInIsChecked || !InSoundSubmix.IsValid())
+		if (!bInIsChecked)
 		{
 			return;
 		}
 
-		const uint32 SubmixUniqueId = InSoundSubmix->GetUniqueID();
+		if (!AudioMeterViews.Contains(InSubmixId))
+		{
+			return;
+		}
 
-		if (!AudioMeterViews.Contains(SubmixUniqueId))
+		const TObjectPtr<UObject> LoadedSubmix = FSoftObjectPath(InSubmixName).TryLoad();
+		if (LoadedSubmix == nullptr)
 		{
 			return;
 		}
 
 		// Tear down and re-create meter analyzers on initialization to ensure we receive the correct volume data
-		const int32 SlotIndex = AudioMeterViewsContainer->RemoveSlot(AudioMeterViews[SubmixUniqueId]->GetWidget());
+		const int32 SlotIndex = AudioMeterViewsContainer->RemoveSlot(AudioMeterViews[InSubmixId]->GetWidget());
 
-		AudioMeterViews.Remove(SubmixUniqueId);
-		AudioMeterViews.Add(SubmixUniqueId, MakeShared<FAudioMeterView>(FAudioMeterView::FAudioAssetVariant(TInPlaceType<TWeakObjectPtr<USoundSubmix>>(), InSoundSubmix)));
+		AudioMeterViews.Remove(InSubmixId);
+
+		AudioMeterViews.Add(InSubmixId, MakeShared<FAudioMeterView>(FAudioMeterView::FAudioAssetVariant(TInPlaceType<TWeakObjectPtr<USoundSubmix>>(), Cast<USoundSubmix>(LoadedSubmix))));
 
 		AudioMeterViewsContainer->InsertSlot(SlotIndex)
 		.AutoWidth()
 		.Padding(10.0f, 0.0f, 10.0f, 0.0f)
 		[
-			AudioMeterViews[SubmixUniqueId]->GetWidget()
+			AudioMeterViews[InSubmixId]->GetWidget()
 		];
 
 		if (MeterViewsScrollBox.IsValid())
@@ -144,32 +149,31 @@ namespace UE::Audio::Insights
 		}
 	}
 
-	void FAudioMetersDashboardViewFactory::HandleOnSubmixAssetChecked(const bool bInIsChecked, const TWeakObjectPtr<USoundSubmix> InSoundSubmix)
+	void FAudioMetersDashboardViewFactory::HandleOnSubmixAssetChecked(const bool bInIsChecked, const uint32 InSubmixId, const FString& InSoundSubmixName)
 	{
-		if (!InSoundSubmix.IsValid())
-		{
-			return;
-		}
-
-		const uint32 SubmixUniqueId = InSoundSubmix->GetUniqueID();
-
 		if (bInIsChecked)
 		{
-			AudioMeterViews.Add(SubmixUniqueId, MakeShared<FAudioMeterView>(FAudioMeterView::FAudioAssetVariant(TInPlaceType<TWeakObjectPtr<USoundSubmix>>(), InSoundSubmix)));
+			const TObjectPtr<UObject> LoadedSubmix = FSoftObjectPath(InSoundSubmixName).TryLoad();
+			if (LoadedSubmix == nullptr)
+			{
+				return;
+			}
+
+			AudioMeterViews.Add(InSubmixId, MakeShared<FAudioMeterView>(FAudioMeterView::FAudioAssetVariant(TInPlaceType<TWeakObjectPtr<USoundSubmix>>(), Cast<USoundSubmix>(LoadedSubmix))));
 
 			AudioMeterViewsContainer->AddSlot()
 			.AutoWidth()
 			.Padding(10.0f, 0.0f, 10.0f, 0.0f)
 			[
-				AudioMeterViews[SubmixUniqueId]->GetWidget()
+				AudioMeterViews[InSubmixId]->GetWidget()
 			];
 		}
 		else
 		{
-			if (AudioMeterViews.Contains(SubmixUniqueId))
+			if (AudioMeterViews.Contains(InSubmixId))
 			{
-				AudioMeterViewsContainer->RemoveSlot(AudioMeterViews[SubmixUniqueId]->GetWidget());
-				AudioMeterViews.Remove(SubmixUniqueId);
+				AudioMeterViewsContainer->RemoveSlot(AudioMeterViews[InSubmixId]->GetWidget());
+				AudioMeterViews.Remove(InSubmixId);
 			}
 		}
 
@@ -214,14 +218,28 @@ namespace UE::Audio::Insights
 		}
 	}
 
-	void FAudioMetersDashboardViewFactory::HandleOnAssetRemoved(const TWeakObjectPtr<UObject> InAsset)
+	void FAudioMetersDashboardViewFactory::HandleOnSubmixAssetRemoved(const uint32 InSubmixId)
+	{
+		if (AudioMeterViews.Contains(InSubmixId))
+		{
+			AudioMeterViewsContainer->RemoveSlot(AudioMeterViews[InSubmixId]->GetWidget());
+			AudioMeterViews.Remove(InSubmixId);
+		}
+
+		if (MeterViewsScrollBox.IsValid())
+		{
+			MeterViewsScrollBox->Invalidate(EInvalidateWidget::Layout);
+		}
+	}
+
+	void FAudioMetersDashboardViewFactory::HandleOnAudioBusAssetRemoved(const TWeakObjectPtr<UObject> InAudioBusAsset)
 	{
-		if (!InAsset.IsValid())
+		if (!InAudioBusAsset.IsValid())
 		{
 			return;
 		}
 
-		if (const uint32 AssetUniqueId = InAsset->GetUniqueID();
+		if (const uint32 AssetUniqueId = InAudioBusAsset->GetUniqueID();
 			AudioMeterViews.Contains(AssetUniqueId))
 		{
 			AudioMeterViewsContainer->RemoveSlot(AudioMeterViews[AssetUniqueId]->GetWidget());
diff --git a/Engine/Plugins/AudioInsights/Source/AudioInsightsEditor/Private/Views/AudioMetersDashboardViewFactory.h b/Engine/Plugins/AudioInsights/Source/AudioInsightsEditor/Private/Views/AudioMetersDashboardViewFactory.h
--- a/Engine/Plugins/AudioInsights/Source/AudioInsightsEditor/Private/Views/AudioMetersDashboardViewFactory.h
+++ b/Engine/Plugins/AudioInsights/Source/AudioInsightsEditor/Private/Views/AudioMetersDashboardViewFactory.h
@@ -26,13 +26,14 @@ namespace UE::Audio::Insights
 		virtual TSharedRef<SWidget> MakeWidget(TSharedRef<SDockTab> OwnerTab, const FSpawnTabArgs& SpawnTabArgs) override;
 
 	private:
-		void HandleOnSubmixAssetInit(const bool bInIsChecked, const TWeakObjectPtr<USoundSubmix> InSoundSubmix);
+		void HandleOnSubmixAssetInit(const bool bInIsChecked, const uint32 InSubmixId, const FString& InSubmixName);
 		void HandleOnAudioBusAssetInit(const bool bInIsChecked, const TWeakObjectPtr<UAudioBus> InAudioBus);
 
-		void HandleOnSubmixAssetChecked(const bool bInIsChecked, const TWeakObjectPtr<USoundSubmix> InSoundSubmix);
+		void HandleOnSubmixAssetChecked(const bool bInIsChecked, const uint32 InSubmixId, const FString& InSubmixName);
 		void HandleOnAudioBusAssetChecked(const bool bInIsChecked, const TWeakObjectPtr<UAudioBus> InAudioBus);
 
-		void HandleOnAssetRemoved(const TWeakObjectPtr<UObject> InAsset);
+		void HandleOnSubmixAssetRemoved(const uint32 InSubmixId);
+		void HandleOnAudioBusAssetRemoved(const TWeakObjectPtr<UObject> InAudioBusAsset);
 
 		TSharedPtr<SScrollBox> MeterViewsScrollBox;
 		TSharedPtr<SHorizontalBox> AudioMeterViewsContainer;
diff --git a/Engine/Plugins/AudioInsights/Source/AudioInsightsEditor/Private/Views/SubmixesDashboardViewFactory.cpp b/Engine/Plugins/AudioInsights/Source/AudioInsightsEditor/Private/Views/SubmixesDashboardViewFactory.cpp
--- a/Engine/Plugins/AudioInsights/Source/AudioInsightsEditor/Private/Views/SubmixesDashboardViewFactory.cpp
+++ b/Engine/Plugins/AudioInsights/Source/AudioInsightsEditor/Private/Views/SubmixesDashboardViewFactory.cpp
@@ -92,15 +92,11 @@ namespace UE::Audio::Insights
 				return SNullWidget::NullWidget;
 			}
 
-			const TWeakObjectPtr<USoundSubmix> SoundSubmix = SubmixesPrivate::CastEntry(InRowData.Get()).SoundSubmix;
+			const FSoundSubmixAssetDashboardEntry& SoundSubmixAssetDashboardEntry = SubmixesPrivate::CastEntry(*InRowData);
+			const uint32 SubmixId = SoundSubmixAssetDashboardEntry.SubmixId;
 
-			if (!SoundSubmix.IsValid())
-			{
-				return SNullWidget::NullWidget;
-			}
-
-			const bool bInitValue = SubmixCheckboxCheckedStates.FindOrAdd(SoundSubmix, false);
-			OnSubmixAssetInit.Broadcast(bInitValue, SoundSubmix);
+			const bool bInitValue = SubmixCheckboxCheckedStates.FindOrAdd(SubmixId);
+			OnSubmixAssetInit.Broadcast(bInitValue, SubmixId, SoundSubmixAssetDashboardEntry.Name);
 
 			return SNew(SHorizontalBox)
 				.Clipping(EWidgetClipping::ClipToBounds)
@@ -108,21 +104,25 @@ namespace UE::Audio::Insights
 				.AutoWidth()
 				[
 					SNew(SCheckBox)
-					.IsChecked_Lambda([this, SoundSubmix]()
+					.IsChecked_Lambda([this, SubmixId]()
 					{
-						const bool* FoundSubmixCheckedState = SubmixCheckboxCheckedStates.Find(SoundSubmix);
+						const bool* FoundSubmixCheckedState = SubmixCheckboxCheckedStates.Find(SubmixId);
 						const bool bIsSubmixChecked = FoundSubmixCheckedState && *FoundSubmixCheckedState == true;
 
 						return bIsSubmixChecked ? ECheckBoxState::Checked : ECheckBoxState::Unchecked;
 					})
-					.OnCheckStateChanged_Lambda([this, SoundSubmix](ECheckBoxState NewState)
+					.OnCheckStateChanged_Lambda([this, InRowData](ECheckBoxState NewState)
 					{
-						if (bool* FoundSubmixCheckedState = SubmixCheckboxCheckedStates.Find(SoundSubmix))
+						const bool bIsChecked = NewState == ECheckBoxState::Checked;
+
+						const FSoundSubmixAssetDashboardEntry& SoundSubmixAssetDashboardEntry = SubmixesPrivate::CastEntry(*InRowData);
+
+						if (bool* FoundSubmixCheckedState = SubmixCheckboxCheckedStates.Find(SoundSubmixAssetDashboardEntry.SubmixId))
 						{
-							*FoundSubmixCheckedState = NewState == ECheckBoxState::Checked;
+							*FoundSubmixCheckedState = bIsChecked;
 						}
 
-						OnSubmixAssetChecked.Broadcast(NewState == ECheckBoxState::Checked, SoundSubmix);
+						OnSubmixAssetChecked.Broadcast(bIsChecked, SoundSubmixAssetDashboardEntry.SubmixId, SoundSubmixAssetDashboardEntry.Name);
 					})
 				]
 				+ SHorizontalBox::Slot()
@@ -202,9 +202,14 @@ namespace UE::Audio::Insights
 			}
 		}
 
-		for (const auto& [AudioBus, bIsCheckedState] : SubmixCheckboxCheckedStates)
+		for (const TSharedPtr<IDashboardDataViewEntry>& Entry : DataViewEntries)
 		{
-			OnSubmixAssetInit.Broadcast(bIsCheckedState, AudioBus);
+			const FSoundSubmixAssetDashboardEntry& SoundSubmixAssetDashboardEntry = SubmixesPrivate::CastEntry(*Entry);
+
+			if (const bool* bFoundSubmixCheckboxCheckedState = SubmixCheckboxCheckedStates.Find(SoundSubmixAssetDashboardEntry.SubmixId))
+			{
+				OnSubmixAssetInit.Broadcast(*bFoundSubmixCheckboxCheckedState, SoundSubmixAssetDashboardEntry.SubmixId, SoundSubmixAssetDashboardEntry.Name);
+			}
 		}
 
 		return DashboardWidget->AsShared();
@@ -298,9 +303,12 @@ namespace UE::Audio::Insights
 	{
 		if (SelectedItem.IsValid())
 		{
-			const FSoundSubmixAssetDashboardEntry& SoundSubmixAssetDashboardEntry = SubmixesPrivate::CastEntry(*SelectedItem.Get());
+			const FSoundSubmixAssetDashboardEntry& SoundSubmixAssetDashboardEntry = SubmixesPrivate::CastEntry(*SelectedItem);
 
-			OnSubmixSelectionChanged.Broadcast(SoundSubmixAssetDashboardEntry.SoundSubmix);
+			if (const TObjectPtr<UObject> LoadedSubmix = FSoftObjectPath(SoundSubmixAssetDashboardEntry.Name).TryLoad())
+			{
+				OnSubmixSelectionChanged.Broadcast(Cast<USoundSubmix>(LoadedSubmix));
+			}
 		}
 	}
 
@@ -312,7 +320,7 @@ namespace UE::Audio::Insights
 		}
 	}
 
-	void FSubmixesDashboardViewFactory::HandleOnSubmixAssetListUpdated(const TWeakObjectPtr<UObject> InAsset)
+	void FSubmixesDashboardViewFactory::HandleOnSubmixAssetListUpdated(const uint32 InSubmixId)
 	{
 		RequestListRefresh();
 	}
diff --git a/Engine/Plugins/AudioInsights/Source/AudioInsightsEditor/Private/Views/SubmixesDashboardViewFactory.h b/Engine/Plugins/AudioInsights/Source/AudioInsightsEditor/Private/Views/SubmixesDashboardViewFactory.h
--- a/Engine/Plugins/AudioInsights/Source/AudioInsightsEditor/Private/Views/SubmixesDashboardViewFactory.h
+++ b/Engine/Plugins/AudioInsights/Source/AudioInsightsEditor/Private/Views/SubmixesDashboardViewFactory.h
@@ -24,10 +24,10 @@ namespace UE::Audio::Insights
 		virtual EDefaultDashboardTabStack GetDefaultTabStack() const override;
 		virtual TSharedRef<SWidget> MakeWidget(TSharedRef<SDockTab> OwnerTab, const FSpawnTabArgs& SpawnTabArgs) override;
 
-		DECLARE_MULTICAST_DELEGATE_TwoParams(FOnSubmixAssetInit, const bool /*bIsChecked*/, TWeakObjectPtr<USoundSubmix> /*SoundSubmix*/);
+		DECLARE_MULTICAST_DELEGATE_ThreeParams(FOnSubmixAssetInit, const bool /*bIsChecked*/, const uint32 /*SubmixId*/, const FString& /*SubmixName*/);
 		inline static FOnSubmixAssetInit OnSubmixAssetInit;
 
-		DECLARE_MULTICAST_DELEGATE_TwoParams(FOnSubmixAssetChecked, const bool /*bIsChecked*/, TWeakObjectPtr<USoundSubmix> /*SoundSubmix*/);
+		DECLARE_MULTICAST_DELEGATE_ThreeParams(FOnSubmixAssetChecked, const bool /*bIsChecked*/, const uint32 /*SubmixId*/, const FString& /*SubmixName*/);
 		inline static FOnSubmixAssetChecked OnSubmixAssetChecked;
 
 		DECLARE_MULTICAST_DELEGATE_OneParam(FOnSubmixSelectionChanged, const TWeakObjectPtr<USoundSubmix> /*SoundSubmix*/);
@@ -44,9 +44,9 @@ namespace UE::Audio::Insights
 
 		void RequestListRefresh();
 
-		void HandleOnSubmixAssetListUpdated(const TWeakObjectPtr<UObject> InAsset);
+		void HandleOnSubmixAssetListUpdated(const uint32 InSubmixId);
 
 		TSharedPtr<FSoundSubmixProvider> SoundSubmixProvider;
-		TMap<const TWeakObjectPtr<USoundSubmix>, bool> SubmixCheckboxCheckedStates;
+		TMap<uint32, bool> SubmixCheckboxCheckedStates;
 	};
 } // namespace UE::Audio::Insights
diff --git a/Engine/Source/Runtime/AudioMixer/Private/AudioMixerSubmix.cpp b/Engine/Source/Runtime/AudioMixer/Private/AudioMixerSubmix.cpp
--- a/Engine/Source/Runtime/AudioMixer/Private/AudioMixerSubmix.cpp
+++ b/Engine/Source/Runtime/AudioMixer/Private/AudioMixerSubmix.cpp
@@ -1614,7 +1614,7 @@ namespace Audio
 			{
 				UE_TRACE_LOG(Audio, SoundSubmixHasActivity, AudioChannel)
 					<< SoundSubmixHasActivity.DeviceId(MixerDevice->DeviceID)
-					<< SoundSubmixHasActivity.SubmixId(OwningSubmixObject->GetUniqueID())
+					<< SoundSubmixHasActivity.SubmixId(GetTypeHash(OwningSubmixObject->GetPathName()))
 					<< SoundSubmixHasActivity.Timestamp(FPlatformTime::Cycles64())
 					<< SoundSubmixHasActivity.HasActivity(!bIsNowSilent);
 			}
