[Audio Insights] Sound dashboard: Crash fix in Sound Trace Provider - crash when opening a trace in standalone, garbage data in returned sound entry from GetOrCreateActiveSoundEntry.
- Problem was that OutActiveSoundEntry in ProcessMessageQueue<FSoundStartMessage>would become garbage in rare cases when loading traces in standalone
- The TSharedPtr<FSoundDashboardEntry>* ReturnedSoundEntry was only guarenteed to be alive during the scope of the lambda GetOrCreateActiveSoundEntry, hence the rare crash as it was being passed back by pointer
- GetOrCreateActiveSoundEntry is now a function that passes back the TSharedPtr<FSoundDashboardEntry> by argument reference rather than return. We now use the same pattern across all ProcessMessageQueues of passing a shared pointer in via a lambda capture reference from the outer scope, so we can be confident the shared pointer will remain in scope.

Submitted on behalf of [REDACTED_HANDLE]

diff --git a/Engine/Plugins/AudioInsights/Source/AudioInsights/Private/Providers/SoundTraceProvider.cpp b/Engine/Plugins/AudioInsights/Source/AudioInsights/Private/Providers/SoundTraceProvider.cpp
--- a/Engine/Plugins/AudioInsights/Source/AudioInsights/Private/Providers/SoundTraceProvider.cpp
+++ b/Engine/Plugins/AudioInsights/Source/AudioInsights/Private/Providers/SoundTraceProvider.cpp
@@ -154,7 +154,6 @@ namespace UE::Audio::Insights
 	FSoundTraceProvider::FSoundTraceProvider()
 		: TDeviceDataMapTraceProvider<ESoundDashboardEntryType, TSharedPtr<FSoundDashboardEntry>>(GetName_Static())
 	{
-		InitLambdas();
 	}
 
 	FName FSoundTraceProvider::GetName_Static()
@@ -295,12 +294,19 @@ namespace UE::Audio::Insights
 	{
 		// Helper lambdas
 		TSharedPtr<FSoundDashboardEntry> ActiveSoundEntryReturn = nullptr;
-		auto GetActiveSoundEntryFromStartMessage = [this, &ActiveSoundEntryReturn](const FSoundWaveStartMessage& Msg) -> TSharedPtr<FSoundDashboardEntry>*
+		auto GetActiveSoundEntryFromActiveStartMessage = [this, &ActiveSoundEntryReturn](const FSoundStartMessage& Msg) -> TSharedPtr<FSoundDashboardEntry>*
 		{
-			GetActiveSoundEntryFromIDs(Msg.ActiveSoundPlayOrder, Msg.DeviceId, ActiveSoundEntryReturn);
+			GetOrCreateActiveSoundEntry(Msg, ActiveSoundEntryReturn);
 			return &ActiveSoundEntryReturn;
 		};
 
+		TSharedPtr<FSoundDashboardEntry> SoundWaveActiveEntryReturn = nullptr;
+		auto GetActiveSoundEntryFromSoundWaveStartMessage = [this, &SoundWaveActiveEntryReturn](const FSoundWaveStartMessage& Msg) -> TSharedPtr<FSoundDashboardEntry>*
+		{
+			GetActiveSoundEntryFromIDs(Msg.ActiveSoundPlayOrder, Msg.DeviceId, SoundWaveActiveEntryReturn);
+			return &SoundWaveActiveEntryReturn;
+		};
+
 		TSharedPtr<FSoundDashboardEntry> ParamReturn = nullptr;
 		auto GetActiveSoundEntryFromParameterMessage = [this, &ParamReturn](const FMixerSourceParameterMessage& Msg) -> TSharedPtr<FSoundDashboardEntry>*
 		{
@@ -316,7 +322,7 @@ namespace UE::Audio::Insights
 		};
 
 		// Process messages
-		ProcessMessageQueue<FSoundStartMessage>(TraceMessages.ActiveSoundStartMessages, GetOrCreateActiveSoundEntry,
+		ProcessMessageQueue<FSoundStartMessage>(TraceMessages.ActiveSoundStartMessages, GetActiveSoundEntryFromActiveStartMessage,
 		[this](const FSoundStartMessage& Msg, const TSharedPtr<FSoundDashboardEntry>* OutActiveSoundEntry)
 		{
 			if (OutActiveSoundEntry == nullptr || !OutActiveSoundEntry->IsValid())
@@ -336,7 +342,7 @@ namespace UE::Audio::Insights
 			ActiveSoundEntryRef.EntryType = Msg.EntryType;
 		});
 
-		ProcessMessageQueue<FSoundWaveStartMessage>(TraceMessages.SoundWaveStartMessages, GetActiveSoundEntryFromStartMessage,
+		ProcessMessageQueue<FSoundWaveStartMessage>(TraceMessages.SoundWaveStartMessages, GetActiveSoundEntryFromSoundWaveStartMessage,
 		[this](const FSoundWaveStartMessage& Msg, const TSharedPtr<FSoundDashboardEntry>* OutActiveSoundEntry)
 		{
 			if (OutActiveSoundEntry == nullptr || !OutActiveSoundEntry->IsValid())
@@ -533,72 +539,67 @@ namespace UE::Audio::Insights
 		return false;
 	}
 
-	void FSoundTraceProvider::InitLambdas()
+	void FSoundTraceProvider::GetOrCreateActiveSoundEntry(const FSoundStartMessage& Msg, TSharedPtr<FSoundDashboardEntry>& OutReturnedSoundEntry)
 	{
-		GetOrCreateActiveSoundEntry = [this](const FSoundStartMessage& Msg)
-		{
-			TSharedPtr<FSoundDashboardEntry>* ReturnedSoundEntry = nullptr;
+		OutReturnedSoundEntry = nullptr;
 
-			UpdateDeviceEntry(Msg.DeviceId, Msg.EntryType, [this, &ReturnedSoundEntry, &Msg](TSharedPtr<FSoundDashboardEntry>& CategoryEntry)
+		UpdateDeviceEntry(Msg.DeviceId, Msg.EntryType, [this, &OutReturnedSoundEntry, &Msg](TSharedPtr<FSoundDashboardEntry>& CategoryEntry)
+		{
+			if (!CategoryEntry.IsValid())
 			{
-				if (!CategoryEntry.IsValid())
-				{
-					CategoryEntry = MakeShared<FSoundDashboardEntry>();
-					CategoryEntry->DeviceId = Msg.DeviceId;
-					CategoryEntry->EntryType = Msg.EntryType;
-					CategoryEntry->Name = FSoundTraceProviderPrivate::GetEntryTypeDisplayName(Msg.EntryType);
-				}
+				CategoryEntry = MakeShared<FSoundDashboardEntry>();
+				CategoryEntry->DeviceId = Msg.DeviceId;
+				CategoryEntry->EntryType = Msg.EntryType;
+				CategoryEntry->Name = FSoundTraceProviderPrivate::GetEntryTypeDisplayName(Msg.EntryType);
+			}
 
-				TSharedPtr<IDashboardDataTreeViewEntry>* SoundEntry = CategoryEntry->Children.FindByPredicate([&Msg](TSharedPtr<IDashboardDataTreeViewEntry> ChildEntry)
+			TSharedPtr<IDashboardDataTreeViewEntry>* SoundEntry = CategoryEntry->Children.FindByPredicate([&Msg](TSharedPtr<IDashboardDataTreeViewEntry> ChildEntry)
 				{
 					FSoundDashboardEntry& NewEntry = FSoundTraceProviderPrivate::CastEntry(*ChildEntry.Get());
 					return NewEntry.PlayOrder == Msg.PlayOrder;
 				});
 
-				if (SoundEntry == nullptr || !SoundEntry->IsValid())
-				{
-					TSharedPtr<FSoundDashboardEntry> NewEntry = MakeShared<FSoundDashboardEntry>();
-					NewEntry->DeviceId = Msg.DeviceId;
-					NewEntry->PlayOrder = Msg.PlayOrder;
+			if (SoundEntry == nullptr || !SoundEntry->IsValid())
+			{
+				TSharedPtr<FSoundDashboardEntry> NewEntry = MakeShared<FSoundDashboardEntry>();
+				NewEntry->DeviceId = Msg.DeviceId;
+				NewEntry->PlayOrder = Msg.PlayOrder;
 
-					CategoryEntry->Children.Add(NewEntry);
+				CategoryEntry->Children.Add(NewEntry);
 
-					ActiveSoundToEntryKeysMap.Add(Msg.PlayOrder, SoundEntryKeys{ Msg.EntryType });
+				ActiveSoundToEntryKeysMap.Add(Msg.PlayOrder, SoundEntryKeys{ Msg.EntryType });
 
-					ReturnedSoundEntry = &NewEntry;
-				}
-				else
+				OutReturnedSoundEntry = NewEntry;
+			}
+			else
+			{
+				TSharedPtr<FSoundDashboardEntry> NewEntry = StaticCastSharedPtr<FSoundDashboardEntry>(*SoundEntry);
+				if (NewEntry->TimeoutTimestamp != INVALID_TIMEOUT)
 				{
-					TSharedPtr<FSoundDashboardEntry> NewEntry = StaticCastSharedPtr<FSoundDashboardEntry>(*SoundEntry);
-					if (NewEntry->TimeoutTimestamp != INVALID_TIMEOUT)
+					NewEntry->TimeoutTimestamp = INVALID_TIMEOUT;
+
+					EntriesTimingOut.RemoveAll([&Msg](const SoundMessageIDs& SoundMessageIDs)
 					{
-						NewEntry->TimeoutTimestamp = INVALID_TIMEOUT;
-
-						EntriesTimingOut.RemoveAll([&Msg](const SoundMessageIDs& SoundMessageIDs)
-						{
-							return Msg.DeviceId == SoundMessageIDs.DeviceId && Msg.PlayOrder == SoundMessageIDs.PlayOrder;
-						});
-
-						// If this entry has SoundWaves, they will replay with new play order IDs, so we need to clear the old ones out
-						NewEntry->Children.RemoveAll([](TSharedPtr<IDashboardDataTreeViewEntry> ChildEntry) 
-						{
-							FSoundDashboardEntry& SoundWaveEntry = FSoundTraceProviderPrivate::CastEntry(*ChildEntry.Get());
-							return SoundWaveEntry.TimeoutTimestamp != INVALID_TIMEOUT;
-						});
-					}
-					ReturnedSoundEntry = &NewEntry;
-				}
+						return Msg.DeviceId == SoundMessageIDs.DeviceId && Msg.PlayOrder == SoundMessageIDs.PlayOrder;
+					});
 
-				if (ReturnedSoundEntry == nullptr)
-				{
-					return;
+					// If this entry has SoundWaves, they will replay with new play order IDs, so we need to clear the old ones out
+					NewEntry->Children.RemoveAll([](TSharedPtr<IDashboardDataTreeViewEntry> ChildEntry)
+					{
+						FSoundDashboardEntry& SoundWaveEntry = FSoundTraceProviderPrivate::CastEntry(*ChildEntry.Get());
+						return SoundWaveEntry.TimeoutTimestamp != INVALID_TIMEOUT;
+					});
 				}
+				OutReturnedSoundEntry = NewEntry;
+			}
 
-				(*ReturnedSoundEntry)->Timestamp = Msg.Timestamp;
-			});
+			if (OutReturnedSoundEntry == nullptr)
+			{
+				return;
+			}
 
-			return ReturnedSoundEntry;
-		};
+			OutReturnedSoundEntry->Timestamp = Msg.Timestamp;
+		});
 	}
 
 	void FSoundTraceProvider::GetActiveSoundEntryFromIDs(const uint32 ActiveSoundPlayOrder, const ::Audio::FDeviceId DeviceId, TSharedPtr<FSoundDashboardEntry>& OutActiveSoundEntry)
@@ -657,6 +658,13 @@ namespace UE::Audio::Insights
 
 		DeviceDataMap.Empty();
 
+		TSharedPtr<FSoundDashboardEntry> ActiveSoundEntryReturn = nullptr;
+		auto GetActiveSoundEntryFromActiveStartMessage = [this, &ActiveSoundEntryReturn](const FSoundStartMessage& Msg) -> TSharedPtr<FSoundDashboardEntry>*
+		{
+			GetOrCreateActiveSoundEntry(Msg, ActiveSoundEntryReturn);
+			return &ActiveSoundEntryReturn;
+		};
+
 		TSharedPtr<FSoundDashboardEntry> SoundEntryReturn = nullptr;
 		auto GetActiveSoundEntryFromStopMessage = [this, &SoundEntryReturn](const FSoundStopMessage& Msg) -> TSharedPtr<FSoundDashboardEntry>*
 		{
@@ -680,7 +688,7 @@ namespace UE::Audio::Insights
 				break;
 			}
 
-			TSharedPtr<FSoundDashboardEntry>* ActiveSoundEntry = GetOrCreateActiveSoundEntry(StartCachedMessage);
+			TSharedPtr<FSoundDashboardEntry>* ActiveSoundEntry = GetActiveSoundEntryFromActiveStartMessage(StartCachedMessage);
 			if (ActiveSoundEntry == nullptr || !ActiveSoundEntry->IsValid())
 			{
 				continue;
diff --git a/Engine/Plugins/AudioInsights/Source/AudioInsights/Private/Providers/SoundTraceProvider.h b/Engine/Plugins/AudioInsights/Source/AudioInsights/Private/Providers/SoundTraceProvider.h
--- a/Engine/Plugins/AudioInsights/Source/AudioInsights/Private/Providers/SoundTraceProvider.h
+++ b/Engine/Plugins/AudioInsights/Source/AudioInsights/Private/Providers/SoundTraceProvider.h
@@ -36,7 +36,7 @@ namespace UE::Audio::Insights
 		virtual bool ProcessManuallyUpdatedEntries() override;
 
 	private:
-		void InitLambdas();
+		void GetOrCreateActiveSoundEntry(const FSoundStartMessage& Msg, TSharedPtr<FSoundDashboardEntry>& OutReturnedSoundEntry);
 		void GetActiveSoundEntryFromIDs(const uint32 PlayOrderID, const ::Audio::FDeviceId DeviceId, TSharedPtr<FSoundDashboardEntry>& OutSoundEntry);
 		void RemoveActiveSoundEntry(const ::Audio::FDeviceId DeviceId, TSharedPtr<FSoundDashboardEntry> OutEntry);
 
@@ -61,7 +61,5 @@ namespace UE::Audio::Insights
 
 		TMap<uint32, SoundEntryKeys> ActiveSoundToEntryKeysMap;
 		TArray<SoundMessageIDs> EntriesTimingOut;
-
-		TFunction<TSharedPtr<FSoundDashboardEntry>*(const FSoundStartMessage& Msg)> GetOrCreateActiveSoundEntry;
 	};
 } // namespace UE::Audio::Insights
