[Audio Insights] Sound dashboard: Fix for pinned child soundwaves not displaying correctly after being realized

Submitted on behalf of [REDACTED_HANDLE]

diff --git a/Engine/Plugins/AudioInsights/Source/AudioInsights/Private/Providers/SoundTraceProvider.cpp b/Engine/Plugins/AudioInsights/Source/AudioInsights/Private/Providers/SoundTraceProvider.cpp
--- a/Engine/Plugins/AudioInsights/Source/AudioInsights/Private/Providers/SoundTraceProvider.cpp
+++ b/Engine/Plugins/AudioInsights/Source/AudioInsights/Private/Providers/SoundTraceProvider.cpp
@@ -368,7 +368,8 @@ namespace UE::Audio::Insights
 			SoundWaveEntry->Timestamp = Msg.Timestamp;
 			SoundWaveEntry->Name = *Msg.Name;
 			SoundWaveEntry->EntryType = Msg.EntryType;
-			
+			SoundWaveEntry->PinnedEntryType = (*OutActiveSoundEntry)->PinnedEntryType;
+
 			(*OutActiveSoundEntry)->Children.Add(SoundWaveEntry);
 		});
 
diff --git a/Engine/Plugins/AudioInsights/Source/AudioInsights/Private/Views/SoundDashboardViewFactory.cpp b/Engine/Plugins/AudioInsights/Source/AudioInsights/Private/Views/SoundDashboardViewFactory.cpp
--- a/Engine/Plugins/AudioInsights/Source/AudioInsights/Private/Views/SoundDashboardViewFactory.cpp
+++ b/Engine/Plugins/AudioInsights/Source/AudioInsights/Private/Views/SoundDashboardViewFactory.cpp
@@ -421,11 +421,37 @@ namespace UE::Audio::Insights
 
 	void FPinnedSoundEntryWrapper::CleanUp()
 	{
+		// Remove any pinned items whose original data entries have been removed
+		
+		// Note: Active sounds restart with the same PlayOrderID when realizing after virtualizing, but WaveInstances start with new Play Order IDs, which creates new dashboard entries.
+		// To fix this edge case, when a pinned entry loses it's original entry, double check that a new one hasn't appeared in it's place.
+		// If it has, recreate the child entries.
+		bool bCanBeRecovered = false;
+		if (!bIsCategoryEntry && OriginalDataEntry.IsValid())
+		{
+			const FSoundDashboardEntry& OriginalSoundEntry = FSoundDashboardViewFactoryPrivate::CastEntry(*OriginalDataEntry.Pin());
+
+			// A sound entry may be recovereable if it is still active, is not timing out, has child entries and it is currently pinned
+			bCanBeRecovered = OriginalSoundEntry.PinnedEntryType == FSoundDashboardEntry::EPinnedEntryType::HiddenOriginalEntry 
+							&& OriginalSoundEntry.TimeoutTimestamp == INVALID_TIMEOUT
+							&& OriginalSoundEntry.Children.Num() > 0;
+		}
+
+		bool bRecreateAfterClean = false;
 		for (int Index = PinnedWrapperChildren.Num() - 1; Index >= 0; --Index)
 		{
 			const TSharedPtr<FPinnedSoundEntryWrapper> Child = PinnedWrapperChildren[Index];
 			if (FPinnedSoundEntryWrapperPrivate::CanBeDeleted(Child))
 			{
+				if (bCanBeRecovered)
+				{
+					// If the parent sound is still alive, but the child is no longer valid, destroy and recreate all pinned child entries
+					PinnedSectionEntry->Children.Empty();
+					PinnedWrapperChildren.Empty();
+					bRecreateAfterClean = true;
+					break;
+				}
+
 				TSharedPtr<IDashboardDataTreeViewEntry> DataEntry = Child->GetPinnedSectionEntry();
 				PinnedSectionEntry->Children.Remove(DataEntry);
 				PinnedWrapperChildren.Remove(Child);
@@ -435,6 +461,14 @@ namespace UE::Audio::Insights
 				Child->CleanUp();
 			}
 		}
+
+		if (bRecreateAfterClean)
+		{
+			for (const TSharedPtr<IDashboardDataTreeViewEntry> Child : OriginalDataEntry.Pin()->Children)
+			{
+				AddChildEntry(Child, false);
+			}
+		}
 	}
 
 	void FPinnedSoundEntryWrapper::MarkToDelete()
