[Audio Insights] Sound dashboard: Fix for Linux build error

diff --git a/Engine/Plugins/AudioInsights/Source/AudioInsights/Private/Views/SoundDashboardViewFactory.cpp b/Engine/Plugins/AudioInsights/Source/AudioInsights/Private/Views/SoundDashboardViewFactory.cpp
--- a/Engine/Plugins/AudioInsights/Source/AudioInsights/Private/Views/SoundDashboardViewFactory.cpp
+++ b/Engine/Plugins/AudioInsights/Source/AudioInsights/Private/Views/SoundDashboardViewFactory.cpp
@@ -464,7 +464,7 @@ namespace UE::Audio::Insights
 
 		if (bRecreateAfterClean)
 		{
-			for (const TSharedPtr<IDashboardDataTreeViewEntry> Child : OriginalDataEntry.Pin()->Children)
+			for (const TSharedPtr<IDashboardDataTreeViewEntry>& Child : OriginalDataEntry.Pin()->Children)
 			{
 				AddChildEntry(Child, false);
 			}
