[Audio Insights] Sound dashboard: Bug fix for keyboard shortcuts not working in the dashboard + fixed conflicting Pin/Unpin keyboard shortcut + misc. code style adjustment

diff --git a/Engine/Plugins/AudioInsights/Source/AudioInsights/Private/SoundDashboardCommands.cpp b/Engine/Plugins/AudioInsights/Source/AudioInsights/Private/SoundDashboardCommands.cpp
--- a/Engine/Plugins/AudioInsights/Source/AudioInsights/Private/SoundDashboardCommands.cpp
+++ b/Engine/Plugins/AudioInsights/Source/AudioInsights/Private/SoundDashboardCommands.cpp
@@ -18,7 +18,7 @@ namespace UE::Audio::Insights
 	void FSoundDashboardCommands::RegisterCommands()
 	{
 		UI_COMMAND(Pin, "Pin", "Pins the selected sound in the Pinned category.", EUserInterfaceActionType::Button, FInputChord(EKeys::P, EModifierKey::Control));
-		UI_COMMAND(Unpin, "Unpin", "Removes the selected sound from the Pinned category.", EUserInterfaceActionType::Button, FInputChord(EKeys::P, EModifierKey::Control));
+		UI_COMMAND(Unpin, "Unpin", "Removes the selected sound from the Pinned category.", EUserInterfaceActionType::Button, FInputChord(EKeys::U, EModifierKey::Control));
 		UI_COMMAND(Browse, "Browse To Asset", "Browses to the selected sound asset in the content browser.", EUserInterfaceActionType::Button, FInputChord(EKeys::B, EModifierKey::Control));
 		UI_COMMAND(Edit, "Edit", "Opens the selected sound for edit.", EUserInterfaceActionType::Button, FInputChord(EKeys::E, EModifierKey::Control));
 		// @TODO UE-250399: Hide category pending to implement
diff --git a/Engine/Plugins/AudioInsights/Source/AudioInsights/Private/Views/SoundDashboardViewFactory.cpp b/Engine/Plugins/AudioInsights/Source/AudioInsights/Private/Views/SoundDashboardViewFactory.cpp
--- a/Engine/Plugins/AudioInsights/Source/AudioInsights/Private/Views/SoundDashboardViewFactory.cpp
+++ b/Engine/Plugins/AudioInsights/Source/AudioInsights/Private/Views/SoundDashboardViewFactory.cpp
@@ -1087,21 +1087,26 @@ namespace UE::Audio::Insights
 		MenuBuilder.BeginSection("SoundDashboardActions", LOCTEXT("SoundDashboard_Actions_HeaderText", "Sound Options"));
 
 		{
-			MenuBuilder.AddMenuEntry(Commands.GetPinCommand()
-									, NAME_None
-									, TAttribute<FText>()
-									, TAttribute<FText>()
-									, FSlateStyle::Get().CreateIcon("AudioInsights.Icon.SoundDashboard.Pin")
-									, NAME_None
-									, TAttribute<EVisibility>::CreateLambda([this]() { return SelectionIncludesUnpinnedItem() ? EVisibility::Visible : EVisibility::Collapsed; }));
+			MenuBuilder.AddMenuEntry(
+				Commands.GetPinCommand(),
+				NAME_None,
+				TAttribute<FText>(),
+				TAttribute<FText>(),
+				FSlateStyle::Get().CreateIcon("AudioInsights.Icon.SoundDashboard.Pin"),
+				NAME_None,
+				TAttribute<EVisibility>::CreateLambda([this]() { return SelectionIncludesUnpinnedItem() ? EVisibility::Visible : EVisibility::Collapsed; })
+			);
 			
-			MenuBuilder.AddMenuEntry(Commands.GetUnpinCommand()
-									, NAME_None
-									, TAttribute<FText>()
-									, TAttribute<FText>()
-									, FSlateStyle::Get().CreateIcon("AudioInsights.Icon.SoundDashboard.Pin")
-									, NAME_None
-									, TAttribute<EVisibility>::CreateLambda([this]() { return SelectionIncludesUnpinnedItem() ? EVisibility::Collapsed : EVisibility::Visible; }));
+			MenuBuilder.AddMenuEntry(
+				Commands.GetUnpinCommand(),
+				NAME_None,
+				TAttribute<FText>(),
+				TAttribute<FText>(),
+				FSlateStyle::Get().CreateIcon("AudioInsights.Icon.SoundDashboard.Pin"),
+				NAME_None,
+				TAttribute<EVisibility>::CreateLambda([this]() { return SelectionIncludesUnpinnedItem() ? EVisibility::Collapsed : EVisibility::Visible; })
+			);
+
 #if WITH_EDITOR
 			MenuBuilder.AddMenuEntry(Commands.GetBrowseCommand(), NAME_None, TAttribute<FText>(), TAttribute<FText>(), FSlateStyle::Get().CreateIcon("AudioInsights.Icon.SoundDashboard.Browse"));
 			MenuBuilder.AddMenuEntry(Commands.GetEditCommand(),   NAME_None, TAttribute<FText>(), TAttribute<FText>(), FSlateStyle::Get().CreateIcon("AudioInsights.Icon.SoundDashboard.Edit"));
@@ -1116,6 +1121,11 @@ namespace UE::Audio::Insights
 		return MenuBuilder.MakeWidget();
 	}
 
+	FReply FSoundDashboardViewFactory::OnDataRowKeyInput(const FGeometry& InGeometry, const FKeyEvent& InKeyEvent) const
+	{
+		return (CommandList && CommandList->ProcessCommandBindings(InKeyEvent)) ? FReply::Handled() : FReply::Unhandled();
+	}
+
 	const TMap<FName, FTraceTreeDashboardViewFactory::FHeaderRowColumnData>& FSoundDashboardViewFactory::GetHeaderRowColumns() const
 	{
 		static const TMap<FName, FTraceTreeDashboardViewFactory::FHeaderRowColumnData> HeaderRowColumnData =
diff --git a/Engine/Plugins/AudioInsights/Source/AudioInsights/Public/Views/SoundDashboardViewFactory.h b/Engine/Plugins/AudioInsights/Source/AudioInsights/Public/Views/SoundDashboardViewFactory.h
--- a/Engine/Plugins/AudioInsights/Source/AudioInsights/Public/Views/SoundDashboardViewFactory.h
+++ b/Engine/Plugins/AudioInsights/Source/AudioInsights/Public/Views/SoundDashboardViewFactory.h
@@ -122,6 +122,7 @@ namespace UE::Audio::Insights
 		virtual TSharedRef<ITableRow> OnGenerateRow(TSharedPtr<IDashboardDataTreeViewEntry> Item, const TSharedRef<STableViewBase>& OwnerTable) override;
 
 		virtual TSharedPtr<SWidget> OnConstructContextMenu() override;
+		virtual FReply OnDataRowKeyInput(const FGeometry& InGeometry, const FKeyEvent& InKeyEvent) const override;
 
 		virtual const TMap<FName, FTraceTreeDashboardViewFactory::FHeaderRowColumnData>& GetHeaderRowColumns() const override;
 		virtual const TMap<FName, FTraceTreeDashboardViewFactory::FColumnData>& GetColumns() const override;
