[Audio Insights] Sound dashboard: Bug fix for Pinned category always visible no matter which category filters were enabled

diff --git a/Engine/Plugins/AudioInsights/Source/AudioInsights/Private/Views/SoundDashboardViewFactory.cpp b/Engine/Plugins/AudioInsights/Source/AudioInsights/Private/Views/SoundDashboardViewFactory.cpp
--- a/Engine/Plugins/AudioInsights/Source/AudioInsights/Private/Views/SoundDashboardViewFactory.cpp
+++ b/Engine/Plugins/AudioInsights/Source/AudioInsights/Private/Views/SoundDashboardViewFactory.cpp
@@ -290,16 +290,6 @@ namespace UE::Audio::Insights
 		}
 #endif // WITH_EDITOR
 
-
-		constexpr ESoundDashboardFilterFlags AllFilterFlags = 
-			ESoundDashboardFilterFlags::MetaSound        |
-			ESoundDashboardFilterFlags::SoundCue         |
-			ESoundDashboardFilterFlags::ProceduralSource |
-			ESoundDashboardFilterFlags::SoundWave        |
-			ESoundDashboardFilterFlags::SoundCueTemplate |
-			ESoundDashboardFilterFlags::Pinned           |
-			ESoundDashboardFilterFlags::Hidden;
-
 		const FText FiltersName = LOCTEXT("SoundDashboard_Filter_CategoryText", "Filters");
 		const FText FiltersTooltip = LOCTEXT("CurveFiltersToolTip", "Filters what kind of sounds types can be displayed.");
 
@@ -711,17 +701,26 @@ namespace UE::Audio::Insights
 			.CustomFilters(Filters)
 			.OnFilterChanged_Lambda([this, Filters]()
 			{
-				SelectedFilterFlags = ESoundDashboardFilterFlags::None;
-
-				for(const TSharedRef<FFilterBase<ESoundDashboardFilterFlags>>& Filter : Filters)
+				auto GetActiveFilterFlags = [&Filters]()
 				{
-					TSharedRef<FSoundDashboardFilter> SoundDashboardFilter = StaticCastSharedRef<FSoundDashboardFilter>(Filter);
+					ESoundDashboardFilterFlags ActiveFilterFlags = ESoundDashboardFilterFlags::None;
 
-					if(SoundDashboardFilter->IsActive())
+					for (const TSharedRef<FFilterBase<ESoundDashboardFilterFlags>>& Filter : Filters)
 					{
-						SelectedFilterFlags |= SoundDashboardFilter->GetFlags();
+						TSharedRef<FSoundDashboardFilter> SoundDashboardFilter = StaticCastSharedRef<FSoundDashboardFilter>(Filter);
+
+						if (SoundDashboardFilter->IsActive())
+						{
+							ActiveFilterFlags |= SoundDashboardFilter->GetFlags();
+						}
 					}
-				}
+
+					// By default, if there are no active filters selected it means that all filters are enabled
+					return ActiveFilterFlags != ESoundDashboardFilterFlags::None ? ActiveFilterFlags : AllFilterFlags;
+				};
+
+				SelectedFilterFlags = GetActiveFilterFlags();
+				bIsPinnedCategoryFilterEnabled = EnumHasAnyFlags(SelectedFilterFlags, ESoundDashboardFilterFlags::Pinned);
 
 				UpdateFilterReason = EProcessReason::FilterUpdated;
 			});
@@ -995,6 +994,13 @@ namespace UE::Audio::Insights
 			.Visibility_Lambda([this, Item]()
 			{
 				const FSoundDashboardEntry& SoundCategoryEntry = CastEntry(*Item);
+
+				// Don't show pinned items if Pinned category is not enabled
+				if (SoundCategoryEntry.PinnedEntryType == FSoundDashboardEntry::EPinnedEntryType::PinnedCopy && !bIsPinnedCategoryFilterEnabled)
+				{
+					return EVisibility::Hidden;
+				}
+
 				const int32 NumUnpinnedChildren = GetNumChildrenWithoutPinEntryType(*Item, FSoundDashboardEntry::EPinnedEntryType::HiddenOriginalEntry);
 
 				const bool bRowShouldBeVisible = SoundCategoryEntry.bIsVisible 
@@ -1012,8 +1018,7 @@ namespace UE::Audio::Insights
 		// Filter by category
 		FTraceTreeDashboardViewFactory::FilterEntries<FSoundTraceProvider>([this](const IDashboardDataTreeViewEntry& InEntry)
 		{
-			const bool bAllCategoryFiltersEnabled = SelectedFilterFlags == ESoundDashboardFilterFlags::None || EnumHasAllFlags(SelectedFilterFlags, FSoundDashboardViewFactoryPrivate::AllFilterFlags);
-			if (bAllCategoryFiltersEnabled)
+			if (SelectedFilterFlags == AllFilterFlags)
 			{
 				return true;
 			}
diff --git a/Engine/Plugins/AudioInsights/Source/AudioInsights/Public/Views/SoundDashboardViewFactory.h b/Engine/Plugins/AudioInsights/Source/AudioInsights/Public/Views/SoundDashboardViewFactory.h
--- a/Engine/Plugins/AudioInsights/Source/AudioInsights/Public/Views/SoundDashboardViewFactory.h
+++ b/Engine/Plugins/AudioInsights/Source/AudioInsights/Public/Views/SoundDashboardViewFactory.h
@@ -19,8 +19,9 @@ namespace UE::Audio::Insights
 		ProceduralSource = 1 << 2,
 		SoundWave = 1 << 3,
 		SoundCueTemplate = 1 << 4,
-		Pinned = 1 << 5,
-		Hidden = 1 << 6
+		Pinned = 1 << 5
+		// @TODO UE-250399: Hide category pending to implement
+		//, Hidden = 1 << 6
 	};
 
 	ENUM_CLASS_FLAGS(ESoundDashboardFilterFlags);
@@ -166,6 +167,16 @@ namespace UE::Audio::Insights
 		// @TODO UE-250399: Hide category pending to implement
 		//void HideSound();
 
+		static constexpr ESoundDashboardFilterFlags AllFilterFlags = 
+			ESoundDashboardFilterFlags::MetaSound        |
+			ESoundDashboardFilterFlags::SoundCue         |
+			ESoundDashboardFilterFlags::ProceduralSource |
+			ESoundDashboardFilterFlags::SoundWave        |
+			ESoundDashboardFilterFlags::SoundCueTemplate |
+			ESoundDashboardFilterFlags::Pinned;
+			// @TODO UE-250399: Hide category pending to implement
+			//| ESoundDashboardFilterFlags::Hidden;
+
 		TSharedPtr<FPinnedSoundEntryWrapper> PinnedItemEntries;
 		TArray<TSharedPtr<IDashboardDataTreeViewEntry>> FullTree;
 
@@ -173,6 +184,7 @@ namespace UE::Audio::Insights
 		TSharedPtr<SWidget> SoundsFilterBar;
 		TSharedPtr<SWidget> SoundsFilterBarButton;
 
-		ESoundDashboardFilterFlags SelectedFilterFlags = ESoundDashboardFilterFlags::None;
+		ESoundDashboardFilterFlags SelectedFilterFlags = AllFilterFlags;
+		bool bIsPinnedCategoryFilterEnabled = true;
 	};
 } // namespace UE::Audio::Insights
