[Audio Insights] Disable Viewport dashboard until further notice

diff --git a/Engine/Plugins/AudioInsights/Source/AudioInsightsEditor/Private/AudioInsightsEditorModule.cpp b/Engine/Plugins/AudioInsights/Source/AudioInsightsEditor/Private/AudioInsightsEditorModule.cpp
--- a/Engine/Plugins/AudioInsights/Source/AudioInsightsEditor/Private/AudioInsightsEditorModule.cpp
+++ b/Engine/Plugins/AudioInsights/Source/AudioInsightsEditor/Private/AudioInsightsEditorModule.cpp
@@ -35,7 +35,8 @@ namespace UE::Audio::Insights
 
 			DashboardFactory = MakeShared<FEditorDashboardFactory>();
 			
-			DashboardFactory->RegisterViewFactory(MakeShared<FViewportDashboardViewFactory>());
+			// @TODO UE-274216: Decide what to do with the Viewport dashboard
+			//DashboardFactory->RegisterViewFactory(MakeShared<FViewportDashboardViewFactory>());
 			DashboardFactory->RegisterViewFactory(MakeShared<FLogDashboardViewFactory>());
 			DashboardFactory->RegisterViewFactory(MakeShared<FMixerSourceDashboardViewFactory>());
 			DashboardFactory->RegisterViewFactory(MakeShared<FVirtualLoopDashboardViewFactory>());
