[Audio Insights] Fixed Sounds Dashboard sometimes displaying uninitialized parameter values in standalone on non-live traces.
- When we try to get an active sound entry, make sure we reset any shared pointers used to return the entry.
- Previously returned dashboard entries received from parameter messages inside OnTimingViewTimeMarkerChanged were being kept alive and were passing the validation checks, resulting in uninitialized data overwriting the actual data.

diff --git a/Engine/Plugins/AudioInsights/Source/AudioInsights/Private/Providers/SoundTraceProvider.cpp b/Engine/Plugins/AudioInsights/Source/AudioInsights/Private/Providers/SoundTraceProvider.cpp
--- a/Engine/Plugins/AudioInsights/Source/AudioInsights/Private/Providers/SoundTraceProvider.cpp
+++ b/Engine/Plugins/AudioInsights/Source/AudioInsights/Private/Providers/SoundTraceProvider.cpp
@@ -296,6 +296,8 @@ namespace UE::Audio::Insights
 		TSharedPtr<FSoundDashboardEntry> ActiveSoundEntryReturn = nullptr;
 		auto GetActiveSoundEntryFromActiveStartMessage = [this, &ActiveSoundEntryReturn](const FSoundStartMessage& Msg) -> TSharedPtr<FSoundDashboardEntry>*
 		{
+			ActiveSoundEntryReturn = nullptr;
+
 			GetOrCreateActiveSoundEntry(Msg, ActiveSoundEntryReturn);
 			return &ActiveSoundEntryReturn;
 		};
@@ -303,6 +305,8 @@ namespace UE::Audio::Insights
 		TSharedPtr<FSoundDashboardEntry> SoundWaveActiveEntryReturn = nullptr;
 		auto GetActiveSoundEntryFromSoundWaveStartMessage = [this, &SoundWaveActiveEntryReturn](const FSoundWaveStartMessage& Msg) -> TSharedPtr<FSoundDashboardEntry>*
 		{
+			SoundWaveActiveEntryReturn = nullptr;
+
 			GetActiveSoundEntryFromIDs(Msg.ActiveSoundPlayOrder, Msg.DeviceId, SoundWaveActiveEntryReturn);
 			return &SoundWaveActiveEntryReturn;
 		};
@@ -310,6 +314,8 @@ namespace UE::Audio::Insights
 		TSharedPtr<FSoundDashboardEntry> ParamReturn = nullptr;
 		auto GetActiveSoundEntryFromParameterMessage = [this, &ParamReturn](const FMixerSourceParameterMessage& Msg) -> TSharedPtr<FSoundDashboardEntry>*
 		{
+			ParamReturn = nullptr;
+
 			GetActiveSoundEntryFromIDs(Msg.ActiveSoundPlayOrder, Msg.DeviceId, ParamReturn);
 			return &ParamReturn;
 		};
@@ -317,6 +323,8 @@ namespace UE::Audio::Insights
 		TSharedPtr<FSoundDashboardEntry> SoundEntryReturn = nullptr;
 		auto GetSoundEntryFromStopMessage = [this, &SoundEntryReturn](const FSoundStopMessage& Msg) -> TSharedPtr<FSoundDashboardEntry>*
 		{
+			SoundEntryReturn = nullptr;
+
 			GetActiveSoundEntryFromIDs(Msg.PlayOrder, Msg.DeviceId, SoundEntryReturn);
 			return &SoundEntryReturn;
 		};
@@ -662,6 +670,8 @@ namespace UE::Audio::Insights
 		TSharedPtr<FSoundDashboardEntry> ActiveSoundEntryReturn = nullptr;
 		auto GetActiveSoundEntryFromActiveStartMessage = [this, &ActiveSoundEntryReturn](const FSoundStartMessage& Msg) -> TSharedPtr<FSoundDashboardEntry>*
 		{
+			ActiveSoundEntryReturn = nullptr;
+
 			GetOrCreateActiveSoundEntry(Msg, ActiveSoundEntryReturn);
 			return &ActiveSoundEntryReturn;
 		};
@@ -669,6 +679,8 @@ namespace UE::Audio::Insights
 		TSharedPtr<FSoundDashboardEntry> SoundEntryReturn = nullptr;
 		auto GetActiveSoundEntryFromStopMessage = [this, &SoundEntryReturn](const FSoundStopMessage& Msg) -> TSharedPtr<FSoundDashboardEntry>*
 		{
+			SoundEntryReturn = nullptr;
+
 			GetActiveSoundEntryFromIDs(Msg.PlayOrder, Msg.DeviceId, SoundEntryReturn);
 			return &SoundEntryReturn;
 		};
@@ -676,6 +688,8 @@ namespace UE::Audio::Insights
 		TSharedPtr<FSoundDashboardEntry> ParamReturn = nullptr;
 		auto GetActiveSoundEntryFromParameterMessage = [this, &ParamReturn](const FMixerSourceParameterMessage& Msg) -> TSharedPtr<FSoundDashboardEntry>*
 		{
+			ParamReturn = nullptr;
+
 			GetActiveSoundEntryFromIDs(Msg.ActiveSoundPlayOrder, Msg.DeviceId, ParamReturn);
 			return &ParamReturn;
 		};
