[Audio Insights] Sound dashboard: Fix for -> don't mute/solo pinned items if a non pinned category is set to mute/solo

diff --git a/Engine/Plugins/AudioInsights/Source/AudioInsights/Private/Views/SoundDashboardViewFactory.cpp b/Engine/Plugins/AudioInsights/Source/AudioInsights/Private/Views/SoundDashboardViewFactory.cpp
--- a/Engine/Plugins/AudioInsights/Source/AudioInsights/Private/Views/SoundDashboardViewFactory.cpp
+++ b/Engine/Plugins/AudioInsights/Source/AudioInsights/Private/Views/SoundDashboardViewFactory.cpp
@@ -152,6 +152,12 @@ namespace UE::Audio::Insights
 
 				const FSoundDashboardEntry& SoundEntry = FSoundDashboardViewFactoryPrivate::CastEntry(InEntry);
 
+				// Skip setting mute/solo, a copy of this entry is currently in the Pinned category 
+				if (SoundEntry.PinnedEntryType == FSoundDashboardEntry::EPinnedEntryType::HiddenOriginalEntry)
+				{
+					return;
+				}
+
 				const FName SoundAssetDisplayName(SoundEntry.GetDisplayName().ToString());
 
 				const bool bIsSoundCueType = SoundEntry.EntryType == ESoundDashboardEntryType::SoundCue || SoundEntry.EntryType == ESoundDashboardEntryType::SoundCueTemplate;
@@ -234,6 +240,12 @@ namespace UE::Audio::Insights
 
 				const FSoundDashboardEntry& SoundEntry = FSoundDashboardViewFactoryPrivate::CastEntry(InEntry);
 
+				// Treat hidden original entries as muted/soloed to ensure the parent category reflects the correct state
+				if (SoundEntry.PinnedEntryType == FSoundDashboardEntry::EPinnedEntryType::HiddenOriginalEntry)
+				{
+					return true;
+				}
+
 				const FName SoundAssetDisplayName(SoundEntry.GetDisplayName().ToString());
 
 				const bool bIsSoundCueType = SoundEntry.EntryType == ESoundDashboardEntryType::SoundCue || SoundEntry.EntryType == ESoundDashboardEntryType::SoundCueTemplate;
