[Audio Insights] Sound dashboard: Added tooltips to mute/solo selected buttons + added missing Clear All Mute/Solos button

diff --git a/Engine/Plugins/AudioInsights/Source/AudioInsights/Private/Views/SoundDashboardViewFactory.cpp b/Engine/Plugins/AudioInsights/Source/AudioInsights/Private/Views/SoundDashboardViewFactory.cpp
--- a/Engine/Plugins/AudioInsights/Source/AudioInsights/Private/Views/SoundDashboardViewFactory.cpp
+++ b/Engine/Plugins/AudioInsights/Source/AudioInsights/Private/Views/SoundDashboardViewFactory.cpp
@@ -300,6 +300,16 @@ namespace UE::Audio::Insights
 
 			return false;
 		}
+
+		void ClearMutesAndSolos()
+		{
+#if ENABLE_AUDIO_DEBUG
+			if (FAudioDeviceManager* AudioDeviceManager = FAudioDeviceManager::Get())
+			{
+				AudioDeviceManager->GetDebugger().ClearMutesAndSolos();
+			}
+#endif
+		}
 #endif // WITH_EDITOR
 
 		const FText FiltersName = LOCTEXT("SoundDashboard_Filter_CategoryText", "Filters");
@@ -599,6 +609,7 @@ namespace UE::Audio::Insights
 			[
 				SNew(SButton)
 				.ButtonStyle(&FAppStyle::Get().GetWidgetStyle<FButtonStyle>("SimpleButton"))
+				.ToolTipText(LOCTEXT("SoundDashboard_MuteTooltipText", "Toggles the mute state of the selected items."))
 				.OnClicked_Lambda([this]()
 				{
 					if (FilteredEntriesListView.IsValid())
@@ -619,6 +630,7 @@ namespace UE::Audio::Insights
 			[
 				SNew(SButton)
 				.ButtonStyle(&FAppStyle::Get().GetWidgetStyle<FButtonStyle>("SimpleButton"))
+				.ToolTipText(LOCTEXT("SoundDashboard_SoloTooltipText", "Toggles the solo state of the selected items."))
 				.OnClicked_Lambda([this]()
 				{
 					if (FilteredEntriesListView.IsValid())
@@ -631,6 +643,24 @@ namespace UE::Audio::Insights
 				[
 					CreateButtonContentWidget("AudioInsights.Icon.SoundDashboard.Solo", LOCTEXT("SoundDashboard_SoloButtonText", "Solo Selected"), "SmallButtonText")
 				]
+			]
+			// Clear Mutes/Solos button
+			+ SHorizontalBox::Slot()
+			.AutoWidth()
+			.Padding(8.0f, 6.0f)
+			[
+				SNew(SButton)
+				.ButtonStyle(&FAppStyle::Get().GetWidgetStyle<FButtonStyle>("SimpleButton"))
+				.ToolTipText(LOCTEXT("SoundsDashboard_ClearMutesAndSolosTooltipText", "Clears all assigned mute/solo states."))
+				.OnClicked_Lambda([]()
+				{
+					ClearMutesAndSolos();
+
+					return FReply::Handled();
+				})
+				[
+					CreateButtonContentWidget("AudioInsights.Icon.SoundDashboard.Reset", LOCTEXT("SoundsDashboard_ClearMutesAndSolosButtonText", "Clear All Mutes/Solos"), "SmallButtonText")
+				]
 			];
 	}
 #endif // WITH_EDITOR
diff --git a/Engine/Plugins/AudioInsights/Source/AudioInsights/Public/AudioInsightsStyle.h b/Engine/Plugins/AudioInsights/Source/AudioInsights/Public/AudioInsightsStyle.h
--- a/Engine/Plugins/AudioInsights/Source/AudioInsights/Public/AudioInsightsStyle.h
+++ b/Engine/Plugins/AudioInsights/Source/AudioInsights/Public/AudioInsightsStyle.h
@@ -175,6 +175,7 @@ namespace UE::Audio::Insights
 			Set("AudioInsights.Icon.SoundDashboard.Mute",             new IMAGE_BRUSH_SVG(TEXT("Icons/SoundDashboard/mute"),             Icon16));
 			Set("AudioInsights.Icon.SoundDashboard.Pin",              new IMAGE_BRUSH_SVG(TEXT("Icons/SoundDashboard/pin"),              Icon16));
 			Set("AudioInsights.Icon.SoundDashboard.ProceduralSource", new IMAGE_BRUSH_SVG(TEXT("Icons/SoundDashboard/proceduralsource"), Icon16));
+			Set("AudioInsights.Icon.SoundDashboard.Reset",            new IMAGE_BRUSH_SVG(TEXT("Icons/SoundDashboard/reset"),            Icon16));
 			Set("AudioInsights.Icon.SoundDashboard.Solo",             new IMAGE_BRUSH_SVG(TEXT("Icons/SoundDashboard/solo"),             Icon16));
 			Set("AudioInsights.Icon.SoundDashboard.SoundCue",         new IMAGE_BRUSH_SVG(TEXT("Icons/SoundDashboard/soundcue"),         Icon16));
 			Set("AudioInsights.Icon.SoundDashboard.SoundWave",        new IMAGE_BRUSH_SVG(TEXT("Icons/SoundDashboard/soundwave"),        Icon16));
