[Audio Insights] Sound dashboard: Retain height for filter button and search box widgets when adjusting dashboard width

diff --git a/Engine/Plugins/AudioInsights/Source/AudioInsights/Private/Views/TreeDashboardViewFactory.cpp b/Engine/Plugins/AudioInsights/Source/AudioInsights/Private/Views/TreeDashboardViewFactory.cpp
--- a/Engine/Plugins/AudioInsights/Source/AudioInsights/Private/Views/TreeDashboardViewFactory.cpp
+++ b/Engine/Plugins/AudioInsights/Source/AudioInsights/Private/Views/TreeDashboardViewFactory.cpp
@@ -80,7 +80,12 @@ namespace UE::Audio::Insights
 				.AutoWidth()
 				.Padding(FilterBarButton ? 3.0f : 0.0f, 0.0f, 0.0f, 0.0f)
 				[
-					FilterBarButton ? FilterBarButton.ToSharedRef() : SNullWidget::NullWidget
+					SNew(SVerticalBox)
+					+ SVerticalBox::Slot()
+					.MaxHeight(30.0f)
+					[
+						FilterBarButton ? FilterBarButton.ToSharedRef() : SNullWidget::NullWidget
+					]
 				]
 				// Search box
 				+ SHorizontalBox::Slot()
@@ -88,6 +93,7 @@ namespace UE::Audio::Insights
 				[
 					SNew(SVerticalBox)
 					+ SVerticalBox::Slot()
+					.MaxHeight(22.0f)
 					.Padding(0.0f, 4.0f, 0.0f, 6.0f)
 					[
 						SAssignNew(SearchBoxWidget, SSearchBox)
