[Audio Insights] Sound dashboard: Always reset visibility when doing text filtering since children need to get updated when there is still text in the search box filter

diff --git a/Engine/Plugins/AudioInsights/Source/AudioInsights/Private/Views/SoundDashboardViewFactory.cpp b/Engine/Plugins/AudioInsights/Source/AudioInsights/Private/Views/SoundDashboardViewFactory.cpp
--- a/Engine/Plugins/AudioInsights/Source/AudioInsights/Private/Views/SoundDashboardViewFactory.cpp
+++ b/Engine/Plugins/AudioInsights/Source/AudioInsights/Private/Views/SoundDashboardViewFactory.cpp
@@ -1101,14 +1101,12 @@ namespace UE::Audio::Insights
 			{
 				if (Entry.IsValid())
 				{
+					ResetVisibility(*Entry);
+
 					if (bFilterHasText)
 					{
 						SetFilteredVisibility(*Entry, FilterString);
 					}
-					else
-					{
-						ResetVisibility(*Entry);
-					}
 				}
 			}
 		}
