[Audio Insights] Sound dashboard: In multiple selection we need to discard children entries to avoid double mute/solo toggling

diff --git a/Engine/Plugins/AudioInsights/Source/AudioInsights/Private/Views/SoundDashboardViewFactory.cpp b/Engine/Plugins/AudioInsights/Source/AudioInsights/Private/Views/SoundDashboardViewFactory.cpp
--- a/Engine/Plugins/AudioInsights/Source/AudioInsights/Private/Views/SoundDashboardViewFactory.cpp
+++ b/Engine/Plugins/AudioInsights/Source/AudioInsights/Private/Views/SoundDashboardViewFactory.cpp
@@ -914,6 +914,23 @@ namespace UE::Audio::Insights
 		return PinnedItemEntries.IsValid() && PinnedItemEntries->GetPinnedSectionEntry().IsValid() && PinnedItemEntries->GetPinnedSectionEntry()->Children.Contains(InEntry);
 	}
 
+
+	bool FSoundDashboardViewFactory::IsDescendant(const TSharedPtr<IDashboardDataTreeViewEntry>& InEntry, const TSharedPtr<IDashboardDataTreeViewEntry>& InChildCandidate) const
+	{
+		if (InEntry.IsValid() && EntryCanHaveChildren(InEntry.ToSharedRef()))
+		{
+			for (const TSharedPtr<IDashboardDataTreeViewEntry>& ChildEntry : InEntry->Children)
+			{
+				if (ChildEntry == InChildCandidate || IsDescendant(ChildEntry, InChildCandidate))
+				{
+					return true;
+				}
+			}
+		}
+
+		return false;
+	}
+
 	TSharedRef<SWidget> FSoundDashboardViewFactory::GenerateWidgetForColumn(const TSharedRef<FTraceTreeDashboardViewFactory::SRowWidget>& InRowWidget, const TSharedRef<IDashboardDataTreeViewEntry>& InRowData, const FName& InColumn)
 	{
 		using namespace FSoundDashboardViewFactoryPrivate;
@@ -1654,7 +1671,43 @@ namespace UE::Audio::Insights
 	{
 		using namespace FSoundDashboardViewFactoryPrivate;
 
-		for (const TSharedPtr<IDashboardDataTreeViewEntry>& Entry : InEntries)
+		TArray<TSharedPtr<IDashboardDataTreeViewEntry>> EntriesToMuteSolo;
+
+		// In multiple selection we need to discard children entries to avoid double mute/solo toggling
+		if (InEntries.Num() > 1)
+		{
+			EntriesToMuteSolo.Reserve(InEntries.Num());
+
+			for (const TSharedPtr<IDashboardDataTreeViewEntry>& Entry : InEntries)
+			{
+				if (!Entry.IsValid())
+				{
+					continue;
+				}
+
+				bool bIsTopLevelEntry = true;
+
+				for (const TSharedPtr<IDashboardDataTreeViewEntry>& Other : InEntries)
+				{
+					if (Other != Entry && IsDescendant(Other, Entry))
+					{
+						bIsTopLevelEntry = false;
+						break;
+					}
+				}
+
+				if (bIsTopLevelEntry)
+				{
+					EntriesToMuteSolo.Add(Entry);
+				}
+			}
+		}
+		else
+		{
+			EntriesToMuteSolo = InEntries;
+		}
+
+		for (const TSharedPtr<IDashboardDataTreeViewEntry>& Entry : EntriesToMuteSolo)
 		{
 			if (!Entry.IsValid())
 			{
@@ -1667,7 +1720,7 @@ namespace UE::Audio::Insights
 
 				for (const TSharedPtr<IDashboardDataTreeViewEntry>& ChildEntry : Entry->Children)
 				{
-					if (Entry.IsValid())
+					if (ChildEntry.IsValid())
 					{
 						FSoundDashboardViewFactoryPrivate::SetMuteSolo(*ChildEntry, InMuteSoloMode, !bAreChildrenMuteSolo);
 					}
diff --git a/Engine/Plugins/AudioInsights/Source/AudioInsights/Public/Views/SoundDashboardViewFactory.h b/Engine/Plugins/AudioInsights/Source/AudioInsights/Public/Views/SoundDashboardViewFactory.h
--- a/Engine/Plugins/AudioInsights/Source/AudioInsights/Public/Views/SoundDashboardViewFactory.h
+++ b/Engine/Plugins/AudioInsights/Source/AudioInsights/Public/Views/SoundDashboardViewFactory.h
@@ -116,6 +116,7 @@ namespace UE::Audio::Insights
 		bool IsRootItem(const TSharedRef<IDashboardDataTreeViewEntry>& InEntry) const;
 		bool EntryCanHaveChildren(const TSharedRef<IDashboardDataTreeViewEntry>& InEntry) const;
 		bool IsCategoryItem(const TSharedRef<IDashboardDataTreeViewEntry>& InEntry) const;
+		bool IsDescendant(const TSharedPtr<IDashboardDataTreeViewEntry>& InEntry, const TSharedPtr<IDashboardDataTreeViewEntry>& InChildCandidate) const;
 
 		TSharedRef<SWidget> GenerateWidgetForRootColumn(const TSharedRef<FTraceTreeDashboardViewFactory::SRowWidget>& InRowWidget, const TSharedRef<IDashboardDataTreeViewEntry>& InRowData, const FName& InColumn, const FText& InValueText);
 		virtual TSharedRef<SWidget> GenerateWidgetForColumn(const TSharedRef<FTraceTreeDashboardViewFactory::SRowWidget>& InRowWidget, const TSharedRef<IDashboardDataTreeViewEntry>& InRowData, const FName& InColumn) override;
