[Audio Insights] Sound dashboard: small category refactor + fixes for wrong item count when doing text filtering + items not disappearing when count is 0

diff --git a/Engine/Plugins/AudioInsights/Source/AudioInsights/Private/Providers/SoundTraceProvider.cpp b/Engine/Plugins/AudioInsights/Source/AudioInsights/Private/Providers/SoundTraceProvider.cpp
--- a/Engine/Plugins/AudioInsights/Source/AudioInsights/Private/Providers/SoundTraceProvider.cpp
+++ b/Engine/Plugins/AudioInsights/Source/AudioInsights/Private/Providers/SoundTraceProvider.cpp
@@ -560,6 +560,7 @@ namespace UE::Audio::Insights
 				CategoryEntry->DeviceId = Msg.DeviceId;
 				CategoryEntry->EntryType = Msg.EntryType;
 				CategoryEntry->Name = FSoundTraceProviderPrivate::GetEntryTypeDisplayName(Msg.EntryType);
+				CategoryEntry->bIsCategory = true;
 			}
 
 			TSharedPtr<IDashboardDataTreeViewEntry>* SoundEntry = CategoryEntry->Children.FindByPredicate([&Msg](TSharedPtr<IDashboardDataTreeViewEntry> ChildEntry)
diff --git a/Engine/Plugins/AudioInsights/Source/AudioInsights/Private/Views/SoundDashboardViewFactory.cpp b/Engine/Plugins/AudioInsights/Source/AudioInsights/Private/Views/SoundDashboardViewFactory.cpp
--- a/Engine/Plugins/AudioInsights/Source/AudioInsights/Private/Views/SoundDashboardViewFactory.cpp
+++ b/Engine/Plugins/AudioInsights/Source/AudioInsights/Private/Views/SoundDashboardViewFactory.cpp
@@ -90,6 +90,18 @@ namespace UE::Audio::Insights
 			}
 		}
 
+		bool IsCategoryItem(const IDashboardDataTreeViewEntry& InEntry)
+		{
+			const FSoundDashboardEntry& SoundEntry = CastEntry(InEntry);
+			return SoundEntry.bIsCategory;
+		}
+
+		bool IsVisible(const IDashboardDataTreeViewEntry& InEntry)
+		{
+			const FSoundDashboardEntry& SoundEntry = CastEntry(InEntry);
+			return SoundEntry.bIsVisible;
+		}
+
 		bool HasPinEntryType(const IDashboardDataTreeViewEntry& InEntry, const FSoundDashboardEntry::EPinnedEntryType PinnedEntryType)
 		{
 			const FSoundDashboardEntry& SoundEntry = CastEntry(InEntry);
@@ -106,7 +118,7 @@ namespace UE::Audio::Insights
 					continue;
 				}
 
-				if (!HasPinEntryType(*Child, ExcludedPinnedEntryType))
+				if (!HasPinEntryType(*Child, ExcludedPinnedEntryType) && IsVisible(*Child))
 				{
 					++NumChildrenWithoutType;
 				}
@@ -126,13 +138,19 @@ namespace UE::Audio::Insights
 
 						if (SoundDashboardEntry.TimeoutTimestamp == INVALID_TIMEOUT)
 						{
-							if (HasPinEntryType(*InChild, FSoundDashboardEntry::EPinnedEntryType::HiddenOriginalEntry))
+							if (HasPinEntryType(*InChild, FSoundDashboardEntry::EPinnedEntryType::HiddenOriginalEntry) || !SoundDashboardEntry.bIsVisible)
 							{
 								return Accum;
 							}
 
 							const int32 NumNestedChildren = GetNumChildrenWithoutPinEntryType(*InChild, FSoundDashboardEntry::EPinnedEntryType::HiddenOriginalEntry);
-							return Accum + (NumNestedChildren > 0 ? NumNestedChildren : 1);
+
+							if (NumNestedChildren > 0)
+							{
+								return Accum + NumNestedChildren;
+							}
+
+							return SoundDashboardEntry.bIsCategory ? Accum : Accum + 1;
 						}
 					}
 
@@ -331,7 +349,9 @@ namespace UE::Audio::Insights
 	{
 		bool CanBeDeleted(const TSharedPtr<FPinnedSoundEntryWrapper>& Entry)
 		{
-			return !Entry->EntryIsValid() || (Entry->bIsCategoryEntry && Entry->PinnedWrapperChildren.IsEmpty());
+			using namespace FSoundDashboardViewFactoryPrivate;
+
+			return !Entry->EntryIsValid() || (IsCategoryItem(*Entry->GetPinnedSectionEntry()) && Entry->PinnedWrapperChildren.IsEmpty());
 		}
 
 		void CopyDataToPinnedEntry(FSoundDashboardEntry& PinnedEntry, const FSoundDashboardEntry& OriginalEntry)
@@ -369,9 +389,8 @@ namespace UE::Audio::Insights
 
 	/////////////////////////////////////////////////////////////////////////////////////////
 	// FPinnedSoundEntryWrapper
-	FPinnedSoundEntryWrapper::FPinnedSoundEntryWrapper(const TSharedPtr<IDashboardDataTreeViewEntry>& OriginalEntry, const bool bIsCategory)
-		: bIsCategoryEntry(bIsCategory)
-		, OriginalDataEntry(OriginalEntry)
+	FPinnedSoundEntryWrapper::FPinnedSoundEntryWrapper(const TSharedPtr<IDashboardDataTreeViewEntry>& OriginalEntry)
+		: OriginalDataEntry(OriginalEntry)
 	{
 		if (!OriginalEntry.IsValid())
 		{
@@ -388,16 +407,15 @@ namespace UE::Audio::Insights
 
 		for (const TSharedPtr<IDashboardDataTreeViewEntry>& Child : OriginalEntry->Children)
 		{
-			AddChildEntry(Child.ToSharedRef(), false /*bIsCategory*/);
+			AddChildEntry(Child.ToSharedRef());
 		}
 	}
 
-	TSharedPtr<FPinnedSoundEntryWrapper> FPinnedSoundEntryWrapper::AddChildEntry(const TSharedPtr<IDashboardDataTreeViewEntry> Child, const bool bIsCategory)
+	TSharedPtr<FPinnedSoundEntryWrapper> FPinnedSoundEntryWrapper::AddChildEntry(const TSharedPtr<IDashboardDataTreeViewEntry> Child)
 	{
-		TSharedPtr<FPinnedSoundEntryWrapper> NewChild = MakeShared<FPinnedSoundEntryWrapper>(Child, bIsCategory);
+		TSharedPtr<FPinnedSoundEntryWrapper> NewChild = MakeShared<FPinnedSoundEntryWrapper>(Child);
 
 		FSoundDashboardEntry& NewChildSound = FSoundDashboardViewFactoryPrivate::CastEntry(*NewChild->GetPinnedSectionEntry());
-		NewChildSound.PinnedEntryType = FSoundDashboardEntry::EPinnedEntryType::PinnedCopy;
 
 		PinnedWrapperChildren.Add(NewChild);
 		PinnedSectionEntry->Children.Add(NewChild->GetPinnedSectionEntry());
@@ -433,13 +451,15 @@ namespace UE::Audio::Insights
 
 	void FPinnedSoundEntryWrapper::CleanUp()
 	{
+		using namespace FSoundDashboardViewFactoryPrivate;
+
 		// Remove any pinned items whose original data entries have been removed
 		
 		// Note: Active sounds restart with the same PlayOrderID when realizing after virtualizing, but WaveInstances start with new Play Order IDs, which creates new dashboard entries.
 		// To fix this edge case, when a pinned entry loses it's original entry, double check that a new one hasn't appeared in it's place.
 		// If it has, recreate the child entries.
 		bool bCanBeRecovered = false;
-		if (!bIsCategoryEntry && OriginalDataEntry.IsValid())
+		if (!IsCategoryItem(*PinnedSectionEntry) && OriginalDataEntry.IsValid())
 		{
 			const FSoundDashboardEntry& OriginalSoundEntry = FSoundDashboardViewFactoryPrivate::CastEntry(*OriginalDataEntry.Pin());
 
@@ -478,7 +498,7 @@ namespace UE::Audio::Insights
 		{
 			for (const TSharedPtr<IDashboardDataTreeViewEntry>& Child : OriginalDataEntry.Pin()->Children)
 			{
-				AddChildEntry(Child, false);
+				AddChildEntry(Child);
 			}
 		}
 	}
@@ -490,7 +510,7 @@ namespace UE::Audio::Insights
 
 	bool FPinnedSoundEntryWrapper::EntryIsValid() const
 	{
-		return PinnedSectionEntry.IsValid() && (OriginalDataEntry.IsValid() || bIsCategoryEntry);
+		return PinnedSectionEntry.IsValid() && (OriginalDataEntry.IsValid() || FSoundDashboardViewFactoryPrivate::IsCategoryItem(*PinnedSectionEntry));
 	}
 
 	/////////////////////////////////////////////////////////////////////////////////////////
@@ -897,21 +917,13 @@ namespace UE::Audio::Insights
 
 	bool FSoundDashboardViewFactory::EntryCanHaveChildren(const TSharedRef<IDashboardDataTreeViewEntry>& InEntry) const
 	{
+		using namespace FSoundDashboardViewFactoryPrivate;
+
 		const FSoundDashboardEntry& SoundEntry = FSoundDashboardViewFactoryPrivate::CastEntry(*InEntry);
 
 		const bool bIsSoundCueType = SoundEntry.EntryType == ESoundDashboardEntryType::SoundCue || SoundEntry.EntryType == ESoundDashboardEntryType::SoundCueTemplate;
 
-		return IsCategoryItem(InEntry) || bIsSoundCueType;
-	}
-
-	bool FSoundDashboardViewFactory::IsCategoryItem(const TSharedRef<IDashboardDataTreeViewEntry>& InEntry) const
-	{
-		if (IsRootItem(InEntry))
-		{
-			return true;
-		}
-
-		return PinnedItemEntries.IsValid() && PinnedItemEntries->GetPinnedSectionEntry().IsValid() && PinnedItemEntries->GetPinnedSectionEntry()->Children.Contains(InEntry);
+		return IsCategoryItem(*InEntry) || bIsSoundCueType;
 	}
 
 
@@ -1052,19 +1064,21 @@ namespace UE::Audio::Insights
 		return SNew(SRowWidget, OwnerTable, Item, AsShared())
 			.Visibility_Lambda([this, Item]()
 			{
-				const FSoundDashboardEntry& SoundCategoryEntry = CastEntry(*Item);
+				const FSoundDashboardEntry& SoundEntry = CastEntry(*Item);
 
-				// Don't show pinned items if Pinned category is not enabled
-				if (SoundCategoryEntry.PinnedEntryType == FSoundDashboardEntry::EPinnedEntryType::PinnedCopy && !bIsPinnedCategoryFilterEnabled)
+				if (SoundEntry.PinnedEntryType == FSoundDashboardEntry::EPinnedEntryType::PinnedCopy)
 				{
-					return EVisibility::Hidden;
+					if (!bIsPinnedCategoryFilterEnabled || (SoundEntry.bIsCategory && CountNumChildren(*Item) == 0))
+					{
+						return EVisibility::Hidden;
+					}
 				}
 
 				const int32 NumUnpinnedChildren = GetNumChildrenWithoutPinEntryType(*Item, FSoundDashboardEntry::EPinnedEntryType::HiddenOriginalEntry);
 
-				const bool bRowShouldBeVisible = SoundCategoryEntry.bIsVisible 
-												&& SoundCategoryEntry.PinnedEntryType != FSoundDashboardEntry::EPinnedEntryType::HiddenOriginalEntry 
-												&& (!IsCategoryItem(Item.ToSharedRef()) || (NumUnpinnedChildren > 0));
+				const bool bRowShouldBeVisible = SoundEntry.bIsVisible 
+												&& SoundEntry.PinnedEntryType != FSoundDashboardEntry::EPinnedEntryType::HiddenOriginalEntry 
+												&& (!IsCategoryItem(*Item) || (NumUnpinnedChildren > 0));
 
 				return bRowShouldBeVisible ? EVisibility::Visible : EVisibility::Hidden;
 			});
@@ -1131,11 +1145,13 @@ namespace UE::Audio::Insights
 
 	TSharedPtr<SWidget> FSoundDashboardViewFactory::OnConstructContextMenu()
 	{
+		using namespace FSoundDashboardViewFactoryPrivate;
+
 		// Skip right click for category items
 		if (FilteredEntriesListView.IsValid() && FilteredEntriesListView->GetSelectedItems().Num() == 1)
 		{
 			const TSharedPtr<IDashboardDataTreeViewEntry> SelectedEntry = FilteredEntriesListView->GetSelectedItems()[0];
-			if (IsCategoryItem(SelectedEntry.ToSharedRef()))
+			if (IsCategoryItem(*SelectedEntry))
 			{
 				return SNullWidget::NullWidget;
 			}
@@ -1592,7 +1608,7 @@ namespace UE::Audio::Insights
 		for (const TSharedPtr<IDashboardDataTreeViewEntry>& Entry : OutTree)
 		{
 			const FSoundDashboardEntry& EntryData = CastEntry(*Entry);
-			if (IsCategoryItem(Entry.ToSharedRef()))
+			if (IsCategoryItem(*Entry))
 			{
 				return;
 			}
@@ -1885,7 +1901,7 @@ namespace UE::Audio::Insights
 
 			if (FoundMatchingSelectedEntry != nullptr)
 			{
-				if (IsCategoryItem(PinnedWrapperEntry->GetPinnedSectionEntry().ToSharedRef()))
+				if (IsCategoryItem(*PinnedWrapperEntry->GetPinnedSectionEntry()))
 				{
 					// If the current parent pinned item is a category, move the child to the unpinned area and continue
 					MarkBranchAsPinned(OriginalChildEntry, false /*bIsPinned*/);
@@ -1934,8 +1950,9 @@ namespace UE::Audio::Insights
 		PinnedCategory->Name = FSoundDashboardViewFactoryPrivate::PinnedCategoryName.ToString();
 		PinnedCategory->EntryType = ESoundDashboardEntryType::Pinned;
 		PinnedCategory->PinnedEntryType = FSoundDashboardEntry::EPinnedEntryType::PinnedCopy;
+		PinnedCategory->bIsCategory = true;
 
-		PinnedItemEntries = MakeShared<FPinnedSoundEntryWrapper>(PinnedCategory, true /*bIsCategoryEntry*/);
+		PinnedItemEntries = MakeShared<FPinnedSoundEntryWrapper>(PinnedCategory);
 	}
 
 	void FSoundDashboardViewFactory::CreatePinnedEntry(TSharedPtr<IDashboardDataTreeViewEntry> Entry)
@@ -1977,7 +1994,7 @@ namespace UE::Audio::Insights
 				// If we didn't find this entry already inside the pinned area, add it here
 				if (FoundExisingEntry == nullptr)
 				{
-					PinnedCategoryEntry->AddChildEntry(Entry, false /*bIsCategory*/);
+					PinnedCategoryEntry->AddChildEntry(Entry);
 				}
 				break;
 			}
@@ -1993,8 +2010,8 @@ namespace UE::Audio::Insights
 					TSharedPtr<FSoundDashboardEntry> NewPinnedCategoryEntry = MakeShared<FSoundDashboardEntry>(CastEntry(*DataCategoryEntry));
 					NewPinnedCategoryEntry->Children.Reset();
 
-					TSharedPtr<FPinnedSoundEntryWrapper> PinnedCategory = PinnedItemEntries->AddChildEntry(NewPinnedCategoryEntry, true /*bIsCategory*/);
-					PinnedCategory->AddChildEntry(Entry, false /*bIsCategory*/);
+					TSharedPtr<FPinnedSoundEntryWrapper> PinnedCategory = PinnedItemEntries->AddChildEntry(NewPinnedCategoryEntry);
+					PinnedCategory->AddChildEntry(Entry);
 
 					break;
 				}
diff --git a/Engine/Plugins/AudioInsights/Source/AudioInsights/Public/Messages/SoundTraceMessages.h b/Engine/Plugins/AudioInsights/Source/AudioInsights/Public/Messages/SoundTraceMessages.h
--- a/Engine/Plugins/AudioInsights/Source/AudioInsights/Public/Messages/SoundTraceMessages.h
+++ b/Engine/Plugins/AudioInsights/Source/AudioInsights/Public/Messages/SoundTraceMessages.h
@@ -169,6 +169,7 @@ namespace UE::Audio::Insights
 		FString Name;
 		ESoundDashboardEntryType EntryType = ESoundDashboardEntryType::None;
 		EPinnedEntryType PinnedEntryType = EPinnedEntryType::None;
+		bool bIsCategory = false;
 		bool bShouldForceExpandChildren = true;
 		bool bIsVisible = true;
 
diff --git a/Engine/Plugins/AudioInsights/Source/AudioInsights/Public/Views/SoundDashboardViewFactory.h b/Engine/Plugins/AudioInsights/Source/AudioInsights/Public/Views/SoundDashboardViewFactory.h
--- a/Engine/Plugins/AudioInsights/Source/AudioInsights/Public/Views/SoundDashboardViewFactory.h
+++ b/Engine/Plugins/AudioInsights/Source/AudioInsights/Public/Views/SoundDashboardViewFactory.h
@@ -72,12 +72,12 @@ namespace UE::Audio::Insights
 	{
 	public:
 		FPinnedSoundEntryWrapper() = delete;
-		FPinnedSoundEntryWrapper(const TSharedPtr<IDashboardDataTreeViewEntry>& OriginalEntry, const bool bIsCategory);
+		FPinnedSoundEntryWrapper(const TSharedPtr<IDashboardDataTreeViewEntry>& OriginalEntry);
 
 		TSharedPtr<IDashboardDataTreeViewEntry> GetPinnedSectionEntry() const { return PinnedSectionEntry; }
 		TSharedPtr<IDashboardDataTreeViewEntry> GetOriginalDataEntry() const { return OriginalDataEntry.IsValid() ? OriginalDataEntry.Pin() : nullptr; }
 
-		TSharedPtr<FPinnedSoundEntryWrapper> AddChildEntry(const TSharedPtr<IDashboardDataTreeViewEntry> Child, const bool bIsCategory);
+		TSharedPtr<FPinnedSoundEntryWrapper> AddChildEntry(const TSharedPtr<IDashboardDataTreeViewEntry> Child);
 
 		void UpdateParams();
 
@@ -87,7 +87,6 @@ namespace UE::Audio::Insights
 		bool EntryIsValid() const;
 
 		TArray<TSharedPtr<FPinnedSoundEntryWrapper>> PinnedWrapperChildren;
-		const bool bIsCategoryEntry = false;
 
 	private:
 		TSharedPtr<IDashboardDataTreeViewEntry> PinnedSectionEntry;
@@ -115,7 +114,6 @@ namespace UE::Audio::Insights
 
 		bool IsRootItem(const TSharedRef<IDashboardDataTreeViewEntry>& InEntry) const;
 		bool EntryCanHaveChildren(const TSharedRef<IDashboardDataTreeViewEntry>& InEntry) const;
-		bool IsCategoryItem(const TSharedRef<IDashboardDataTreeViewEntry>& InEntry) const;
 		bool IsDescendant(const TSharedPtr<IDashboardDataTreeViewEntry>& InEntry, const TSharedPtr<IDashboardDataTreeViewEntry>& InChildCandidate) const;
 
 		TSharedRef<SWidget> GenerateWidgetForRootColumn(const TSharedRef<FTraceTreeDashboardViewFactory::SRowWidget>& InRowWidget, const TSharedRef<IDashboardDataTreeViewEntry>& InRowData, const FName& InColumn, const FText& InValueText);
