[Audio Insights] Make Sounds Dashboard the first tab that opens in Audio Insights by default

Submitted on behalf of [REDACTED_HANDLE]

diff --git a/Engine/Plugins/AudioInsights/Source/AudioInsights/Private/AudioInsightsDashboardFactory.cpp b/Engine/Plugins/AudioInsights/Source/AudioInsights/Private/AudioInsightsDashboardFactory.cpp
--- a/Engine/Plugins/AudioInsights/Source/AudioInsights/Private/AudioInsightsDashboardFactory.cpp
+++ b/Engine/Plugins/AudioInsights/Source/AudioInsights/Private/AudioInsightsDashboardFactory.cpp
@@ -211,7 +211,7 @@ namespace UE::Audio::Insights
 			}
 		}
 
-		AnalysisTabStack->SetForegroundTab(FName("MixerSources"));
+		AnalysisTabStack->SetForegroundTab(FName("Sounds"));
 
 		return FTabManager::NewLayout("AudioDashboard_Layout_v1")
 		->AddArea
diff --git a/Engine/Plugins/AudioInsights/Source/AudioInsightsEditor/Private/AudioInsightsEditorDashboardFactory.cpp b/Engine/Plugins/AudioInsights/Source/AudioInsightsEditor/Private/AudioInsightsEditorDashboardFactory.cpp
--- a/Engine/Plugins/AudioInsights/Source/AudioInsightsEditor/Private/AudioInsightsEditorDashboardFactory.cpp
+++ b/Engine/Plugins/AudioInsights/Source/AudioInsightsEditor/Private/AudioInsightsEditorDashboardFactory.cpp
@@ -616,7 +616,7 @@ namespace UE::Audio::Insights
 			}
 		}
 
-		AnalysisTabStack->SetForegroundTab(FName("MixerSources"));
+		AnalysisTabStack->SetForegroundTab(FName("Sounds"));
 
 		return FTabManager::NewLayout("AudioDashboard_Editor_Layout_v2")
 		->AddArea
