[Audio Insights] Fix for Sounds Dashboard losing expansion/pinned states of entries when moving the time marker in Standalone

Submitted on behalf of [REDACTED_HANDLE]

diff --git a/Engine/Plugins/AudioInsights/Source/AudioInsights/Private/Providers/SoundTraceProvider.cpp b/Engine/Plugins/AudioInsights/Source/AudioInsights/Private/Providers/SoundTraceProvider.cpp
--- a/Engine/Plugins/AudioInsights/Source/AudioInsights/Private/Providers/SoundTraceProvider.cpp
+++ b/Engine/Plugins/AudioInsights/Source/AudioInsights/Private/Providers/SoundTraceProvider.cpp
@@ -666,10 +666,8 @@ namespace UE::Audio::Insights
 			return;
 		}
 
-		DeviceDataMap.Empty();
-
 		TSharedPtr<FSoundDashboardEntry> ActiveSoundEntryReturn = nullptr;
-		auto GetActiveSoundEntryFromActiveStartMessage = [this, &ActiveSoundEntryReturn](const FSoundStartMessage& Msg) -> TSharedPtr<FSoundDashboardEntry>*
+		auto GetOrCreateActiveSoundEntryFromActiveStartMessage = [this, &ActiveSoundEntryReturn](const FSoundStartMessage& Msg) -> TSharedPtr<FSoundDashboardEntry>*
 		{
 			ActiveSoundEntryReturn = nullptr;
 
@@ -677,6 +675,15 @@ namespace UE::Audio::Insights
 			return &ActiveSoundEntryReturn;
 		};
 
+		TSharedPtr<FSoundDashboardEntry> StartMessageReturn = nullptr;
+		auto GetActiveSoundEntryFromStartMessage = [this, &StartMessageReturn](const FSoundStartMessage& Msg) -> TSharedPtr<FSoundDashboardEntry>*
+		{
+			StartMessageReturn = nullptr;
+
+			GetActiveSoundEntryFromIDs(Msg.PlayOrder, Msg.DeviceId, StartMessageReturn);
+			return &StartMessageReturn;
+		};
+
 		TSharedPtr<FSoundDashboardEntry> SoundEntryReturn = nullptr;
 		auto GetActiveSoundEntryFromStopMessage = [this, &SoundEntryReturn](const FSoundStopMessage& Msg) -> TSharedPtr<FSoundDashboardEntry>*
 		{
@@ -696,7 +703,6 @@ namespace UE::Audio::Insights
 		};
 
 		// Collect all the start messages registered until this point in time 
-		TSharedPtr<FSoundDashboardEntry>* ReturnedSoundEntry = nullptr;
 		for (const FSoundStartMessage& StartCachedMessage : SessionCachedMessages->StartCachedMessages)
 		{
 			if (StartCachedMessage.Timestamp > TimeMarker)
@@ -704,7 +710,7 @@ namespace UE::Audio::Insights
 				break;
 			}
 
-			TSharedPtr<FSoundDashboardEntry>* ActiveSoundEntry = GetActiveSoundEntryFromActiveStartMessage(StartCachedMessage);
+			TSharedPtr<FSoundDashboardEntry>* ActiveSoundEntry = GetOrCreateActiveSoundEntryFromActiveStartMessage(StartCachedMessage);
 			if (ActiveSoundEntry == nullptr || !ActiveSoundEntry->IsValid())
 			{
 				continue;
@@ -763,6 +769,21 @@ namespace UE::Audio::Insights
 				RemoveActiveSoundEntry(StopCachedMessage.DeviceId, *ActiveSoundEntry);
 			}
 		}
+
+		// Re-run through start messages after stop messages have been processed
+		// Any entries found with timestamps higher than the current time marker need to be cleared out too.
+		for (const FSoundStartMessage& StartCachedMessage : SessionCachedMessages->StartCachedMessages)
+		{
+			if (StartCachedMessage.Timestamp > TimeMarker)
+			{
+				TSharedPtr<FSoundDashboardEntry>* FoundSoundEntry = GetActiveSoundEntryFromStartMessage(StartCachedMessage);
+
+				if (FoundSoundEntry && FoundSoundEntry->IsValid())
+				{
+					RemoveActiveSoundEntry(StartCachedMessage.DeviceId, *FoundSoundEntry);
+				}
+			}
+		}
 		
 		// For now we only retrieve information from AudioDeviceId 1 (main device in standalone games)
 		const FDeviceData* DeviceData = DeviceDataMap.Find(1);
