[Audio Insights] Create option to enable/disable showing sounds that have recently stopped

Submitted on behalf of [REDACTED_HANDLE]

diff --git a/Engine/Plugins/AudioInsights/Source/AudioInsights/Private/Views/SoundDashboardViewFactory.cpp b/Engine/Plugins/AudioInsights/Source/AudioInsights/Private/Views/SoundDashboardViewFactory.cpp
--- a/Engine/Plugins/AudioInsights/Source/AudioInsights/Private/Views/SoundDashboardViewFactory.cpp
+++ b/Engine/Plugins/AudioInsights/Source/AudioInsights/Private/Views/SoundDashboardViewFactory.cpp
@@ -96,10 +96,10 @@ namespace UE::Audio::Insights
 			return SoundEntry.bIsCategory;
 		}
 
-		bool IsVisible(const IDashboardDataTreeViewEntry& InEntry)
+		bool IsVisible(const IDashboardDataTreeViewEntry& InEntry, const bool bShowRecentlyStoppedSounds)
 		{
 			const FSoundDashboardEntry& SoundEntry = CastEntry(InEntry);
-			return SoundEntry.bIsVisible;
+			return SoundEntry.bIsVisible && (bShowRecentlyStoppedSounds || SoundEntry.TimeoutTimestamp == INVALID_TIMEOUT);
 		}
 
 		bool HasPinEntryType(const IDashboardDataTreeViewEntry& InEntry, const FSoundDashboardEntry::EPinnedEntryType PinnedEntryType)
@@ -108,7 +108,7 @@ namespace UE::Audio::Insights
 			return SoundEntry.PinnedEntryType == PinnedEntryType;
 		}
 
-		int32 GetNumChildrenWithoutPinEntryType(const IDashboardDataTreeViewEntry& InEntry, const FSoundDashboardEntry::EPinnedEntryType ExcludedPinnedEntryType)
+		int32 GetNumChildrenWithoutPinEntryType(const IDashboardDataTreeViewEntry& InEntry, const FSoundDashboardEntry::EPinnedEntryType ExcludedPinnedEntryType, const bool bShowRecentlyStoppedSounds)
 		{
 			int32 NumChildrenWithoutType = 0;
 			for (const TSharedPtr<IDashboardDataTreeViewEntry>& Child : InEntry.Children)
@@ -118,7 +118,7 @@ namespace UE::Audio::Insights
 					continue;
 				}
 
-				if (!HasPinEntryType(*Child, ExcludedPinnedEntryType) && IsVisible(*Child))
+				if (!HasPinEntryType(*Child, ExcludedPinnedEntryType) && IsVisible(*Child, bShowRecentlyStoppedSounds))
 				{
 					++NumChildrenWithoutType;
 				}
@@ -127,10 +127,10 @@ namespace UE::Audio::Insights
 			return NumChildrenWithoutType;
 		}
 
-		int32 CountNumChildren(const IDashboardDataTreeViewEntry& InEntry)
+		int32 CountNumChildren(const IDashboardDataTreeViewEntry& InEntry, const bool bShowRecentlyStoppedSounds)
 		{
 			const uint32 TotalNumChildren = Algo::Accumulate(InEntry.Children, 0,
-				[](uint32 Accum, TSharedPtr<IDashboardDataTreeViewEntry> InChild)
+				[bShowRecentlyStoppedSounds](uint32 Accum, TSharedPtr<IDashboardDataTreeViewEntry> InChild)
 				{
 					if (InChild.IsValid())
 					{
@@ -143,7 +143,7 @@ namespace UE::Audio::Insights
 								return Accum;
 							}
 
-							const int32 NumNestedChildren = GetNumChildrenWithoutPinEntryType(*InChild, FSoundDashboardEntry::EPinnedEntryType::HiddenOriginalEntry);
+							const int32 NumNestedChildren = GetNumChildrenWithoutPinEntryType(*InChild, FSoundDashboardEntry::EPinnedEntryType::HiddenOriginalEntry, bShowRecentlyStoppedSounds);
 
 							if (NumNestedChildren > 0)
 							{
@@ -681,6 +681,16 @@ namespace UE::Audio::Insights
 				[
 					CreateButtonContentWidget("AudioInsights.Icon.SoundDashboard.Reset", LOCTEXT("SoundsDashboard_ClearMutesAndSolosButtonText", "Clear All Mutes/Solos"), "SmallButtonText")
 				]
+			]
+			// Empty Spacing
+			+ SHorizontalBox::Slot()
+			.FillWidth(1.0f)
+			// Show recently stopped sounds button
+			+ SHorizontalBox::Slot()
+			.HAlign(HAlign_Right)
+			.AutoWidth()
+			[
+				MakeShowRecentlyStoppedSoundsWidget()
 			];
 	}
 #endif // WITH_EDITOR
@@ -817,6 +827,14 @@ namespace UE::Audio::Insights
 			[
 				MakeMuteSoloWidget()
 			]
+#else
+			+ SVerticalBox::Slot()
+			.AutoHeight()
+			.HAlign(HAlign_Right)
+			.Padding(0.0f, 2.0f, 0.0f, 0.0f)
+			[
+				MakeShowRecentlyStoppedSoundsWidget()
+			]
 #endif // WITH_EDITOR
 			+ SVerticalBox::Slot()
 			.HAlign(HAlign_Fill)
@@ -899,9 +917,9 @@ namespace UE::Audio::Insights
 				[
 					SNew(STextBlock)
 					.ColorAndOpacity(FSlateColor(FColor::White.WithAlpha(128)))
-					.Text_Lambda([InRowData]()
+					.Text_Lambda([this, InRowData]()
 					{
-						const uint32 TotalNumChildren = CountNumChildren(*InRowData);
+						const uint32 TotalNumChildren = CountNumChildren(*InRowData, bShowRecentlyStoppedSounds);
 						return FText::FromString("(" + FString::FromInt(TotalNumChildren) + ")");
 					})
 				];
@@ -1047,9 +1065,9 @@ namespace UE::Audio::Insights
 					InColumn == "Name" && !InRowData->Children.IsEmpty()
 						? SNew(STextBlock)
 						  .ColorAndOpacity(FSlateColor(FColor::White.WithAlpha(128)))
-						  .Text_Lambda([InRowData]()
+						  .Text_Lambda([this, InRowData]()
 						  {
-							  const uint32 TotalNumChildren = CountNumChildren(*InRowData);
+							  const uint32 TotalNumChildren = CountNumChildren(*InRowData, bShowRecentlyStoppedSounds);
 							  return FText::FromString("(" + FString::FromInt(TotalNumChildren) + ")");
 						  })
 						: SNullWidget::NullWidget
@@ -1068,15 +1086,16 @@ namespace UE::Audio::Insights
 
 				if (SoundEntry.PinnedEntryType == FSoundDashboardEntry::EPinnedEntryType::PinnedCopy)
 				{
-					if (!bIsPinnedCategoryFilterEnabled || (SoundEntry.bIsCategory && CountNumChildren(*Item) == 0))
+					if (!bIsPinnedCategoryFilterEnabled || (SoundEntry.bIsCategory && CountNumChildren(*Item, bShowRecentlyStoppedSounds) == 0))
 					{
 						return EVisibility::Hidden;
 					}
 				}
 
-				const int32 NumUnpinnedChildren = GetNumChildrenWithoutPinEntryType(*Item, FSoundDashboardEntry::EPinnedEntryType::HiddenOriginalEntry);
+				const int32 NumUnpinnedChildren = GetNumChildrenWithoutPinEntryType(*Item, FSoundDashboardEntry::EPinnedEntryType::HiddenOriginalEntry, bShowRecentlyStoppedSounds);
 
 				const bool bRowShouldBeVisible = SoundEntry.bIsVisible 
+												&& (bShowRecentlyStoppedSounds || SoundEntry.TimeoutTimestamp == INVALID_TIMEOUT)
 												&& SoundEntry.PinnedEntryType != FSoundDashboardEntry::EPinnedEntryType::HiddenOriginalEntry 
 												&& (!IsCategoryItem(*Item) || (NumUnpinnedChildren > 0));
 
@@ -1651,6 +1670,51 @@ namespace UE::Audio::Insights
 		RecursiveSort(DataViewEntries, Predicate, bSortByChildren);
 	}
 
+	TSharedRef<SWidget> FSoundDashboardViewFactory::MakeShowRecentlyStoppedSoundsWidget()
+	{
+		return SNew(SHorizontalBox)
+			+ SHorizontalBox::Slot()
+			.AutoWidth()
+			.Padding(8.0f, 8.0f)
+			[
+				SNew(SCheckBox)
+				.HAlign(EHorizontalAlignment::HAlign_Center)
+				.Style(&FSlateStyle::Get().GetWidgetStyle<FCheckBoxStyle>("SoundDashboard.MuteSoloButton"))
+				.ToolTipText(LOCTEXT("SoundsDashboard_ShowStoppedSoundsTooltip", "Shows sounds that have recently stopped playing"))
+				.IsChecked(bShowRecentlyStoppedSounds ? ECheckBoxState::Checked : ECheckBoxState::Unchecked)
+				.OnCheckStateChanged_Lambda([this](ECheckBoxState NewState)
+				{
+					bShowRecentlyStoppedSounds = NewState == ECheckBoxState::Checked;
+				})
+				[
+					SNew(SHorizontalBox)
+					+ SHorizontalBox::Slot()
+					.HAlign(HAlign_Left)
+					.VAlign(VAlign_Center)
+					.AutoWidth()
+					.Padding(2.0f, 0.0f)
+					[
+						SNew(SImage)
+				 		.Image_Lambda([this]()
+				 		{
+				 			const FName IconName = bShowRecentlyStoppedSounds ? "AudioInsights.Icon.SoundDashboard.Visible" : "AudioInsights.Icon.SoundDashboard.Invisible";
+		 
+				 			return FSlateStyle::Get().GetBrush(IconName);
+				 		})
+					]
+					+ SHorizontalBox::Slot()
+					.HAlign(HAlign_Right)
+					.VAlign(VAlign_Center)
+					.AutoWidth()
+					.Padding(2.0f, 0.0f)
+					[
+						SNew(STextBlock)
+						.Text(LOCTEXT("SoundsDashboard_ShowStoppedSounds", "Show Stopped"))
+					]
+				]
+			];
+	}
+
 #if WITH_EDITOR
 	TSharedRef<SWidget> FSoundDashboardViewFactory::CreateMuteSoloButton(const TSharedRef<IDashboardDataTreeViewEntry>& InRowData, 
 		const FName& InColumn, 
diff --git a/Engine/Plugins/AudioInsights/Source/AudioInsights/Public/AudioInsightsStyle.h b/Engine/Plugins/AudioInsights/Source/AudioInsights/Public/AudioInsightsStyle.h
--- a/Engine/Plugins/AudioInsights/Source/AudioInsights/Public/AudioInsightsStyle.h
+++ b/Engine/Plugins/AudioInsights/Source/AudioInsights/Public/AudioInsightsStyle.h
@@ -181,6 +181,8 @@ namespace UE::Audio::Insights
 			Set("AudioInsights.Icon.SoundDashboard.SoundWave",        new IMAGE_BRUSH_SVG(TEXT("Icons/SoundDashboard/soundwave"),        Icon16));
 			Set("AudioInsights.Icon.SoundDashboard.Tab",              new IMAGE_BRUSH_SVG(TEXT("Icons/SoundDashboard/tab"),              Icon16));
 			Set("AudioInsights.Icon.SoundDashboard.Transparent",      new IMAGE_BRUSH_SVG(TEXT("Icons/SoundDashboard/transparent"),      Icon16));
+			Set("AudioInsights.Icon.SoundDashboard.Visible",		  new IMAGE_BRUSH_SVG(TEXT("Icons/SoundDashboard/visible"),			 Icon16));
+			Set("AudioInsights.Icon.SoundDashboard.Invisible",		  new IMAGE_BRUSH_SVG(TEXT("Icons/SoundDashboard/invisible"),		 Icon16));
 			
 			// Category colors
 			Set("SoundDashboard.MetaSoundColor", FLinearColor(0.008f, 0.76f, 0.078f));
diff --git a/Engine/Plugins/AudioInsights/Source/AudioInsights/Public/Views/SoundDashboardViewFactory.h b/Engine/Plugins/AudioInsights/Source/AudioInsights/Public/Views/SoundDashboardViewFactory.h
--- a/Engine/Plugins/AudioInsights/Source/AudioInsights/Public/Views/SoundDashboardViewFactory.h
+++ b/Engine/Plugins/AudioInsights/Source/AudioInsights/Public/Views/SoundDashboardViewFactory.h
@@ -134,6 +134,8 @@ namespace UE::Audio::Insights
 		void RecursiveSort(TArray<TSharedPtr<IDashboardDataTreeViewEntry>>& OutTree, TFunctionRef<bool(const FSoundDashboardEntry&, const FSoundDashboardEntry&)> Predicate, const bool bSortByChildren);
 		void SortByPredicate(TFunctionRef<bool(const FSoundDashboardEntry&, const FSoundDashboardEntry&)> Predicate, const bool bSortByChildren = true);
 
+		TSharedRef<SWidget> MakeShowRecentlyStoppedSoundsWidget();
+
 #if WITH_EDITOR
 		TSharedRef<SWidget> MakeMuteSoloWidget();
 
@@ -185,5 +187,6 @@ namespace UE::Audio::Insights
 
 		ESoundDashboardFilterFlags SelectedFilterFlags = AllFilterFlags;
 		bool bIsPinnedCategoryFilterEnabled = true;
+		bool bShowRecentlyStoppedSounds = false;
 	};
 } // namespace UE::Audio::Insights
