[Audio Insights] Fixes for pinned items disappearing to early/late

Submitted on behalf of [REDACTED_HANDLE]

diff --git a/Engine/Plugins/AudioInsights/Source/AudioInsights/Private/Views/SoundDashboardViewFactory.cpp b/Engine/Plugins/AudioInsights/Source/AudioInsights/Private/Views/SoundDashboardViewFactory.cpp
--- a/Engine/Plugins/AudioInsights/Source/AudioInsights/Private/Views/SoundDashboardViewFactory.cpp
+++ b/Engine/Plugins/AudioInsights/Source/AudioInsights/Private/Views/SoundDashboardViewFactory.cpp
@@ -127,23 +127,23 @@ namespace UE::Audio::Insights
 			return NumChildrenWithoutType;
 		}
 
-		int32 CountNumChildren(const IDashboardDataTreeViewEntry& InEntry, const bool bShowRecentlyStoppedSounds)
+		int32 CountNumChildren(const IDashboardDataTreeViewEntry& InEntry, const bool bIncludeTimingOutSounds = false)
 		{
 			const uint32 TotalNumChildren = Algo::Accumulate(InEntry.Children, 0,
-				[bShowRecentlyStoppedSounds](uint32 Accum, TSharedPtr<IDashboardDataTreeViewEntry> InChild)
+				[bIncludeTimingOutSounds](uint32 Accum, TSharedPtr<IDashboardDataTreeViewEntry> InChild)
 				{
 					if (InChild.IsValid())
 					{
 						const FSoundDashboardEntry& SoundDashboardEntry = FSoundDashboardViewFactoryPrivate::CastEntry(*InChild);
 
-						if (SoundDashboardEntry.TimeoutTimestamp == INVALID_TIMEOUT)
+						if (bIncludeTimingOutSounds || SoundDashboardEntry.TimeoutTimestamp == INVALID_TIMEOUT)
 						{
 							if (HasPinEntryType(*InChild, FSoundDashboardEntry::EPinnedEntryType::HiddenOriginalEntry) || !SoundDashboardEntry.bIsVisible)
 							{
 								return Accum;
 							}
 
-							const int32 NumNestedChildren = GetNumChildrenWithoutPinEntryType(*InChild, FSoundDashboardEntry::EPinnedEntryType::HiddenOriginalEntry, bShowRecentlyStoppedSounds);
+							const int32 NumNestedChildren = GetNumChildrenWithoutPinEntryType(*InChild, FSoundDashboardEntry::EPinnedEntryType::HiddenOriginalEntry, bIncludeTimingOutSounds);
 
 							if (NumNestedChildren > 0)
 							{
@@ -1086,7 +1086,7 @@ namespace UE::Audio::Insights
 
 				if (SoundEntry.PinnedEntryType == FSoundDashboardEntry::EPinnedEntryType::PinnedCopy)
 				{
-					if (!bIsPinnedCategoryFilterEnabled || (SoundEntry.bIsCategory && CountNumChildren(*Item, bShowRecentlyStoppedSounds) == 0))
+					if (!bIsPinnedCategoryFilterEnabled || (SoundEntry.bIsCategory && CountNumChildren(*Item, true /*bIncludeTimingOutSounds*/) == 0))
 					{
 						return EVisibility::Hidden;
 					}
@@ -1865,8 +1865,12 @@ namespace UE::Audio::Insights
 			{
 				PinSelectedItems(Entry, SelectedItems);
 			}
-
+			
+			// Make sure to clear the selection and the internal selector in SListView
+			// to ensure no shared references keep the entry alive past the point
+			// it has been removed from the dashboard
 			FilteredEntriesListView->ClearSelection();
+			FilteredEntriesListView->SetSelection(nullptr);
 		}
 	}
 
@@ -1881,7 +1885,11 @@ namespace UE::Audio::Insights
 				UnpinSelectedItems(PinnedEntry, SelectedItems);
 			}
 
+			// Make sure to clear the selection and the internal selector in SListView
+			// to ensure no shared references keep the entry alive past the point
+			// it has been removed from the dashboard
 			FilteredEntriesListView->ClearSelection();
+			FilteredEntriesListView->SetSelection(nullptr);
 		}
 	}
 
