[Audio Insights] Fixes for pinned items disappearing to early/late (adjustment from merge conflict)

diff --git a/Engine/Plugins/AudioInsights/Source/AudioInsights/Private/Views/SoundDashboardViewFactory.cpp b/Engine/Plugins/AudioInsights/Source/AudioInsights/Private/Views/SoundDashboardViewFactory.cpp
--- a/Engine/Plugins/AudioInsights/Source/AudioInsights/Private/Views/SoundDashboardViewFactory.cpp
+++ b/Engine/Plugins/AudioInsights/Source/AudioInsights/Private/Views/SoundDashboardViewFactory.cpp
@@ -127,10 +127,10 @@ namespace UE::Audio::Insights
 			return NumChildrenWithoutType;
 		}
 
-		int32 CountNumChildren(const IDashboardDataTreeViewEntry& InEntry, const bool bIncludeTimingOutSounds = false)
+		int32 CountNumChildren(const IDashboardDataTreeViewEntry& InEntry, const bool bShowRecentlyStoppedSounds, const bool bIncludeTimingOutSounds = false)
 		{
 			const uint32 TotalNumChildren = Algo::Accumulate(InEntry.Children, 0,
-				[bIncludeTimingOutSounds](uint32 Accum, TSharedPtr<IDashboardDataTreeViewEntry> InChild)
+				[bShowRecentlyStoppedSounds, bIncludeTimingOutSounds](uint32 Accum, TSharedPtr<IDashboardDataTreeViewEntry> InChild)
 				{
 					if (InChild.IsValid())
 					{
@@ -143,7 +143,7 @@ namespace UE::Audio::Insights
 								return Accum;
 							}
 
-							const int32 NumNestedChildren = GetNumChildrenWithoutPinEntryType(*InChild, FSoundDashboardEntry::EPinnedEntryType::HiddenOriginalEntry, bIncludeTimingOutSounds);
+							const int32 NumNestedChildren = GetNumChildrenWithoutPinEntryType(*InChild, FSoundDashboardEntry::EPinnedEntryType::HiddenOriginalEntry, bShowRecentlyStoppedSounds);
 
 							if (NumNestedChildren > 0)
 							{
@@ -1086,7 +1086,7 @@ namespace UE::Audio::Insights
 
 				if (SoundEntry.PinnedEntryType == FSoundDashboardEntry::EPinnedEntryType::PinnedCopy)
 				{
-					if (!bIsPinnedCategoryFilterEnabled || (SoundEntry.bIsCategory && CountNumChildren(*Item, true /*bIncludeTimingOutSounds*/) == 0))
+					if (!bIsPinnedCategoryFilterEnabled || (SoundEntry.bIsCategory && CountNumChildren(*Item, bShowRecentlyStoppedSounds, true /*bIncludeTimingOutSounds*/) == 0))
 					{
 						return EVisibility::Hidden;
 					}
