[Audio Insights] Don't expand wave instances in the Sounds Dashboard tab by default. Do expand pinned categories in the Sounds Dashboard by default.

diff --git a/Engine/Plugins/AudioInsights/Source/AudioInsights/Private/Providers/SoundTraceProvider.cpp b/Engine/Plugins/AudioInsights/Source/AudioInsights/Private/Providers/SoundTraceProvider.cpp
--- a/Engine/Plugins/AudioInsights/Source/AudioInsights/Private/Providers/SoundTraceProvider.cpp
+++ b/Engine/Plugins/AudioInsights/Source/AudioInsights/Private/Providers/SoundTraceProvider.cpp
@@ -348,6 +348,7 @@ namespace UE::Audio::Insights
 			FSoundDashboardEntry& ActiveSoundEntryRef = *OutActiveSoundEntry->Get();
 			ActiveSoundEntryRef.Name = *Msg.Name;
 			ActiveSoundEntryRef.EntryType = Msg.EntryType;
+			ActiveSoundEntryRef.bShouldForceExpandChildren = false;
 		});
 
 		ProcessMessageQueue<FSoundWaveStartMessage>(TraceMessages.SoundWaveStartMessages, GetActiveSoundEntryFromSoundWaveStartMessage,
@@ -377,6 +378,7 @@ namespace UE::Audio::Insights
 			SoundWaveEntry->Name = *Msg.Name;
 			SoundWaveEntry->EntryType = Msg.EntryType;
 			SoundWaveEntry->PinnedEntryType = (*OutActiveSoundEntry)->PinnedEntryType;
+			SoundWaveEntry->bShouldForceExpandChildren = false;
 
 			(*OutActiveSoundEntry)->Children.Add(SoundWaveEntry);
 		});
diff --git a/Engine/Plugins/AudioInsights/Source/AudioInsights/Private/Views/SoundDashboardViewFactory.cpp b/Engine/Plugins/AudioInsights/Source/AudioInsights/Private/Views/SoundDashboardViewFactory.cpp
--- a/Engine/Plugins/AudioInsights/Source/AudioInsights/Private/Views/SoundDashboardViewFactory.cpp
+++ b/Engine/Plugins/AudioInsights/Source/AudioInsights/Private/Views/SoundDashboardViewFactory.cpp
@@ -2081,6 +2081,7 @@ namespace UE::Audio::Insights
 				{
 					TSharedPtr<FSoundDashboardEntry> NewPinnedCategoryEntry = MakeShared<FSoundDashboardEntry>(CastEntry(*DataCategoryEntry));
 					NewPinnedCategoryEntry->Children.Reset();
+					NewPinnedCategoryEntry->bShouldForceExpandChildren = true;
 
 					TSharedPtr<FPinnedSoundEntryWrapper> PinnedCategory = PinnedItemEntries->AddChildEntry(NewPinnedCategoryEntry);
 					PinnedCategory->AddChildEntry(Entry);
