[Audio Insights] Don't show "Enable Audio Traces" button in standalone Unreal Insights if the package has not been set up to use the messaging bus.
- If the Trace Controller API is not found, display a message reminding the user to run the build with the -messaging command line argument.

Fixed bug where about a third of the time the "Enable Audio Traces" button wouldn't appear.
- This was due to a race between whether the Trace Controller API and the messaging bus was ready to receive the discovery request from Audio Insights. We now repeatedly send a discovery request at a low rate (every 0.5s) until we receive a response.

diff --git a/Engine/Plugins/AudioInsights/Source/AudioInsights/Private/AudioInsightsComponent.cpp b/Engine/Plugins/AudioInsights/Source/AudioInsights/Private/AudioInsightsComponent.cpp
--- a/Engine/Plugins/AudioInsights/Source/AudioInsights/Private/AudioInsightsComponent.cpp
+++ b/Engine/Plugins/AudioInsights/Source/AudioInsights/Private/AudioInsightsComponent.cpp
@@ -72,6 +72,11 @@ namespace UE::Audio::Insights
 
 		FTSTicker::GetCoreTicker().RemoveTicker(OnTickHandle);
 
+#if !WITH_EDITOR
+		IAudioInsightsTraceModule& TraceModule = FAudioInsightsModule::GetChecked().GetTraceModule();
+		TraceModule.ResetTicker();
+#endif // !WITH_EDITOR
+
 		Instance.Reset();
 	}
 
diff --git a/Engine/Plugins/AudioInsights/Source/AudioInsights/Private/AudioInsightsDashboardFactory.cpp b/Engine/Plugins/AudioInsights/Source/AudioInsights/Private/AudioInsightsDashboardFactory.cpp
--- a/Engine/Plugins/AudioInsights/Source/AudioInsights/Private/AudioInsightsDashboardFactory.cpp
+++ b/Engine/Plugins/AudioInsights/Source/AudioInsights/Private/AudioInsightsDashboardFactory.cpp
@@ -24,6 +24,8 @@ namespace UE::Audio::Insights
 		static const FText EnableTracesDescription = LOCTEXT("AudioDashboard_AutomaticallyEnableTracesDescription", "Enables the audio and audio mixer trace channels. Audio Insights will not function without these channels enabled.");
 		static const FText NoTracesEnabledWarning = LOCTEXT("AudioDashboard_NoTracesEnabledWarning", "Audio Insights requires the audio and audio mixer trace channels to be enabled to function.");
 		static const FText EnableThemNowText = LOCTEXT("AudioDashboard_EnableNowText", "Enable them now?");
+		static const FText TraceControllerUnavailabledWarning = LOCTEXT("AudioDashboard_TraceControllerUnavailableWarning", "The Trace Controller API is currently unavailable.");
+		static const FText TryEnablingMessagingText = LOCTEXT("AudioDashboard_TryEnablingMessagingText", "Make sure you have launched this package with the -Messaging command line argument.");
 	}
 
 	::Audio::FDeviceId FDashboardFactory::GetDeviceId() const
@@ -321,7 +323,11 @@ namespace UE::Audio::Insights
 				.VAlign(VAlign_Center)
 				[
 					SNew(STextBlock)
-					.Text(NoTracesEnabledWarning)
+					.Text_Lambda([]()
+					{
+						IAudioInsightsTraceModule& TraceModule = FAudioInsightsModule::GetChecked().GetTraceModule();
+						return TraceModule.TraceControllerIsAvailable() ? NoTracesEnabledWarning : TraceControllerUnavailabledWarning;
+					})
 				]
 
 				+ SVerticalBox::Slot()
@@ -330,7 +336,11 @@ namespace UE::Audio::Insights
 				.VAlign(VAlign_Center)
 				[
 					SNew(STextBlock)
-					.Text(EnableThemNowText)
+					.Text_Lambda([]()
+					{
+						IAudioInsightsTraceModule& TraceModule = FAudioInsightsModule::GetChecked().GetTraceModule();
+						return TraceModule.TraceControllerIsAvailable() ? EnableThemNowText : TryEnablingMessagingText;
+					})
 				]
 
 				+ SVerticalBox::Slot()
@@ -359,6 +369,11 @@ namespace UE::Audio::Insights
 			.OnClicked(this, &FDashboardFactory::ToggleAutoEnableAudioTraces)
 			.ToolTipText(EnableTracesDescription)
 			.ContentPadding(FMargin(0.0f, 0.0f, 0.0f, 0.0f))
+			.Visibility_Lambda([]() -> EVisibility
+			{
+				IAudioInsightsTraceModule& TraceModule = FAudioInsightsModule::GetChecked().GetTraceModule();
+				return TraceModule.TraceControllerIsAvailable() ? EVisibility::Visible : EVisibility::Hidden;
+			})
 			.Content()
 			[
 				SNew(SBox)
diff --git a/Engine/Plugins/AudioInsights/Source/AudioInsights/Private/AudioInsightsTraceModule.cpp b/Engine/Plugins/AudioInsights/Source/AudioInsights/Private/AudioInsightsTraceModule.cpp
--- a/Engine/Plugins/AudioInsights/Source/AudioInsights/Private/AudioInsightsTraceModule.cpp
+++ b/Engine/Plugins/AudioInsights/Source/AudioInsights/Private/AudioInsightsTraceModule.cpp
@@ -97,6 +97,8 @@ namespace UE::Audio::Insights
 					Commands.SetChannels(AudioChannels, EmptyArray);
 				}
 			});
+
+			return true;
 		}
 		return false;
 #endif // WITH_EDITOR
@@ -207,6 +209,40 @@ namespace UE::Audio::Insights
 
 		return false;
 	}
+
+	void FTraceModule::SendDiscoveryRequestToTraceController() const
+	{
+		// We need to send discovery requests to ITraceController so we can update the active channels later
+		ISessionServicesModule& SessionServicesModule = FModuleManager::LoadModuleChecked<ISessionServicesModule>("SessionServices");
+
+		TSharedPtr<ISessionManager> SessionManager = SessionServicesModule.GetSessionManager();
+		TSharedPtr<ITraceController> TraceController = SessionServicesModule.GetTraceController();
+		if (SessionManager.IsValid() && TraceController.IsValid())
+		{
+			const FGuid SessionId = FApp::GetSessionId();
+
+			if (TraceController.IsValid())
+			{
+				TraceController->SendDiscoveryRequest(SessionId, InstanceID);
+				TraceController->SendStatusUpdateRequest();
+			}
+		}
+	}
+
+	bool FTraceModule::Tick(float DeltaTime)
+	{
+		if (TraceControllerIsAvailable())
+		{
+			RequestChannelUpdate();
+			ResetTicker();
+		}
+		else
+		{
+			SendDiscoveryRequestToTraceController();
+		}
+
+		return true;
+	}
 #endif // !WITH_EDITOR
 
 	void FTraceModule::OnAnalysisBegin(TraceServices::IAnalysisSession& InSession)
@@ -320,22 +356,10 @@ namespace UE::Audio::Insights
 	{
 		InstanceID = SessionInfo.InstanceId;
 
-		// We need to send discovery requests to ITraceController so we can update the active channels later
-		ISessionServicesModule& SessionServicesModule = FModuleManager::LoadModuleChecked<ISessionServicesModule>("SessionServices");
-
-		TSharedPtr<ISessionManager> SessionManager = SessionServicesModule.GetSessionManager();
-		TSharedPtr<ITraceController> TraceController = SessionServicesModule.GetTraceController();
-		if (SessionManager.IsValid() && TraceController.IsValid())
-		{
-			const FGuid SessionId = FApp::GetSessionId();
-
-			if (TraceController.IsValid())
-			{
-				TraceController->SendDiscoveryRequest(SessionId, InstanceID);
-				TraceController->SendStatusUpdateRequest();
-			}
-		}
+		OnTick = FTickerDelegate::CreateRaw(this, &FTraceModule::Tick);
 
+		constexpr float TickDelay = 0.5f; // 500 ms. delay between ticks
+		OnTickHandle = FTSTicker::GetCoreTicker().AddTicker(OnTick, TickDelay);
 	}
 
 	void FTraceModule::RequestChannelUpdate()
@@ -350,6 +374,31 @@ namespace UE::Audio::Insights
 
 		TraceController->SendChannelUpdateRequest();
 	}
+
+	void FTraceModule::ResetTicker()
+	{
+		if (OnTickHandle.IsValid())
+		{
+			FTSTicker::GetCoreTicker().RemoveTicker(OnTickHandle);
+		}
+	}
+
+	bool FTraceModule::TraceControllerIsAvailable() const
+	{
+		if (!InstanceID.IsValid())
+		{
+			return false;
+		}
+
+		ISessionServicesModule& SessionServicesModule = FModuleManager::LoadModuleChecked<ISessionServicesModule>("SessionServices");
+		TSharedPtr<ITraceController> TraceController = SessionServicesModule.GetTraceController();
+		if (!TraceController.IsValid())
+		{
+			return false;
+		}
+
+		return TraceController->HasAvailableInstance(InstanceID);
+	}
 #endif // !WITH_EDITOR
 
 	void FTraceModule::GetLoggers(TArray<const TCHAR *>& OutLoggers)
diff --git a/Engine/Plugins/AudioInsights/Source/AudioInsights/Public/AudioInsightsTraceModule.h b/Engine/Plugins/AudioInsights/Source/AudioInsights/Public/AudioInsightsTraceModule.h
--- a/Engine/Plugins/AudioInsights/Source/AudioInsights/Public/AudioInsightsTraceModule.h
+++ b/Engine/Plugins/AudioInsights/Source/AudioInsights/Public/AudioInsightsTraceModule.h
@@ -2,6 +2,7 @@
 #pragma once
 
 #include "Containers/Map.h"
+#include "Containers/Ticker.h"
 #include "IAudioInsightsTraceModule.h"
 #include "ITraceController.h"
 #include "Templates/SharedPointer.h"
@@ -43,6 +44,8 @@ namespace UE::Audio::Insights
 #if !WITH_EDITOR
 		virtual void InitializeSessionInfo(const TraceServices::FSessionInfo& SessionInfo) override;
 		virtual void RequestChannelUpdate() override;
+		virtual void ResetTicker() override;
+		virtual bool TraceControllerIsAvailable() const override;
 #endif
 
 		DECLARE_MULTICAST_DELEGATE_OneParam(FOnAnalysisStarting, const double /*Timestamp*/);
@@ -62,6 +65,9 @@ namespace UE::Audio::Insights
 #if !WITH_EDITOR
 		const UE::Trace::FStoreClient::FSessionInfo* GetCurrentAnalysisSessionInfo() const;
 		bool GetAudioTracesAreEnabled() const;
+		void SendDiscoveryRequestToTraceController() const;
+
+		bool Tick(float DeltaTime);
 #endif
 
 		TMap<FName, TSharedPtr<FTraceProviderBase>> TraceProviders;
@@ -73,6 +79,9 @@ namespace UE::Audio::Insights
 #if !WITH_EDITOR
 		FGuid InstanceID;
 
+		FTickerDelegate OnTick;
+		FTSTicker::FDelegateHandle OnTickHandle;
+
 		const TArray<FString> AudioChannels;
 		const TArray<FString> EmptyArray;
 #endif
diff --git a/Engine/Plugins/AudioInsights/Source/AudioInsights/Public/IAudioInsightsTraceModule.h b/Engine/Plugins/AudioInsights/Source/AudioInsights/Public/IAudioInsightsTraceModule.h
--- a/Engine/Plugins/AudioInsights/Source/AudioInsights/Public/IAudioInsightsTraceModule.h
+++ b/Engine/Plugins/AudioInsights/Source/AudioInsights/Public/IAudioInsightsTraceModule.h
@@ -30,5 +30,7 @@ public:
 #if !WITH_EDITOR
 	virtual void InitializeSessionInfo(const TraceServices::FSessionInfo& SessionInfo) = 0;
 	virtual void RequestChannelUpdate() = 0;
+	virtual void ResetTicker() = 0;
+	virtual bool TraceControllerIsAvailable() const = 0;
 #endif // !WITH_EDITOR
 };
