[Audio Insights] Don't allow sound dashboard mute/solo columns to be sortable

diff --git a/Engine/Plugins/AudioInsights/Source/AudioInsights/Private/Views/TreeDashboardViewFactory.cpp b/Engine/Plugins/AudioInsights/Source/AudioInsights/Private/Views/TreeDashboardViewFactory.cpp
--- a/Engine/Plugins/AudioInsights/Source/AudioInsights/Private/Views/TreeDashboardViewFactory.cpp
+++ b/Engine/Plugins/AudioInsights/Source/AudioInsights/Private/Views/TreeDashboardViewFactory.cpp
@@ -264,6 +264,12 @@ namespace UE::Audio::Insights
 
 	void FTraceTreeDashboardViewFactory::OnColumnSortModeChanged(const EColumnSortPriority::Type InSortPriority, const FName& InColumnId, const EColumnSortMode::Type InSortMode)
 	{
+		// Disable sorting by the solo/mute columns
+		if (InColumnId == "Solo" || InColumnId == "Mute")
+		{
+			return;
+		}
+
 		SortByColumn = InColumnId;
 		SortMode = InSortMode;
 
