[Audio Insights] Fixing bug where using the keyboard command for unpin while selecting multiple pinned and unpinned items in the Sounds Dashboard will sometimes unpin unselected sounds.

diff --git a/Engine/Plugins/AudioInsights/Source/AudioInsights/Private/Views/SoundDashboardViewFactory.cpp b/Engine/Plugins/AudioInsights/Source/AudioInsights/Private/Views/SoundDashboardViewFactory.cpp
--- a/Engine/Plugins/AudioInsights/Source/AudioInsights/Private/Views/SoundDashboardViewFactory.cpp
+++ b/Engine/Plugins/AudioInsights/Source/AudioInsights/Private/Views/SoundDashboardViewFactory.cpp
@@ -1968,6 +1968,11 @@ namespace UE::Audio::Insights
 				}
 
 				const FSoundDashboardEntry& SelectedSoundEntry = CastEntry(*SelectedEntry);
+				if (SelectedSoundEntry.PinnedEntryType != FSoundDashboardEntry::EPinnedEntryType::PinnedCopy)
+				{
+					return false;
+				}
+
 				return OriginalChildSoundEntry.PlayOrder == SelectedSoundEntry.PlayOrder;
 			});
 
