[Audio Insights] Don't reset the expansion settings of rows in the sounds dashboard when the filter changes

diff --git a/Engine/Plugins/AudioInsights/Source/AudioInsights/Private/Views/SoundDashboardViewFactory.cpp b/Engine/Plugins/AudioInsights/Source/AudioInsights/Private/Views/SoundDashboardViewFactory.cpp
--- a/Engine/Plugins/AudioInsights/Source/AudioInsights/Private/Views/SoundDashboardViewFactory.cpp
+++ b/Engine/Plugins/AudioInsights/Source/AudioInsights/Private/Views/SoundDashboardViewFactory.cpp
@@ -102,6 +102,19 @@ namespace UE::Audio::Insights
 			return SoundEntry.bIsVisible && (bShowRecentlyStoppedSounds || SoundEntry.TimeoutTimestamp == INVALID_TIMEOUT);
 		}
 
+		void CacheInitExpandStateRecursive(FSoundDashboardEntry& SoundEntry)
+		{
+			for (TSharedPtr<IDashboardDataTreeViewEntry>& ChildEntry : SoundEntry.Children)
+			{
+				if (ChildEntry.IsValid())
+				{
+					CacheInitExpandStateRecursive(CastEntry(*ChildEntry));
+				}
+			}
+
+			SoundEntry.bShouldForceExpandChildren = SoundEntry.bIsExpanded;
+		}
+
 		bool HasPinEntryType(const IDashboardDataTreeViewEntry& InEntry, const FSoundDashboardEntry::EPinnedEntryType PinnedEntryType)
 		{
 			const FSoundDashboardEntry& SoundEntry = CastEntry(InEntry);
@@ -1108,32 +1121,43 @@ namespace UE::Audio::Insights
 		using namespace FSoundDashboardViewFactoryPrivate;
 
 		// Filter by category
-		FTraceTreeDashboardViewFactory::FilterEntries<FSoundTraceProvider>([this](const IDashboardDataTreeViewEntry& InEntry)
+		FTraceTreeDashboardViewFactory::FilterEntries<FSoundTraceProvider>([this](IDashboardDataTreeViewEntry& InEntry)
 		{
 			if (SelectedFilterFlags == AllFilterFlags)
 			{
 				return true;
 			}
 
-			const FSoundDashboardEntry& SoundCategoryEntry = FSoundDashboardViewFactoryPrivate::CastEntry(InEntry);
+			FSoundDashboardEntry& SoundCategoryEntry = FSoundDashboardViewFactoryPrivate::CastEntry(InEntry);
+			bool bEntryTypePassesFilter = false;
 
 			switch (SoundCategoryEntry.EntryType)
 			{
 				case ESoundDashboardEntryType::MetaSound:
-					return EnumHasAnyFlags(SelectedFilterFlags, ESoundDashboardFilterFlags::MetaSound);
+					bEntryTypePassesFilter = EnumHasAnyFlags(SelectedFilterFlags, ESoundDashboardFilterFlags::MetaSound);
+					break;
 				case ESoundDashboardEntryType::SoundCue:
-					return EnumHasAnyFlags(SelectedFilterFlags, ESoundDashboardFilterFlags::SoundCue);
+					bEntryTypePassesFilter = EnumHasAnyFlags(SelectedFilterFlags, ESoundDashboardFilterFlags::SoundCue);
+					break;
 				case ESoundDashboardEntryType::ProceduralSource:
-					return EnumHasAnyFlags(SelectedFilterFlags, ESoundDashboardFilterFlags::ProceduralSource);
+					bEntryTypePassesFilter = EnumHasAnyFlags(SelectedFilterFlags, ESoundDashboardFilterFlags::ProceduralSource);
+					break;
 				case ESoundDashboardEntryType::SoundWave:
-					return EnumHasAnyFlags(SelectedFilterFlags, ESoundDashboardFilterFlags::SoundWave);
+					bEntryTypePassesFilter = EnumHasAnyFlags(SelectedFilterFlags, ESoundDashboardFilterFlags::SoundWave);
+					break;
 				case ESoundDashboardEntryType::SoundCueTemplate:
-					return EnumHasAnyFlags(SelectedFilterFlags, ESoundDashboardFilterFlags::SoundCueTemplate);
+					bEntryTypePassesFilter = EnumHasAnyFlags(SelectedFilterFlags, ESoundDashboardFilterFlags::SoundCueTemplate);
+					break;
 				default:
 					break;
 			}
 
-			return false;
+			if (!bEntryTypePassesFilter)
+			{
+				CacheInitExpandStateRecursive(SoundCategoryEntry);
+			}
+
+			return bEntryTypePassesFilter;
 		});
 
 		// Filter by text
diff --git a/Engine/Plugins/AudioInsights/Source/AudioInsights/Private/Views/TreeDashboardViewFactory.cpp b/Engine/Plugins/AudioInsights/Source/AudioInsights/Private/Views/TreeDashboardViewFactory.cpp
--- a/Engine/Plugins/AudioInsights/Source/AudioInsights/Private/Views/TreeDashboardViewFactory.cpp
+++ b/Engine/Plugins/AudioInsights/Source/AudioInsights/Private/Views/TreeDashboardViewFactory.cpp
@@ -141,6 +141,10 @@ namespace UE::Audio::Insights
 						}
 					})
 					.OnSetExpansionRecursive(this, &FTraceTreeDashboardViewFactory::HandleRecursiveExpansion)
+					.OnExpansionChanged_Lambda([](TSharedPtr<IDashboardDataTreeViewEntry> Item, bool bIsExpanded)
+					{
+						Item->bIsExpanded = bIsExpanded;
+					})
 				]
 			];
 		}
diff --git a/Engine/Plugins/AudioInsights/Source/AudioInsights/Public/AudioInsightsDataSource.h b/Engine/Plugins/AudioInsights/Source/AudioInsights/Public/AudioInsightsDataSource.h
--- a/Engine/Plugins/AudioInsights/Source/AudioInsights/Public/AudioInsightsDataSource.h
+++ b/Engine/Plugins/AudioInsights/Source/AudioInsights/Public/AudioInsightsDataSource.h
@@ -24,5 +24,6 @@ namespace UE::Audio::Insights
 		virtual void ResetShouldInitExpandChildren() = 0;
 
 		TArray<TSharedPtr<IDashboardDataTreeViewEntry>> Children;
+		bool bIsExpanded = false;
 	};
 } // namespace UE::Audio::Insights
diff --git a/Engine/Plugins/AudioInsights/Source/AudioInsights/Public/Views/TreeDashboardViewFactory.h b/Engine/Plugins/AudioInsights/Source/AudioInsights/Public/Views/TreeDashboardViewFactory.h
--- a/Engine/Plugins/AudioInsights/Source/AudioInsights/Public/Views/TreeDashboardViewFactory.h
+++ b/Engine/Plugins/AudioInsights/Source/AudioInsights/Public/Views/TreeDashboardViewFactory.h
@@ -121,7 +121,7 @@ namespace UE::Audio::Insights
 #endif // WITH_EDITOR
 
 		template<typename TableProviderType>
-		bool FilterEntries(TFunctionRef<bool(const IDashboardDataTreeViewEntry&)> InPredicate)
+		bool FilterEntries(TFunctionRef<bool(IDashboardDataTreeViewEntry&)> InPredicate)
 		{
 			const TSharedPtr<const TableProviderType> Provider = FindProvider<const TableProviderType>();
 			if (Provider.IsValid())
