[Audio Insights] Fix sounds flickering in and out of the sounds dashboard in scenes with high voice counts when the text filter is active

diff --git a/Engine/Plugins/AudioInsights/Source/AudioInsights/Private/Views/SoundDashboardViewFactory.cpp b/Engine/Plugins/AudioInsights/Source/AudioInsights/Private/Views/SoundDashboardViewFactory.cpp
--- a/Engine/Plugins/AudioInsights/Source/AudioInsights/Private/Views/SoundDashboardViewFactory.cpp
+++ b/Engine/Plugins/AudioInsights/Source/AudioInsights/Private/Views/SoundDashboardViewFactory.cpp
@@ -417,6 +417,7 @@ namespace UE::Audio::Insights
 
 		FSoundDashboardEntry& PinnedSectionSoundEntry = FSoundDashboardViewFactoryPrivate::CastEntry(*PinnedSectionEntry);
 		PinnedSectionSoundEntry.PinnedEntryType = FSoundDashboardEntry::EPinnedEntryType::PinnedCopy;
+		PinnedSectionSoundEntry.bIsVisible = true;
 
 		for (const TSharedPtr<IDashboardDataTreeViewEntry>& Child : OriginalEntry->Children)
 		{
@@ -1171,16 +1172,13 @@ namespace UE::Audio::Insights
 				continue;
 			}
 
+			ResetVisibility(*CategoryEntry);
+
 			for (const TSharedPtr<IDashboardDataTreeViewEntry>& Entry : CategoryEntry->Children)
 			{
-				if (Entry.IsValid())
+				if (Entry.IsValid() && bFilterHasText)
 				{
-					ResetVisibility(*Entry);
-
-					if (bFilterHasText)
-					{
-						SetFilteredVisibility(*Entry, FilterString);
-					}
+					SetFilteredVisibility(*Entry, FilterString);
 				}
 			}
 		}
diff --git a/Engine/Plugins/AudioInsights/Source/AudioInsights/Public/Messages/SoundTraceMessages.h b/Engine/Plugins/AudioInsights/Source/AudioInsights/Public/Messages/SoundTraceMessages.h
--- a/Engine/Plugins/AudioInsights/Source/AudioInsights/Public/Messages/SoundTraceMessages.h
+++ b/Engine/Plugins/AudioInsights/Source/AudioInsights/Public/Messages/SoundTraceMessages.h
@@ -171,7 +171,7 @@ namespace UE::Audio::Insights
 		EPinnedEntryType PinnedEntryType = EPinnedEntryType::None;
 		bool bIsCategory = false;
 		bool bShouldForceExpandChildren = true;
-		bool bIsVisible = true;
+		bool bIsVisible = false;
 
 		::Audio::TCircularAudioBuffer<FDataPoint> PriorityDataPoints;
 		::Audio::TCircularAudioBuffer<FDataPoint> DistanceDataPoints;
