[Audio Insights] Help reduce amount of re-ordering in the sounds dashboard when sorting by a column - if two values are almost identical, fall back on the Play Order instead

diff --git a/Engine/Plugins/AudioInsights/Source/AudioInsights/Private/Views/SoundDashboardViewFactory.cpp b/Engine/Plugins/AudioInsights/Source/AudioInsights/Private/Views/SoundDashboardViewFactory.cpp
--- a/Engine/Plugins/AudioInsights/Source/AudioInsights/Private/Views/SoundDashboardViewFactory.cpp
+++ b/Engine/Plugins/AudioInsights/Source/AudioInsights/Private/Views/SoundDashboardViewFactory.cpp
@@ -1545,34 +1545,75 @@ namespace UE::Audio::Insights
 
 		UpdatePinnedSection();
 
-		auto SortByName = [](const FSoundDashboardEntry& First, const FSoundDashboardEntry& Second)
+		auto SortByPlayOrder = [](const FSoundDashboardEntry& First, const FSoundDashboardEntry& Second)
 		{
-			return First.GetDisplayName().CompareToCaseIgnored(Second.GetDisplayName()) <= 0;
+			return First.PlayOrder < Second.PlayOrder;
+		};
+		
+		auto SortByName = [&SortByPlayOrder](const FSoundDashboardEntry& First, const FSoundDashboardEntry& Second)
+		{
+			const int32 Comparison = First.GetDisplayName().CompareToCaseIgnored(Second.GetDisplayName());
+			if (Comparison == 0)
+			{
+				return SortByPlayOrder(First, Second);
+			}
+			
+			return Comparison < 0;
 		};
 
-		auto SortByPriority = [](const FSoundDashboardEntry& First, const FSoundDashboardEntry& Second)
+		auto SortByPriority = [&SortByPlayOrder](const FSoundDashboardEntry& First, const FSoundDashboardEntry& Second)
 		{
-			return GetLastEntryArrayValue(First.PriorityDataPoints) <= GetLastEntryArrayValue(Second.PriorityDataPoints);
+			const float ComparisonDiff = GetLastEntryArrayValue(First.PriorityDataPoints) - GetLastEntryArrayValue(Second.PriorityDataPoints);
+			if (FMath::IsNearlyZero(ComparisonDiff, UE_KINDA_SMALL_NUMBER))
+			{
+				return SortByPlayOrder(First, Second);
+			}
+
+			return ComparisonDiff < 0.0f;
 		};
 
-		auto SortByDistance = [](const FSoundDashboardEntry& First, const FSoundDashboardEntry& Second)
+		auto SortByDistance = [&SortByPlayOrder](const FSoundDashboardEntry& First, const FSoundDashboardEntry& Second)
 		{
-			return GetLastEntryArrayValue(First.DistanceDataPoints) <= GetLastEntryArrayValue(Second.DistanceDataPoints);
+			const float ComparisonDiff = GetLastEntryArrayValue(First.DistanceDataPoints) - GetLastEntryArrayValue(Second.DistanceDataPoints);
+			if (FMath::IsNearlyZero(ComparisonDiff, UE_KINDA_SMALL_NUMBER))
+			{
+				return SortByPlayOrder(First, Second);
+			}
+
+			return ComparisonDiff < 0.0f;
 		};
 
-		auto SortByAmplitude = [](const FSoundDashboardEntry& First, const FSoundDashboardEntry& Second)
+		auto SortByAmplitude = [&SortByPlayOrder](const FSoundDashboardEntry& First, const FSoundDashboardEntry& Second)
 		{
-			return GetLastEntryArrayValue(First.AmplitudeDataPoints) <= GetLastEntryArrayValue(Second.AmplitudeDataPoints);
+			const float ComparisonDiff = GetLastEntryArrayValue(First.AmplitudeDataPoints) - GetLastEntryArrayValue(Second.AmplitudeDataPoints);
+			if (FMath::IsNearlyZero(ComparisonDiff, UE_KINDA_SMALL_NUMBER))
+			{
+				return SortByPlayOrder(First, Second);
+			}
+
+			return ComparisonDiff < 0.0f;
 		};
 
-		auto SortByVolume = [](const FSoundDashboardEntry& First, const FSoundDashboardEntry& Second)
+		auto SortByVolume = [&SortByPlayOrder](const FSoundDashboardEntry& First, const FSoundDashboardEntry& Second)
 		{
-			return GetLastEntryArrayValue(First.VolumeDataPoints) <= GetLastEntryArrayValue(Second.VolumeDataPoints);
+			const float ComparisonDiff = GetLastEntryArrayValue(First.VolumeDataPoints) - GetLastEntryArrayValue(Second.VolumeDataPoints);
+			if (FMath::IsNearlyZero(ComparisonDiff, UE_KINDA_SMALL_NUMBER))
+			{
+				return SortByPlayOrder(First, Second);
+			}
+
+			return ComparisonDiff < 0.0f;
 		};
 
-		auto SortByPitch = [](const FSoundDashboardEntry& First, const FSoundDashboardEntry& Second)
+		auto SortByPitch = [&SortByPlayOrder](const FSoundDashboardEntry& First, const FSoundDashboardEntry& Second)
 		{
-			return GetLastEntryArrayValue(First.PitchDataPoints) <= GetLastEntryArrayValue(Second.PitchDataPoints);
+			const float ComparisonDiff = GetLastEntryArrayValue(First.PitchDataPoints) - GetLastEntryArrayValue(Second.PitchDataPoints);
+			if (FMath::IsNearlyZero(ComparisonDiff, UE_KINDA_SMALL_NUMBER))
+			{
+				return SortByPlayOrder(First, Second);
+			}
+
+			return ComparisonDiff < 0.0f;
 		};
 
 		if (SortByColumn == "Name")
