[Audio Insights] Replace LoadModuleChecked with LoadModulePtr inside Enable Traces overlay inside Audio Insights for protection

diff --git a/Engine/Plugins/AudioInsights/Source/AudioInsights/Private/AudioInsightsDashboardFactory.cpp b/Engine/Plugins/AudioInsights/Source/AudioInsights/Private/AudioInsightsDashboardFactory.cpp
--- a/Engine/Plugins/AudioInsights/Source/AudioInsights/Private/AudioInsightsDashboardFactory.cpp
+++ b/Engine/Plugins/AudioInsights/Source/AudioInsights/Private/AudioInsightsDashboardFactory.cpp
@@ -305,7 +305,13 @@ namespace UE::Audio::Insights
 			.BorderBackgroundColor(FLinearColor(0.2f, 0.2f, 0.2f, 0.8f))
 			.Visibility_Lambda([]() -> EVisibility
 			{
-				IAudioInsightsTraceModule& TraceModule = FAudioInsightsModule::GetChecked().GetTraceModule();
+				FAudioInsightsModule* AudioInsightsModulePtr = FAudioInsightsModule::GetModulePtr();
+				if (AudioInsightsModulePtr == nullptr)
+				{
+					return EVisibility::Hidden;
+				}
+
+				IAudioInsightsTraceModule& TraceModule = AudioInsightsModulePtr->GetTraceModule();
 				return TraceModule.AudioChannelsCanBeManuallyEnabled() ? EVisibility::Visible : EVisibility::Hidden;
 			})
 			[
@@ -325,7 +331,13 @@ namespace UE::Audio::Insights
 					SNew(STextBlock)
 					.Text_Lambda([]()
 					{
-						IAudioInsightsTraceModule& TraceModule = FAudioInsightsModule::GetChecked().GetTraceModule();
+						FAudioInsightsModule* AudioInsightsModulePtr = FAudioInsightsModule::GetModulePtr();
+						if (AudioInsightsModulePtr == nullptr)
+						{
+							return FText::GetEmpty();
+						}
+
+						IAudioInsightsTraceModule& TraceModule = AudioInsightsModulePtr->GetTraceModule();
 						return TraceModule.TraceControllerIsAvailable() ? NoTracesEnabledWarning : TraceControllerUnavailabledWarning;
 					})
 				]
@@ -338,7 +350,13 @@ namespace UE::Audio::Insights
 					SNew(STextBlock)
 					.Text_Lambda([]()
 					{
-						IAudioInsightsTraceModule& TraceModule = FAudioInsightsModule::GetChecked().GetTraceModule();
+						FAudioInsightsModule* AudioInsightsModulePtr = FAudioInsightsModule::GetModulePtr();
+						if (AudioInsightsModulePtr == nullptr)
+						{
+							return FText::GetEmpty();
+						}
+
+						IAudioInsightsTraceModule& TraceModule = AudioInsightsModulePtr->GetTraceModule();
 						return TraceModule.TraceControllerIsAvailable() ? EnableThemNowText : TryEnablingMessagingText;
 					})
 				]
@@ -371,7 +389,14 @@ namespace UE::Audio::Insights
 			.ContentPadding(FMargin(0.0f, 0.0f, 0.0f, 0.0f))
 			.Visibility_Lambda([]() -> EVisibility
 			{
-				IAudioInsightsTraceModule& TraceModule = FAudioInsightsModule::GetChecked().GetTraceModule();
+				FAudioInsightsModule* AudioInsightsModulePtr = FAudioInsightsModule::GetModulePtr();
+				if (AudioInsightsModulePtr == nullptr)
+				{
+					return EVisibility::Hidden;
+				}
+
+				IAudioInsightsTraceModule& TraceModule = AudioInsightsModulePtr->GetTraceModule();
+
 				return TraceModule.TraceControllerIsAvailable() ? EVisibility::Visible : EVisibility::Hidden;
 			})
 			.Content()
diff --git a/Engine/Plugins/AudioInsights/Source/AudioInsights/Private/AudioInsightsModule.cpp b/Engine/Plugins/AudioInsights/Source/AudioInsights/Private/AudioInsightsModule.cpp
--- a/Engine/Plugins/AudioInsights/Source/AudioInsights/Private/AudioInsightsModule.cpp
+++ b/Engine/Plugins/AudioInsights/Source/AudioInsights/Private/AudioInsightsModule.cpp
@@ -105,6 +105,16 @@ namespace UE::Audio::Insights
 		return static_cast<FAudioInsightsModule&>(FModuleManager::LoadModuleChecked<IAudioInsightsModule>("AudioInsights"));
 	}
 
+	FAudioInsightsModule* FAudioInsightsModule::GetModulePtr()
+	{
+		if (IAudioInsightsModule* AudioInsightsModulePtr = FModuleManager::LoadModulePtr<IAudioInsightsModule>("AudioInsights"))
+		{
+			return static_cast<FAudioInsightsModule*>(AudioInsightsModulePtr);
+		}
+
+		return nullptr;
+	}
+
 	TSharedRef<FDashboardFactory> FAudioInsightsModule::GetDashboardFactory()
 	{
 		return DashboardFactory->AsShared();
diff --git a/Engine/Plugins/AudioInsights/Source/AudioInsights/Private/AudioInsightsModule.h b/Engine/Plugins/AudioInsights/Source/AudioInsights/Private/AudioInsightsModule.h
--- a/Engine/Plugins/AudioInsights/Source/AudioInsights/Private/AudioInsightsModule.h
+++ b/Engine/Plugins/AudioInsights/Source/AudioInsights/Private/AudioInsightsModule.h
@@ -35,6 +35,7 @@ namespace UE::Audio::Insights
 		virtual ::Audio::FDeviceId GetDeviceId() const override;
 
 		static FAudioInsightsModule& GetChecked();
+		static FAudioInsightsModule* GetModulePtr();
 		virtual IAudioInsightsTraceModule& GetTraceModule() override;
 
 #if !WITH_EDITOR
