[Audio Insights] Enable pinning/unpinning via r-click on categories in the sounds dashboard
- If no individual assets are selected, pin/unpin applies to all assets currently in the category
- If the "Pinned" category is the only one selected, unpin everything
- If multiple assets and categories are selected, just pin/unpin the ones selected
- Disable and grey out "Browse to Asset" and "Edit Asset" r-click options when no assets are selected

diff --git a/Engine/Plugins/AudioInsights/Source/AudioInsights/Private/Views/SoundDashboardViewFactory.cpp b/Engine/Plugins/AudioInsights/Source/AudioInsights/Private/Views/SoundDashboardViewFactory.cpp
--- a/Engine/Plugins/AudioInsights/Source/AudioInsights/Private/Views/SoundDashboardViewFactory.cpp
+++ b/Engine/Plugins/AudioInsights/Source/AudioInsights/Private/Views/SoundDashboardViewFactory.cpp
@@ -581,8 +581,8 @@ namespace UE::Audio::Insights
 		CommandList->MapAction(Commands.GetPinCommand(), FExecuteAction::CreateRaw(this, &FSoundDashboardViewFactory::PinSound));
 		CommandList->MapAction(Commands.GetUnpinCommand(), FExecuteAction::CreateRaw(this, &FSoundDashboardViewFactory::UnpinSound));
 #if WITH_EDITOR
-		CommandList->MapAction(Commands.GetBrowseCommand(), FExecuteAction::CreateRaw(this, &FSoundDashboardViewFactory::BrowseSoundAsset));
-		CommandList->MapAction(Commands.GetEditCommand(),   FExecuteAction::CreateRaw(this, &FSoundDashboardViewFactory::OpenSoundAsset));
+		CommandList->MapAction(Commands.GetBrowseCommand(), FExecuteAction::CreateRaw(this, &FSoundDashboardViewFactory::BrowseSoundAsset), FCanExecuteAction::CreateRaw(this, &FSoundDashboardViewFactory::SelectedItemsIncludesAnAsset));
+		CommandList->MapAction(Commands.GetEditCommand(),   FExecuteAction::CreateRaw(this, &FSoundDashboardViewFactory::OpenSoundAsset), FCanExecuteAction::CreateRaw(this, &FSoundDashboardViewFactory::SelectedItemsIncludesAnAsset));
 #endif // WITH_EDITOR
 
 		// @TODO UE-250399: Hide category pending to implement
@@ -1188,15 +1188,6 @@ namespace UE::Audio::Insights
 	{
 		using namespace FSoundDashboardViewFactoryPrivate;
 
-		// Skip right click for category items
-		if (FilteredEntriesListView.IsValid() && FilteredEntriesListView->GetSelectedItems().Num() == 1)
-		{
-			const TSharedPtr<IDashboardDataTreeViewEntry> SelectedEntry = FilteredEntriesListView->GetSelectedItems()[0];
-			if (IsCategoryItem(*SelectedEntry))
-			{
-				return SNullWidget::NullWidget;
-			}
-		}
 
 		const FSoundDashboardCommands& Commands = FSoundDashboardCommands::Get();
 
@@ -1918,15 +1909,47 @@ namespace UE::Audio::Insights
 	}
 #endif // WITH_EDITOR
 
+	bool FSoundDashboardViewFactory::SelectedItemsIncludesAnAsset() const
+	{
+		using namespace FSoundDashboardViewFactoryPrivate;
+		if (!FilteredEntriesListView.IsValid())
+		{
+			return false;
+		}
+
+		for (const TSharedPtr<IDashboardDataTreeViewEntry>& SelectedEntry : FilteredEntriesListView->GetSelectedItems())
+		{
+			if (!IsCategoryItem(*SelectedEntry))
+			{
+				return true;
+			}
+		}
+
+		return false;
+	}
+
 	void FSoundDashboardViewFactory::PinSound()
 	{
 		if (FilteredEntriesListView.IsValid())
 		{
 			const TArray<TSharedPtr<IDashboardDataTreeViewEntry>> SelectedItems = FilteredEntriesListView->GetSelectedItems();
+			const bool bSelectionContainsAssets = SelectedItemsIncludesAnAsset();
 
 			for (TSharedPtr<IDashboardDataTreeViewEntry>& Entry : DataViewEntries)
 			{
-				PinSelectedItems(Entry, SelectedItems);
+				// If only categories are selected, pin the entire category
+				if (SelectedItems.Contains(Entry) && !bSelectionContainsAssets)
+				{
+					for (TSharedPtr<IDashboardDataTreeViewEntry>& ChildEntry : Entry->Children)
+					{
+						MarkBranchAsPinned(ChildEntry, true /*bIsPinned*/);
+						CreatePinnedEntry(ChildEntry);
+					}
+				}
+				else
+				{
+					PinSelectedItems(Entry, SelectedItems);
+				}
 			}
 			
 			// Make sure to clear the selection and the internal selector in SListView
@@ -1942,12 +1965,23 @@ namespace UE::Audio::Insights
 		if (FilteredEntriesListView.IsValid() && PinnedItemEntries.IsValid())
 		{
 			const TArray<TSharedPtr<IDashboardDataTreeViewEntry>> SelectedItems = FilteredEntriesListView->GetSelectedItems();
+			const bool bSelectionContainsAssets = SelectedItemsIncludesAnAsset();
 
-			for (const TSharedPtr<FPinnedSoundEntryWrapper>& PinnedEntry : PinnedItemEntries->PinnedWrapperChildren)
+			// If the user has only selected the pinned item row, unpin everything
+			if (SelectedItems.Num() == 1 && SelectedItems[0] == PinnedItemEntries->GetPinnedSectionEntry())
+			{
+				for (TSharedPtr<IDashboardDataTreeViewEntry>& OriginalChildEntry : DataViewEntries)
+				{
+					MarkBranchAsPinned(OriginalChildEntry, false /*bIsPinned*/);
+				}
+				
+				PinnedItemEntries.Reset();
+			}
+			else
 			{
-				UnpinSelectedItems(PinnedEntry, SelectedItems);
+				UnpinSelectedItems(PinnedItemEntries, SelectedItems, bSelectionContainsAssets);
 			}
-
+			
 			// Make sure to clear the selection and the internal selector in SListView
 			// to ensure no shared references keep the entry alive past the point
 			// it has been removed from the dashboard
@@ -2007,7 +2041,7 @@ namespace UE::Audio::Insights
 		}
 	}
 
-	void FSoundDashboardViewFactory::UnpinSelectedItems(const TSharedPtr<FPinnedSoundEntryWrapper>& PinnedWrapperEntry, const TArray<TSharedPtr<IDashboardDataTreeViewEntry>>& SelectedItems)
+	void FSoundDashboardViewFactory::UnpinSelectedItems(const TSharedPtr<FPinnedSoundEntryWrapper>& PinnedWrapperEntry, const TArray<TSharedPtr<IDashboardDataTreeViewEntry>>& SelectedItems, const bool bSelectionContainsAssets)
 	{
 		using namespace FSoundDashboardViewFactoryPrivate;
 
@@ -2018,12 +2052,13 @@ namespace UE::Audio::Insights
 			const TSharedPtr<IDashboardDataTreeViewEntry> OriginalChildEntry = PinnedWrapperChild->GetOriginalDataEntry();
 			if (!OriginalChildEntry.IsValid())
 			{
+				UnpinSelectedItems(PinnedWrapperChild, SelectedItems, bSelectionContainsAssets);
 				continue;
 			}
 
 			const FSoundDashboardEntry& OriginalChildSoundEntry = CastEntry(*OriginalChildEntry);
 
-			const TSharedPtr<IDashboardDataTreeViewEntry>* FoundMatchingSelectedEntry = SelectedItems.FindByPredicate([&OriginalChildSoundEntry](const TSharedPtr<IDashboardDataTreeViewEntry> SelectedEntry)
+			const TSharedPtr<IDashboardDataTreeViewEntry>* FoundMatchingSelectedEntry = SelectedItems.FindByPredicate([&OriginalChildSoundEntry, bSelectionContainsAssets](const TSharedPtr<IDashboardDataTreeViewEntry> SelectedEntry)
 			{
 				if (!SelectedEntry.IsValid())
 				{
@@ -2036,6 +2071,11 @@ namespace UE::Audio::Insights
 					return false;
 				}
 
+				if (SelectedSoundEntry.bIsCategory && !bSelectionContainsAssets)
+				{
+					return OriginalChildSoundEntry.EntryType == SelectedSoundEntry.EntryType;
+				}
+				
 				return OriginalChildSoundEntry.PlayOrder == SelectedSoundEntry.PlayOrder;
 			});
 
@@ -2063,7 +2103,7 @@ namespace UE::Audio::Insights
 			else
 			{
 				// If this child item is not selected, check it's children.
-				UnpinSelectedItems(PinnedWrapperChild, SelectedItems);
+				UnpinSelectedItems(PinnedWrapperChild, SelectedItems, bSelectionContainsAssets);
 			}
 		}
 	}
diff --git a/Engine/Plugins/AudioInsights/Source/AudioInsights/Public/Views/SoundDashboardViewFactory.h b/Engine/Plugins/AudioInsights/Source/AudioInsights/Public/Views/SoundDashboardViewFactory.h
--- a/Engine/Plugins/AudioInsights/Source/AudioInsights/Public/Views/SoundDashboardViewFactory.h
+++ b/Engine/Plugins/AudioInsights/Source/AudioInsights/Public/Views/SoundDashboardViewFactory.h
@@ -149,12 +149,13 @@ namespace UE::Audio::Insights
 		TArray<TObjectPtr<UObject>> GetSelectedEditableAssets() const;
 #endif // WITH_EDITOR
 
+		bool SelectedItemsIncludesAnAsset() const;
 		void PinSound();
 		void UnpinSound();
 		bool SelectionIncludesUnpinnedItem() const;
 
 		void PinSelectedItems(const TSharedPtr<IDashboardDataTreeViewEntry>& Entry, const TArray<TSharedPtr<IDashboardDataTreeViewEntry>>& SelectedItems);
-		void UnpinSelectedItems(const TSharedPtr<FPinnedSoundEntryWrapper>& PinnedWrapperEntry, const TArray<TSharedPtr<IDashboardDataTreeViewEntry>>& SelectedItems);
+		void UnpinSelectedItems(const TSharedPtr<FPinnedSoundEntryWrapper>& PinnedWrapperEntry, const TArray<TSharedPtr<IDashboardDataTreeViewEntry>>& SelectedItems, const bool bSelectionContainsAssets);
 
 		void MarkBranchAsPinned(const TSharedPtr<IDashboardDataTreeViewEntry> Entry, const bool bIsPinned);
 		void InitPinnedItemEntries();
