[Audio Insights] Fix for Rewind Debugger causing Audio Insights to not receive any data

diff --git a/Engine/Plugins/AudioInsights/Source/AudioInsights/AudioInsights.Build.cs b/Engine/Plugins/AudioInsights/Source/AudioInsights/AudioInsights.Build.cs
--- a/Engine/Plugins/AudioInsights/Source/AudioInsights/AudioInsights.Build.cs
+++ b/Engine/Plugins/AudioInsights/Source/AudioInsights/AudioInsights.Build.cs
@@ -13,6 +13,7 @@ public class AudioInsights : ModuleRules
 			new string[]
 			{
 				"Core",
+				"RewindDebuggerRuntimeInterface",
 				"TraceInsights",
 				"TraceServices",
 				"ToolWidgets",
diff --git a/Engine/Plugins/AudioInsights/Source/AudioInsights/Private/AudioInsightsModule.cpp b/Engine/Plugins/AudioInsights/Source/AudioInsights/Private/AudioInsightsModule.cpp
--- a/Engine/Plugins/AudioInsights/Source/AudioInsights/Private/AudioInsightsModule.cpp
+++ b/Engine/Plugins/AudioInsights/Source/AudioInsights/Private/AudioInsightsModule.cpp
@@ -31,7 +31,10 @@ namespace UE::Audio::Insights
 		if (!IsRunningCommandlet())
 		{
 			TraceModule = MakeUnique<FTraceModule>();
+			RewindDebuggerExtension = MakeUnique<FRewindDebugger>();
+
 			IModularFeatures::Get().RegisterModularFeature(TraceServices::ModuleFeatureName, TraceModule.Get());
+			IModularFeatures::Get().RegisterModularFeature(IRewindDebuggerRuntimeExtension::ModularFeatureName, RewindDebuggerExtension.Get());
 
 			DashboardFactory = MakeShared<FDashboardFactory>();
 			
@@ -82,6 +85,7 @@ namespace UE::Audio::Insights
 			DashboardFactory.Reset();
 
 			IModularFeatures::Get().UnregisterModularFeature(TraceServices::ModuleFeatureName, TraceModule.Get());
+			IModularFeatures::Get().UnregisterModularFeature(IRewindDebuggerRuntimeExtension::ModularFeatureName, RewindDebuggerExtension.Get());
 		}
 	}
 
diff --git a/Engine/Plugins/AudioInsights/Source/AudioInsights/Private/AudioInsightsModule.h b/Engine/Plugins/AudioInsights/Source/AudioInsights/Private/AudioInsightsModule.h
--- a/Engine/Plugins/AudioInsights/Source/AudioInsights/Private/AudioInsightsModule.h
+++ b/Engine/Plugins/AudioInsights/Source/AudioInsights/Private/AudioInsightsModule.h
@@ -51,6 +51,7 @@ namespace UE::Audio::Insights
 	private:
 		TSharedPtr<FDashboardFactory> DashboardFactory;
 		TUniquePtr<FTraceModule> TraceModule;
+		TUniquePtr<FRewindDebugger> RewindDebuggerExtension;
 
 #if !WITH_EDITOR
 		TSharedPtr<FAudioInsightsComponent> AudioInsightsComponent;
diff --git a/Engine/Plugins/AudioInsights/Source/AudioInsights/Private/AudioInsightsTraceModule.cpp b/Engine/Plugins/AudioInsights/Source/AudioInsights/Private/AudioInsightsTraceModule.cpp
--- a/Engine/Plugins/AudioInsights/Source/AudioInsights/Private/AudioInsightsTraceModule.cpp
+++ b/Engine/Plugins/AudioInsights/Source/AudioInsights/Private/AudioInsightsTraceModule.cpp
@@ -24,6 +24,14 @@ namespace UE::Audio::Insights
 		static const TCHAR* AudioMixerChannelName = TEXT("AudioMixer");
 	}
 
+	void FRewindDebugger::RecordingStarted()
+	{
+		using namespace FTraceModulePrivate;
+
+		UE::Trace::ToggleChannel(AudioChannelName, true);
+		UE::Trace::ToggleChannel(AudioMixerChannelName, true);
+	}
+
 	FTraceModule::FTraceModule()
 #if !WITH_EDITOR
 		: AudioChannels({ FTraceModulePrivate::AudioChannelName, FTraceModulePrivate::AudioMixerChannelName } )
diff --git a/Engine/Plugins/AudioInsights/Source/AudioInsights/Public/AudioInsightsTraceModule.h b/Engine/Plugins/AudioInsights/Source/AudioInsights/Public/AudioInsightsTraceModule.h
--- a/Engine/Plugins/AudioInsights/Source/AudioInsights/Public/AudioInsightsTraceModule.h
+++ b/Engine/Plugins/AudioInsights/Source/AudioInsights/Public/AudioInsightsTraceModule.h
@@ -5,6 +5,7 @@
 #include "Containers/Ticker.h"
 #include "IAudioInsightsTraceModule.h"
 #include "ITraceController.h"
+#include "RewindDebuggerRuntimeInterface/IRewindDebuggerRuntimeExtension.h"
 #include "Templates/SharedPointer.h"
 #include "Trace/StoreClient.h"
 #include "TraceServices/Model/AnalysisSession.h"
@@ -12,6 +13,12 @@
 
 namespace UE::Audio::Insights
 {
+	class FRewindDebugger : public IRewindDebuggerRuntimeExtension
+	{
+	public:
+		virtual void RecordingStarted() override;
+	};
+
 	class AUDIOINSIGHTS_API FTraceModule : public IAudioInsightsTraceModule
 	{
 	public:
@@ -72,6 +79,7 @@ namespace UE::Audio::Insights
 
 		TMap<FName, TSharedPtr<FTraceProviderBase>> TraceProviders;
 		TArray<FString> ChannelsToRestore;
+		FRewindDebugger RewindDebugger;
 
 		bool bTraceAnalysisHasStarted = false;
 		bool bStopTracingAfterAudioInsightsIsFinished = false;
