[Audio Insights] Fix ghost data becoming stuck on the sounds dashboard. Check to make sure we haven't processed a stop message before we received the start message (happens with short sounds in busy scenes)

diff --git a/Engine/Plugins/AudioInsights/Source/AudioInsights/Private/Providers/SoundTraceProvider.cpp b/Engine/Plugins/AudioInsights/Source/AudioInsights/Private/Providers/SoundTraceProvider.cpp
--- a/Engine/Plugins/AudioInsights/Source/AudioInsights/Private/Providers/SoundTraceProvider.cpp
+++ b/Engine/Plugins/AudioInsights/Source/AudioInsights/Private/Providers/SoundTraceProvider.cpp
@@ -279,6 +279,7 @@ namespace UE::Audio::Insights
 
 		ActiveSoundToEntryKeysMap.Empty();
 		EntriesTimingOut.Empty();
+		SoundsStoppedBeforeStart.Empty();
 
 		return new FSoundTraceAnalyzer(AsShared(), InSession);
 	}
@@ -494,19 +495,28 @@ namespace UE::Audio::Insights
 			}
 #endif // !WITH_EDITOR
 
-			if (OutActiveSoundEntry && OutActiveSoundEntry->IsValid() && (*OutActiveSoundEntry)->Timestamp < Msg.Timestamp)
+			if (OutActiveSoundEntry && OutActiveSoundEntry->IsValid())
 			{
-				static constexpr double DashboardTimeoutTime = 3.0;
-				const double TimeoutTimestamp = FPlatformTime::Seconds() + DashboardTimeoutTime;
-
-				(*OutActiveSoundEntry)->TimeoutTimestamp = TimeoutTimestamp;
-				for (TSharedPtr<IDashboardDataTreeViewEntry> ChildEntry : (*OutActiveSoundEntry)->Children)
+				if ((*OutActiveSoundEntry)->Timestamp < Msg.Timestamp)
 				{
-					FSoundDashboardEntry& SoundWaveEntry = FSoundTraceProviderPrivate::CastEntry(*ChildEntry.Get());
-					SoundWaveEntry.TimeoutTimestamp = TimeoutTimestamp;
-				}
+					static constexpr double DashboardTimeoutTime = 3.0;
+					const double TimeoutTimestamp = FPlatformTime::Seconds() + DashboardTimeoutTime;
 
-				EntriesTimingOut.Push(SoundMessageIDs{ Msg.DeviceId, Msg.PlayOrder });
+					(*OutActiveSoundEntry)->TimeoutTimestamp = TimeoutTimestamp;
+					for (TSharedPtr<IDashboardDataTreeViewEntry> ChildEntry : (*OutActiveSoundEntry)->Children)
+					{
+						FSoundDashboardEntry& SoundWaveEntry = FSoundTraceProviderPrivate::CastEntry(*ChildEntry.Get());
+						SoundWaveEntry.TimeoutTimestamp = TimeoutTimestamp;
+					}
+
+					EntriesTimingOut.Push(SoundMessageIDs{ Msg.DeviceId, Msg.PlayOrder });
+				}
+			}
+			else
+			{
+				// Keep track of any messages that have sent stop but no start message
+				// In rare cases these can be processed out of order
+				SoundsStoppedBeforeStart.Add(Msg.PlayOrder);
 			}
 		});
 
@@ -553,6 +563,14 @@ namespace UE::Audio::Insights
 	void FSoundTraceProvider::GetOrCreateActiveSoundEntry(const FSoundStartMessage& Msg, TSharedPtr<FSoundDashboardEntry>& OutReturnedSoundEntry)
 	{
 		OutReturnedSoundEntry = nullptr;
+		
+		// We cannot always guarentee the order of receiving start and stop messages
+		// If the stop message preceeded the start message, do not create a new entry
+		if (SoundsStoppedBeforeStart.Contains(Msg.PlayOrder))
+		{
+			SoundsStoppedBeforeStart.Remove(Msg.PlayOrder);
+			return;
+		}
 
 		UpdateDeviceEntry(Msg.DeviceId, Msg.EntryType, [this, &OutReturnedSoundEntry, &Msg](TSharedPtr<FSoundDashboardEntry>& CategoryEntry)
 		{
diff --git a/Engine/Plugins/AudioInsights/Source/AudioInsights/Private/Providers/SoundTraceProvider.h b/Engine/Plugins/AudioInsights/Source/AudioInsights/Private/Providers/SoundTraceProvider.h
--- a/Engine/Plugins/AudioInsights/Source/AudioInsights/Private/Providers/SoundTraceProvider.h
+++ b/Engine/Plugins/AudioInsights/Source/AudioInsights/Private/Providers/SoundTraceProvider.h
@@ -61,5 +61,6 @@ namespace UE::Audio::Insights
 
 		TMap<uint32, SoundEntryKeys> ActiveSoundToEntryKeysMap;
 		TArray<SoundMessageIDs> EntriesTimingOut;
+		TSet<uint32> SoundsStoppedBeforeStart;
 	};
 } // namespace UE::Audio::Insights
