[Audio Insights] Show aggregate data on active sound types with multiple wave instances inside the Sounds Dashboard. Change sorting method in columns to not default to first child and use the aggregate data instead.

diff --git a/Engine/Plugins/AudioInsights/Source/AudioInsights/Private/Providers/SoundTraceProvider.cpp b/Engine/Plugins/AudioInsights/Source/AudioInsights/Private/Providers/SoundTraceProvider.cpp
--- a/Engine/Plugins/AudioInsights/Source/AudioInsights/Private/Providers/SoundTraceProvider.cpp
+++ b/Engine/Plugins/AudioInsights/Source/AudioInsights/Private/Providers/SoundTraceProvider.cpp
@@ -65,6 +65,17 @@ namespace UE::Audio::Insights
 			return true;
 		}
 
+		FDataPoint PeekLastValue(const ::Audio::TCircularAudioBuffer<FDataPoint>& ParameterBuffer, const float DefaultReturn = 0.0f)
+		{
+			if (ParameterBuffer.Num() == 0)
+			{
+				return { 0.0, DefaultReturn };
+			}
+
+			const ::Audio::DisjointedArrayView<const FDataPoint> DataPointsDisjointedArrayView = ParameterBuffer.PeekInPlace(ParameterBuffer.Num());
+			return DataPointsDisjointedArrayView.FirstBuffer.Last();
+		}
+
 		void UpdateParameterEntry(const FMixerSourceParameterMessage& Msg, const float DataPoint, ::Audio::TCircularAudioBuffer<FDataPoint>& OutDataPoints, TArray<int32, TInlineAllocator<64>>& OutEntriesWithPoppedDataPoints)
 		{
 			if (!OutEntriesWithPoppedDataPoints.Contains(Msg.PlayOrder))
@@ -520,6 +531,8 @@ namespace UE::Audio::Insights
 			}
 		});
 
+		UpdateAggregateActiveSoundData();
+
 		return true;
 	}
 
@@ -676,6 +689,107 @@ namespace UE::Audio::Insights
 		ActiveSoundToEntryKeysMap.Remove(ActiveSoundPlayOrder);
 	}
 
+	void FSoundTraceProvider::UpdateAggregateActiveSoundData()
+	{
+		using namespace FSoundTraceProviderPrivate;
+
+		for (const auto& [AudioDeviceID, DeviceData] : DeviceDataMap)
+		{
+			for (const auto& [EntryType, SoundDashboardEntry] : DeviceData)
+			{
+				if (!SoundDashboardEntry.IsValid() || !EntryTypeHasSoundWaveEntries(EntryType))
+				{
+					continue;
+				}
+
+				for (const TSharedPtr<IDashboardDataTreeViewEntry>& ActiveSoundEntry : SoundDashboardEntry->Children)
+				{
+					if (ActiveSoundEntry.IsValid())
+					{
+						CollectAggregateData(CastEntry(*ActiveSoundEntry));
+					}
+				}
+			}
+		}
+	}
+
+	void FSoundTraceProvider::CollectAggregateData(FSoundDashboardEntry& ActiveSoundEntry)
+	{
+		using namespace FSoundTraceProviderPrivate;
+
+		const int32 NumWaveInstances = ActiveSoundEntry.Children.Num();
+		if (NumWaveInstances <= 0)
+		{
+			return;
+		}
+
+		auto GetMaxDatapoint = [](const FDataPoint& First, const FDataPoint& Second)
+		{
+			return First.Value >= Second.Value ? First : Second;
+		};
+
+		auto GetMinDatapoint = [](const FDataPoint& First, const FDataPoint& Second)
+		{
+			return First.Value < Second.Value ? First : Second;
+		};
+
+		auto AccumulateDataPoints = [](const FDataPoint& AccumulativeValue, const FDataPoint& AdditionalValue) -> FDataPoint
+		{
+			return { FMath::Max(AccumulativeValue.Key, AdditionalValue.Key), AccumulativeValue.Value + AdditionalValue.Value };
+		};
+
+		auto AverageDataPoints = [](const FDataPoint& DataPointsTotal, const int32 NumDataPoints) -> FDataPoint
+		{
+			if (NumDataPoints == 0)
+			{
+				return { 0.0, 0.0f };
+			}
+
+			return { DataPointsTotal.Key,  DataPointsTotal.Value / static_cast<float>(NumDataPoints) };
+		};
+
+		auto AddSingleValueToBuffer = [](::Audio::TCircularAudioBuffer<FDataPoint>& To, const FDataPoint& Value)
+		{
+			if (To.Num() > 0)
+			{
+				To.Pop(1);
+			}
+
+			To.Push(Value);
+		};
+
+		FDataPoint MaxVolume { 0.0, 0.0f };
+		FDataPoint MinDistance { 0.0, TNumericLimits<float>::Max() };
+		FDataPoint MaxPriority { 0.0, 0.0f };
+
+		FDataPoint AveragePeakAmp { 0.0, 0.0f };
+		FDataPoint CumulativePeakAmplitude { 0.0, 0.0f };
+
+		FDataPoint AveragePitch { 0.0, 0.0f };
+		FDataPoint CumulativePitch { 0.0, 0.0f };
+
+		for (const TSharedPtr<IDashboardDataTreeViewEntry>& WaveInstanceEntry : ActiveSoundEntry.Children)
+		{
+			const FSoundDashboardEntry& WaveSoundEntry = CastEntry(*WaveInstanceEntry);
+
+			MaxVolume = GetMaxDatapoint(MaxVolume, PeekLastValue(WaveSoundEntry.VolumeDataPoints, MaxVolume.Value));
+			MinDistance = GetMinDatapoint(MinDistance, PeekLastValue(WaveSoundEntry.DistanceDataPoints, MinDistance.Value));
+			MaxPriority = GetMaxDatapoint(MaxPriority, PeekLastValue(WaveSoundEntry.PriorityDataPoints, MaxPriority.Value));
+
+			CumulativePeakAmplitude = AccumulateDataPoints(CumulativePeakAmplitude, PeekLastValue(WaveSoundEntry.AmplitudeDataPoints));
+			CumulativePitch = AccumulateDataPoints(CumulativePitch, PeekLastValue(WaveSoundEntry.PitchDataPoints));
+		}
+
+		AveragePeakAmp = AverageDataPoints(CumulativePeakAmplitude, NumWaveInstances);
+		AveragePitch = AverageDataPoints(CumulativePitch, NumWaveInstances);
+
+		AddSingleValueToBuffer(ActiveSoundEntry.VolumeDataPoints, MaxVolume);
+		AddSingleValueToBuffer(ActiveSoundEntry.DistanceDataPoints, MinDistance);
+		AddSingleValueToBuffer(ActiveSoundEntry.PriorityDataPoints, MaxPriority);
+		AddSingleValueToBuffer(ActiveSoundEntry.AmplitudeDataPoints, AveragePeakAmp);
+		AddSingleValueToBuffer(ActiveSoundEntry.PitchDataPoints, AveragePitch);
+	}
+
 #if !WITH_EDITOR
 	void FSoundTraceProvider::OnTimingViewTimeMarkerChanged(double TimeMarker)
 	{
diff --git a/Engine/Plugins/AudioInsights/Source/AudioInsights/Private/Providers/SoundTraceProvider.h b/Engine/Plugins/AudioInsights/Source/AudioInsights/Private/Providers/SoundTraceProvider.h
--- a/Engine/Plugins/AudioInsights/Source/AudioInsights/Private/Providers/SoundTraceProvider.h
+++ b/Engine/Plugins/AudioInsights/Source/AudioInsights/Private/Providers/SoundTraceProvider.h
@@ -40,6 +40,9 @@ namespace UE::Audio::Insights
 		void GetActiveSoundEntryFromIDs(const uint32 PlayOrderID, const ::Audio::FDeviceId DeviceId, TSharedPtr<FSoundDashboardEntry>& OutSoundEntry);
 		void RemoveActiveSoundEntry(const ::Audio::FDeviceId DeviceId, TSharedPtr<FSoundDashboardEntry> OutEntry);
 
+		void UpdateAggregateActiveSoundData();
+		void CollectAggregateData(FSoundDashboardEntry& ActiveSoundEntry);
+
 #if !WITH_EDITOR
 		virtual void OnTimingViewTimeMarkerChanged(double TimeMarker) override;
 
diff --git a/Engine/Plugins/AudioInsights/Source/AudioInsights/Private/Views/SoundDashboardViewFactory.cpp b/Engine/Plugins/AudioInsights/Source/AudioInsights/Private/Views/SoundDashboardViewFactory.cpp
--- a/Engine/Plugins/AudioInsights/Source/AudioInsights/Private/Views/SoundDashboardViewFactory.cpp
+++ b/Engine/Plugins/AudioInsights/Source/AudioInsights/Private/Views/SoundDashboardViewFactory.cpp
@@ -1409,7 +1409,7 @@ namespace UE::Audio::Insights
 					.GetDisplayValue = [](const IDashboardDataTreeViewEntry& InData)
 					{
 						const FSoundDashboardEntry& SoundDashboardEntry = CastEntry(InData);
-						if (SoundDashboardEntry.Children.Num() > 0)
+						if (SoundDashboardEntry.bIsCategory)
 						{
 							return FText();
 						}
@@ -1443,7 +1443,7 @@ namespace UE::Audio::Insights
 					.GetDisplayValue = [](const IDashboardDataTreeViewEntry& InData)
 					{
 						const FSoundDashboardEntry& SoundDashboardEntry = CastEntry(InData);
-						if (SoundDashboardEntry.Children.Num() > 0)
+						if (SoundDashboardEntry.bIsCategory)
 						{
 							return FText();
 						}
@@ -1465,7 +1465,7 @@ namespace UE::Audio::Insights
 					.GetDisplayValue = [](const IDashboardDataTreeViewEntry& InData)
 					{
 						const FSoundDashboardEntry& SoundDashboardEntry = CastEntry(InData);
-						if (SoundDashboardEntry.Children.Num() > 0)
+						if (SoundDashboardEntry.bIsCategory)
 						{
 							return FText();
 						}
@@ -1487,7 +1487,7 @@ namespace UE::Audio::Insights
 					.GetDisplayValue = [](const IDashboardDataTreeViewEntry& InData)
 					{
 						const FSoundDashboardEntry& SoundDashboardEntry = CastEntry(InData);
-						if (SoundDashboardEntry.Children.Num() > 0)
+						if (SoundDashboardEntry.bIsCategory)
 						{
 							return FText();
 						}
@@ -1509,7 +1509,7 @@ namespace UE::Audio::Insights
 					.GetDisplayValue = [](const IDashboardDataTreeViewEntry& InData)
 					{
 						const FSoundDashboardEntry& SoundDashboardEntry = CastEntry(InData);
-						if (SoundDashboardEntry.Children.Num() > 0)
+						if (SoundDashboardEntry.bIsCategory)
 						{
 							return FText();
 						}
@@ -1609,7 +1609,7 @@ namespace UE::Audio::Insights
 
 		if (SortByColumn == "Name")
 		{
-			SortByPredicate(SortByName, false/* bSortByChildren*/);
+			SortByPredicate(SortByName);
 		}
 		else if (SortByColumn == "Priority")
 		{
@@ -1666,7 +1666,7 @@ namespace UE::Audio::Insights
 		return bDataReset;
 	}
 
-	void FSoundDashboardViewFactory::RecursiveSort(TArray<TSharedPtr<IDashboardDataTreeViewEntry>>& OutTree, TFunctionRef<bool(const FSoundDashboardEntry&, const FSoundDashboardEntry&)> Predicate, const bool bSortByChildren)
+	void FSoundDashboardViewFactory::RecursiveSort(TArray<TSharedPtr<IDashboardDataTreeViewEntry>>& OutTree, TFunctionRef<bool(const FSoundDashboardEntry&, const FSoundDashboardEntry&)> Predicate)
 	{
 		using namespace FSoundDashboardViewFactoryPrivate;
 
@@ -1674,7 +1674,7 @@ namespace UE::Audio::Insights
 		{
 			if (Entry->Children.Num() > 0)
 			{
-				RecursiveSort(Entry->Children, Predicate, bSortByChildren);
+				RecursiveSort(Entry->Children, Predicate);
 			}
 		}
 
@@ -1687,15 +1687,9 @@ namespace UE::Audio::Insights
 			}
 		}
 
-		auto SortDashboardEntries = [this, bSortByChildren](const TSharedPtr<IDashboardDataTreeViewEntry>& First, const TSharedPtr<IDashboardDataTreeViewEntry>& Second, TFunctionRef<bool(const FSoundDashboardEntry&, const FSoundDashboardEntry&)> Predicate)
+		auto SortDashboardEntries = [this](const TSharedPtr<IDashboardDataTreeViewEntry>& First, const TSharedPtr<IDashboardDataTreeViewEntry>& Second, TFunctionRef<bool(const FSoundDashboardEntry&, const FSoundDashboardEntry&)> Predicate)
 		{
-			const TSharedPtr<IDashboardDataTreeViewEntry>& FirstEntryToSortBy = bSortByChildren && (First->Children.Num() > 0) ? First->Children[0] : First;
-			const TSharedPtr<IDashboardDataTreeViewEntry>& SecondtEntryToSortBy = bSortByChildren && (Second->Children.Num() > 0) ? Second->Children[0] : Second;
-
-			const FSoundDashboardEntry& FirstData = CastEntry(*FirstEntryToSortBy);
-			const FSoundDashboardEntry& SecondData = CastEntry(*SecondtEntryToSortBy);
-
-			return Predicate(FirstData, SecondData);
+			return Predicate(CastEntry(*First), CastEntry(*Second));
 		};
 
 		if (SortMode == EColumnSortMode::Ascending)
@@ -1714,14 +1708,14 @@ namespace UE::Audio::Insights
 		}
 	}
 
-	void FSoundDashboardViewFactory::SortByPredicate(TFunctionRef<bool(const FSoundDashboardEntry&, const FSoundDashboardEntry&)> Predicate, const bool bSortByChildren /*= true*/)
+	void FSoundDashboardViewFactory::SortByPredicate(TFunctionRef<bool(const FSoundDashboardEntry&, const FSoundDashboardEntry&)> Predicate)
 	{
 		if (PinnedItemEntries.IsValid())
 		{
-			RecursiveSort(PinnedItemEntries->GetPinnedSectionEntry()->Children, Predicate, bSortByChildren);
+			RecursiveSort(PinnedItemEntries->GetPinnedSectionEntry()->Children, Predicate);
 		}
 
-		RecursiveSort(DataViewEntries, Predicate, bSortByChildren);
+		RecursiveSort(DataViewEntries, Predicate);
 	}
 
 	TSharedRef<SWidget> FSoundDashboardViewFactory::MakeShowRecentlyStoppedSoundsWidget()
diff --git a/Engine/Plugins/AudioInsights/Source/AudioInsights/Public/Views/SoundDashboardViewFactory.h b/Engine/Plugins/AudioInsights/Source/AudioInsights/Public/Views/SoundDashboardViewFactory.h
--- a/Engine/Plugins/AudioInsights/Source/AudioInsights/Public/Views/SoundDashboardViewFactory.h
+++ b/Engine/Plugins/AudioInsights/Source/AudioInsights/Public/Views/SoundDashboardViewFactory.h
@@ -131,8 +131,8 @@ namespace UE::Audio::Insights
 		virtual TArray<TSharedPtr<IDashboardDataTreeViewEntry>>& GetTreeItemsSource() override { return FullTree; }
 		virtual bool ResetTreeData() override;
 
-		void RecursiveSort(TArray<TSharedPtr<IDashboardDataTreeViewEntry>>& OutTree, TFunctionRef<bool(const FSoundDashboardEntry&, const FSoundDashboardEntry&)> Predicate, const bool bSortByChildren);
-		void SortByPredicate(TFunctionRef<bool(const FSoundDashboardEntry&, const FSoundDashboardEntry&)> Predicate, const bool bSortByChildren = true);
+		void RecursiveSort(TArray<TSharedPtr<IDashboardDataTreeViewEntry>>& OutTree, TFunctionRef<bool(const FSoundDashboardEntry&, const FSoundDashboardEntry&)> Predicate);
+		void SortByPredicate(TFunctionRef<bool(const FSoundDashboardEntry&, const FSoundDashboardEntry&)> Predicate);
 
 		TSharedRef<SWidget> MakeShowRecentlyStoppedSoundsWidget();
 
