[Audio Insights] Sound dashboard: small optimization to apply opacity per row instead of per column + misc. readability adjustments

diff --git a/Engine/Plugins/AudioInsights/Source/AudioInsights/Private/Views/SoundDashboardViewFactory.cpp b/Engine/Plugins/AudioInsights/Source/AudioInsights/Private/Views/SoundDashboardViewFactory.cpp
--- a/Engine/Plugins/AudioInsights/Source/AudioInsights/Private/Views/SoundDashboardViewFactory.cpp
+++ b/Engine/Plugins/AudioInsights/Source/AudioInsights/Private/Views/SoundDashboardViewFactory.cpp
@@ -1094,6 +1094,12 @@ namespace UE::Audio::Insights
 		using namespace FSoundDashboardViewFactoryPrivate;
 
 		return SNew(SRowWidget, OwnerTable, Item, AsShared())
+			.RenderOpacity_Lambda([Item]
+			{
+				const FSoundDashboardEntry& SoundDashboardEntry = CastEntry(*Item);
+
+				return SoundDashboardEntry.TimeoutTimestamp == INVALID_TIMEOUT ? 1.0f : 0.25f;
+			})
 			.Visibility_Lambda([this, Item]()
 			{
 				const FSoundDashboardEntry& SoundEntry = CastEntry(*Item);
@@ -1188,7 +1194,6 @@ namespace UE::Audio::Insights
 	{
 		using namespace FSoundDashboardViewFactoryPrivate;
 
-
 		const FSoundDashboardCommands& Commands = FSoundDashboardCommands::Get();
 
 		constexpr bool bShouldCloseWindowAfterMenuSelection = true;
@@ -1394,12 +1399,6 @@ namespace UE::Audio::Insights
 						}
 
 						return NAME_None;
-					},
-					.GetTextColorValue = [](const IDashboardDataTreeViewEntry& InData) -> FSlateColor
-					{
-						const FSoundDashboardEntry& SoundDashboardEntry = CastEntry(InData);
-						return SoundDashboardEntry.TimeoutTimestamp == INVALID_TIMEOUT ? FSlateColor::UseForeground()
-																					   : FSlateColor(FSlateStyle::Get().GetColor("SoundDashboard.TimingOutTextColor"));
 					}
 				}
 			},
@@ -1411,29 +1410,17 @@ namespace UE::Audio::Insights
 						const FSoundDashboardEntry& SoundDashboardEntry = CastEntry(InData);
 						if (SoundDashboardEntry.bIsCategory)
 						{
-							return FText();
+							return FText::GetEmpty();
 						}
 
-						const ::Audio::TCircularAudioBuffer<FDataPoint>& PriorityDataPoints = SoundDashboardEntry.PriorityDataPoints;
-						const float PriorityValue = GetLastEntryArrayValue(PriorityDataPoints);
+						const float PriorityValue = GetLastEntryArrayValue(SoundDashboardEntry.PriorityDataPoints);
 
 						// Max priority as defined in SoundWave.cpp
 						static constexpr float VolumeWeightedMaxPriority = TNumericLimits<float>::Max() / MAX_VOLUME;
 
-						if (PriorityValue >= VolumeWeightedMaxPriority)
-						{
-							return LOCTEXT("AudioDashboard_Sounds_Max", "MAX");
-						}
-						else
-						{
-							return FText::AsNumber(GetLastEntryArrayValue(PriorityDataPoints), FSlateStyle::Get().GetAmpFloatFormat());
-						}
-					},
-					.GetTextColorValue = [](const IDashboardDataTreeViewEntry& InData) -> FSlateColor
-					{
-						const FSoundDashboardEntry& SoundDashboardEntry = CastEntry(InData);
-						return SoundDashboardEntry.TimeoutTimestamp == INVALID_TIMEOUT ? FSlateColor::UseForeground()
-																					   : FSlateColor(FSlateStyle::Get().GetColor("SoundDashboard.TimingOutTextColor"));
+						return PriorityValue >= VolumeWeightedMaxPriority 
+							? LOCTEXT("AudioDashboard_Sounds_Max", "MAX") 
+							: FText::AsNumber(PriorityValue, FSlateStyle::Get().GetAmpFloatFormat());
 					}
 				}
 			},
@@ -1445,17 +1432,10 @@ namespace UE::Audio::Insights
 						const FSoundDashboardEntry& SoundDashboardEntry = CastEntry(InData);
 						if (SoundDashboardEntry.bIsCategory)
 						{
-							return FText();
+							return FText::GetEmpty();
 						}
 
-						const ::Audio::TCircularAudioBuffer<FDataPoint>& DistanceDataPoints = SoundDashboardEntry.DistanceDataPoints;
-						return FText::AsNumber(GetLastEntryArrayValue(DistanceDataPoints), FSlateStyle::Get().GetDefaultFloatFormat());
-					},
-					.GetTextColorValue = [](const IDashboardDataTreeViewEntry& InData) -> FSlateColor
-					{
-						const FSoundDashboardEntry& SoundDashboardEntry = CastEntry(InData);
-						return SoundDashboardEntry.TimeoutTimestamp == INVALID_TIMEOUT ? FSlateColor::UseForeground()
-																					   : FSlateColor(FSlateStyle::Get().GetColor("SoundDashboard.TimingOutTextColor"));
+						return FText::AsNumber(GetLastEntryArrayValue(SoundDashboardEntry.DistanceDataPoints), FSlateStyle::Get().GetDefaultFloatFormat());
 					}
 				}
 			},
@@ -1467,17 +1447,10 @@ namespace UE::Audio::Insights
 						const FSoundDashboardEntry& SoundDashboardEntry = CastEntry(InData);
 						if (SoundDashboardEntry.bIsCategory)
 						{
-							return FText();
+							return FText::GetEmpty();
 						}
 
-						const ::Audio::TCircularAudioBuffer<FDataPoint>& AmplitudeDataPoints = SoundDashboardEntry.AmplitudeDataPoints;
-						return FText::AsNumber(GetLastEntryArrayValue(AmplitudeDataPoints), FSlateStyle::Get().GetAmpFloatFormat());
-					},
-					.GetTextColorValue = [](const IDashboardDataTreeViewEntry& InData) -> FSlateColor
-					{
-						const FSoundDashboardEntry& SoundDashboardEntry = CastEntry(InData);
-						return SoundDashboardEntry.TimeoutTimestamp == INVALID_TIMEOUT ? FSlateColor::UseForeground()
-																					   : FSlateColor(FSlateStyle::Get().GetColor("SoundDashboard.TimingOutTextColor"));
+						return FText::AsNumber(GetLastEntryArrayValue(SoundDashboardEntry.AmplitudeDataPoints), FSlateStyle::Get().GetAmpFloatFormat());
 					}
 				}
 			},
@@ -1489,17 +1462,10 @@ namespace UE::Audio::Insights
 						const FSoundDashboardEntry& SoundDashboardEntry = CastEntry(InData);
 						if (SoundDashboardEntry.bIsCategory)
 						{
-							return FText();
+							return FText::GetEmpty();
 						}
 
-						const ::Audio::TCircularAudioBuffer<FDataPoint>& VolumeDataPoints = SoundDashboardEntry.VolumeDataPoints;
-						return FText::AsNumber(GetLastEntryArrayValue(VolumeDataPoints), FSlateStyle::Get().GetAmpFloatFormat());
-					},
-					.GetTextColorValue = [](const IDashboardDataTreeViewEntry& InData) -> FSlateColor
-					{
-						const FSoundDashboardEntry& SoundDashboardEntry = CastEntry(InData);
-						return SoundDashboardEntry.TimeoutTimestamp == INVALID_TIMEOUT ? FSlateColor::UseForeground()
-																					   : FSlateColor(FSlateStyle::Get().GetColor("SoundDashboard.TimingOutTextColor"));
+						return FText::AsNumber(GetLastEntryArrayValue(SoundDashboardEntry.VolumeDataPoints), FSlateStyle::Get().GetAmpFloatFormat());
 					}
 				}
 			},
@@ -1511,17 +1477,10 @@ namespace UE::Audio::Insights
 						const FSoundDashboardEntry& SoundDashboardEntry = CastEntry(InData);
 						if (SoundDashboardEntry.bIsCategory)
 						{
-							return FText();
+							return FText::GetEmpty();
 						}
 
-						const ::Audio::TCircularAudioBuffer<FDataPoint>& PitchDataPoints = SoundDashboardEntry.PitchDataPoints;
-						return FText::AsNumber(GetLastEntryArrayValue(PitchDataPoints), FSlateStyle::Get().GetPitchFloatFormat());
-					},
-					.GetTextColorValue = [](const IDashboardDataTreeViewEntry& InData) -> FSlateColor
-					{
-						const FSoundDashboardEntry& SoundDashboardEntry = CastEntry(InData);
-						return SoundDashboardEntry.TimeoutTimestamp == INVALID_TIMEOUT ? FSlateColor::UseForeground()
-																					   : FSlateColor(FSlateStyle::Get().GetColor("SoundDashboard.TimingOutTextColor"));
+						return FText::AsNumber(GetLastEntryArrayValue(SoundDashboardEntry.PitchDataPoints), FSlateStyle::Get().GetPitchFloatFormat());
 					}
 				}
 			}
diff --git a/Engine/Plugins/AudioInsights/Source/AudioInsights/Private/Views/TreeDashboardViewFactory.cpp b/Engine/Plugins/AudioInsights/Source/AudioInsights/Private/Views/TreeDashboardViewFactory.cpp
--- a/Engine/Plugins/AudioInsights/Source/AudioInsights/Private/Views/TreeDashboardViewFactory.cpp
+++ b/Engine/Plugins/AudioInsights/Source/AudioInsights/Private/Views/TreeDashboardViewFactory.cpp
@@ -23,6 +23,8 @@ namespace UE::Audio::Insights
 		Data = InData;
 		Factory = InFactory;
 
+		RenderOpacity = InArgs._RenderOpacity;
+
 		FSuperRowType::FArguments Args = FSuperRowType::FArguments().Style(&FSlateStyle::Get().GetWidgetStyle<FTableRowStyle>("TreeDashboard.TableViewRow"));
 
 		SMultiColumnTableRow<TSharedPtr<IDashboardDataTreeViewEntry>>::Construct(Args, InOwnerTable);
@@ -33,6 +35,14 @@ namespace UE::Audio::Insights
 		return Factory->GenerateWidgetForColumn(StaticCastSharedRef<SRowWidget>(AsShared()), Data->AsShared(), Column);
 	}
 
+	void FTraceTreeDashboardViewFactory::SRowWidget::Tick(const FGeometry& AllottedGeometry, const double InCurrentTime, const float InDeltaTime)
+	{
+		TRACE_CPUPROFILER_EVENT_SCOPE(FSoundDashboardViewFactory::TickRow)
+
+		SetRenderOpacity(RenderOpacity.Get());
+
+		SMultiColumnTableRow<TSharedPtr<IDashboardDataTreeViewEntry>>::Tick(AllottedGeometry, InCurrentTime, InDeltaTime);
+	}
 
 	// FTraceTreeDashboardViewFactory
 	FTraceTreeDashboardViewFactory::FTraceTreeDashboardViewFactory()
diff --git a/Engine/Plugins/AudioInsights/Source/AudioInsights/Public/AudioInsightsStyle.h b/Engine/Plugins/AudioInsights/Source/AudioInsights/Public/AudioInsightsStyle.h
--- a/Engine/Plugins/AudioInsights/Source/AudioInsights/Public/AudioInsightsStyle.h
+++ b/Engine/Plugins/AudioInsights/Source/AudioInsights/Public/AudioInsightsStyle.h
@@ -193,9 +193,6 @@ namespace UE::Audio::Insights
 			Set("SoundDashboard.PinnedColor", FLinearColor(0.9f, 0.9f, 0.9f));
 			Set("SoundDashboard.HiddenColor", FLinearColor(0.4f, 0.4f, 0.4f));
 
-			// Text colors
-			Set("SoundDashboard.TimingOutTextColor", FLinearColor(1.0f, 1.0f, 1.0f, 0.25f));
-
 			// Mute/Solo buttons style
 			const FSlateRoundedBoxBrush RoundedWhiteAlphaBrush(FLinearColor(1.0f, 1.0f, 1.0f, 0.1f), 5.0f);
 
diff --git a/Engine/Plugins/AudioInsights/Source/AudioInsights/Public/Views/TreeDashboardViewFactory.h b/Engine/Plugins/AudioInsights/Source/AudioInsights/Public/Views/TreeDashboardViewFactory.h
--- a/Engine/Plugins/AudioInsights/Source/AudioInsights/Public/Views/TreeDashboardViewFactory.h
+++ b/Engine/Plugins/AudioInsights/Source/AudioInsights/Public/Views/TreeDashboardViewFactory.h
@@ -49,14 +49,19 @@ namespace UE::Audio::Insights
 	protected:
 		struct SRowWidget : public SMultiColumnTableRow<TSharedPtr<IDashboardDataTreeViewEntry>>
 		{
-			SLATE_BEGIN_ARGS(SRowWidget){}
+			SLATE_BEGIN_ARGS(SRowWidget) 
+				: _RenderOpacity(1.0f)
+			{}
+				SLATE_ATTRIBUTE(float, RenderOpacity)
 			SLATE_END_ARGS()
 
 			void Construct(const FArguments& InArgs, const TSharedRef<STableViewBase>& InOwnerTable, TSharedPtr<IDashboardDataTreeViewEntry> InData, TSharedRef<FTraceTreeDashboardViewFactory> InFactory);
 			virtual TSharedRef<SWidget> GenerateWidgetForColumn(const FName& Column) override;
+			virtual void Tick(const FGeometry& AllottedGeometry, const double InCurrentTime, const float InDeltaTime) override;
 
 			TSharedPtr<IDashboardDataTreeViewEntry> Data;
 			TSharedPtr<FTraceTreeDashboardViewFactory> Factory;
+			TAttribute<float> RenderOpacity;
 		};
 
 		struct FHeaderRowColumnData
