[Audio Insights] Sound dashboard: moved EntryCanHaveChildren and IsDescendant to private namespace

diff --git a/Engine/Plugins/AudioInsights/Source/AudioInsights/Private/Views/SoundDashboardViewFactory.cpp b/Engine/Plugins/AudioInsights/Source/AudioInsights/Private/Views/SoundDashboardViewFactory.cpp
--- a/Engine/Plugins/AudioInsights/Source/AudioInsights/Private/Views/SoundDashboardViewFactory.cpp
+++ b/Engine/Plugins/AudioInsights/Source/AudioInsights/Private/Views/SoundDashboardViewFactory.cpp
@@ -114,6 +114,31 @@ namespace UE::Audio::Insights
 
 			SoundEntry.bShouldForceExpandChildren = SoundEntry.bIsExpanded;
 		}
+		
+		bool EntryCanHaveChildren(const TSharedRef<IDashboardDataTreeViewEntry>& InEntry)
+		{
+			const FSoundDashboardEntry& SoundEntry = FSoundDashboardViewFactoryPrivate::CastEntry(*InEntry);
+
+			const bool bIsSoundCueType = SoundEntry.EntryType == ESoundDashboardEntryType::SoundCue || SoundEntry.EntryType == ESoundDashboardEntryType::SoundCueTemplate;
+
+			return IsCategoryItem(*InEntry) || bIsSoundCueType;
+		}
+
+		bool IsDescendant(const TSharedPtr<IDashboardDataTreeViewEntry>& InEntry, const TSharedPtr<IDashboardDataTreeViewEntry>& InChildCandidate)
+		{
+			if (InEntry.IsValid() && EntryCanHaveChildren(InEntry.ToSharedRef()))
+			{
+				for (const TSharedPtr<IDashboardDataTreeViewEntry>& ChildEntry : InEntry->Children)
+				{
+					if (ChildEntry == InChildCandidate || IsDescendant(ChildEntry, InChildCandidate))
+					{
+						return true;
+					}
+				}
+			}
+
+			return false;
+		}
 
 		bool HasPinEntryType(const IDashboardDataTreeViewEntry& InEntry, const FSoundDashboardEntry::EPinnedEntryType PinnedEntryType)
 		{
@@ -947,34 +972,6 @@ namespace UE::Audio::Insights
 		return FilteredEntriesListView.IsValid() && FilteredEntriesListView->GetRootItems().Contains(InEntry);
 	}
 
-	bool FSoundDashboardViewFactory::EntryCanHaveChildren(const TSharedRef<IDashboardDataTreeViewEntry>& InEntry) const
-	{
-		using namespace FSoundDashboardViewFactoryPrivate;
-
-		const FSoundDashboardEntry& SoundEntry = FSoundDashboardViewFactoryPrivate::CastEntry(*InEntry);
-
-		const bool bIsSoundCueType = SoundEntry.EntryType == ESoundDashboardEntryType::SoundCue || SoundEntry.EntryType == ESoundDashboardEntryType::SoundCueTemplate;
-
-		return IsCategoryItem(*InEntry) || bIsSoundCueType;
-	}
-
-
-	bool FSoundDashboardViewFactory::IsDescendant(const TSharedPtr<IDashboardDataTreeViewEntry>& InEntry, const TSharedPtr<IDashboardDataTreeViewEntry>& InChildCandidate) const
-	{
-		if (InEntry.IsValid() && EntryCanHaveChildren(InEntry.ToSharedRef()))
-		{
-			for (const TSharedPtr<IDashboardDataTreeViewEntry>& ChildEntry : InEntry->Children)
-			{
-				if (ChildEntry == InChildCandidate || IsDescendant(ChildEntry, InChildCandidate))
-				{
-					return true;
-				}
-			}
-		}
-
-		return false;
-	}
-
 	TSharedRef<SWidget> FSoundDashboardViewFactory::GenerateWidgetForColumn(const TSharedRef<FTraceTreeDashboardViewFactory::SRowWidget>& InRowWidget, const TSharedRef<IDashboardDataTreeViewEntry>& InRowData, const FName& InColumn)
 	{
 		using namespace FSoundDashboardViewFactoryPrivate;
@@ -1728,6 +1725,8 @@ namespace UE::Audio::Insights
 		TFunction<void(const TArray<TSharedPtr<IDashboardDataTreeViewEntry>>&)> MuteSoloToggleFunc, 
 		TFunctionRef<bool(const IDashboardDataTreeViewEntry&, const bool)> IsMuteSoloFunc)
 	{
+		using namespace FSoundDashboardViewFactoryPrivate;
+
 		return SNew(SHorizontalBox)
 			+ SHorizontalBox::Slot()
 			.AutoWidth()
diff --git a/Engine/Plugins/AudioInsights/Source/AudioInsights/Public/Views/SoundDashboardViewFactory.h b/Engine/Plugins/AudioInsights/Source/AudioInsights/Public/Views/SoundDashboardViewFactory.h
--- a/Engine/Plugins/AudioInsights/Source/AudioInsights/Public/Views/SoundDashboardViewFactory.h
+++ b/Engine/Plugins/AudioInsights/Source/AudioInsights/Public/Views/SoundDashboardViewFactory.h
@@ -113,8 +113,6 @@ namespace UE::Audio::Insights
 		virtual TSharedPtr<SWidget> GetFilterBarButtonWidget() override;
 
 		bool IsRootItem(const TSharedRef<IDashboardDataTreeViewEntry>& InEntry) const;
-		bool EntryCanHaveChildren(const TSharedRef<IDashboardDataTreeViewEntry>& InEntry) const;
-		bool IsDescendant(const TSharedPtr<IDashboardDataTreeViewEntry>& InEntry, const TSharedPtr<IDashboardDataTreeViewEntry>& InChildCandidate) const;
 
 		TSharedRef<SWidget> GenerateWidgetForRootColumn(const TSharedRef<FTraceTreeDashboardViewFactory::SRowWidget>& InRowWidget, const TSharedRef<IDashboardDataTreeViewEntry>& InRowData, const FName& InColumn, const FText& InValueText);
 		virtual TSharedRef<SWidget> GenerateWidgetForColumn(const TSharedRef<FTraceTreeDashboardViewFactory::SRowWidget>& InRowWidget, const TSharedRef<IDashboardDataTreeViewEntry>& InRowData, const FName& InColumn) override;
