[Audio Insights] Add relative render cost column to Sounds Tab

diff --git a/Engine/Plugins/AudioInsights/Source/AudioInsights/Private/Messages/SoundTraceMessages.cpp b/Engine/Plugins/AudioInsights/Source/AudioInsights/Private/Messages/SoundTraceMessages.cpp
--- a/Engine/Plugins/AudioInsights/Source/AudioInsights/Private/Messages/SoundTraceMessages.cpp
+++ b/Engine/Plugins/AudioInsights/Source/AudioInsights/Private/Messages/SoundTraceMessages.cpp
@@ -66,5 +66,6 @@ namespace UE::Audio::Insights
 		AmplitudeDataPoints.SetCapacity(DataPointsCapacity);
 		VolumeDataPoints.SetCapacity(DataPointsCapacity);
 		PitchDataPoints.SetCapacity(DataPointsCapacity);
+		RelativeRenderCostDataPoints.SetCapacity(DataPointsCapacity);
 	}
 }
\ No newline at end of file
diff --git a/Engine/Plugins/AudioInsights/Source/AudioInsights/Private/Providers/SoundTraceProvider.cpp b/Engine/Plugins/AudioInsights/Source/AudioInsights/Private/Providers/SoundTraceProvider.cpp
--- a/Engine/Plugins/AudioInsights/Source/AudioInsights/Private/Providers/SoundTraceProvider.cpp
+++ b/Engine/Plugins/AudioInsights/Source/AudioInsights/Private/Providers/SoundTraceProvider.cpp
@@ -198,6 +198,7 @@ namespace UE::Audio::Insights
 				Builder.RouteEvent(RouteId_Amplitude,			"Audio", "MixerSourceEnvelope");
 				Builder.RouteEvent(RouteId_Volume,				"Audio", "MixerSourceVolume");
 				Builder.RouteEvent(RouteId_Pitch,				"Audio", "MixerSourcePitch");
+				Builder.RouteEvent(RouteId_RelativeRenderCost,	"Audio", "MixerSourceRelativeRenderCost");
 			}
 
 			virtual bool OnEvent(uint16 RouteId, EStyle Style, const FOnEventContext& Context) override
@@ -256,6 +257,12 @@ namespace UE::Audio::Insights
 						break;
 					}
 
+					case RouteId_RelativeRenderCost:
+					{
+						Messages.RelativeRenderCostMessages.Enqueue(FMixerSourceRelativeRenderCostMessage{ Context });
+						break;
+					}
+
 					default:
 					{
 						return OnEventFailure(RouteId, Style, Context);
@@ -282,7 +289,8 @@ namespace UE::Audio::Insights
 				RouteId_Distance,
 				RouteId_Amplitude,
 				RouteId_Volume,
-				RouteId_Pitch
+				RouteId_Pitch,
+				RouteId_RelativeRenderCost
 			};
 
 			TraceServices::IAnalysisSession& Session;
@@ -496,6 +504,25 @@ namespace UE::Audio::Insights
 			FSoundTraceProviderPrivate::UpdateParameterEntry(Msg, Msg.Pitch, SoundEntry->PitchDataPoints, EntriesWithPoppedDataPoints);
 		});
 
+		ProcessMessageQueue<FMixerSourceRelativeRenderCostMessage>(TraceMessages.RelativeRenderCostMessages, GetActiveSoundEntryFromParameterMessage,
+		[this, &EntriesWithPoppedDataPoints](const FMixerSourceRelativeRenderCostMessage& Msg, const TSharedPtr<FSoundDashboardEntry>* OutActiveSoundEntry)
+		{
+			FSoundDashboardEntry* SoundEntry = FSoundTraceProviderPrivate::FindSoundEntryForActiveSound(Msg.PlayOrder, OutActiveSoundEntry);
+			if (SoundEntry == nullptr)
+			{
+				return;
+			}
+
+#if !WITH_EDITOR
+			if (SessionCachedMessages.IsValid())
+			{
+				SessionCachedMessages->RelativeRenderCostCachedMessages.EmplaceBack(Msg);
+			}
+#endif // !WITH_EDITOR
+
+			FSoundTraceProviderPrivate::UpdateParameterEntry(Msg, Msg.RelativeRenderCost, SoundEntry->RelativeRenderCostDataPoints, EntriesWithPoppedDataPoints);
+		});
+
 		ProcessMessageQueue<FSoundStopMessage>(TraceMessages.StopMessages, GetSoundEntryFromStopMessage,
 		[this](const FSoundStopMessage& Msg, const TSharedPtr<FSoundDashboardEntry>* OutActiveSoundEntry)
 		{
@@ -768,6 +795,8 @@ namespace UE::Audio::Insights
 		FDataPoint AveragePitch { 0.0, 0.0f };
 		FDataPoint CumulativePitch { 0.0, 0.0f };
 
+		FDataPoint CumulativeRelativeRenderCost { 0.0, 0.0f };
+
 		for (const TSharedPtr<IDashboardDataTreeViewEntry>& WaveInstanceEntry : ActiveSoundEntry.Children)
 		{
 			const FSoundDashboardEntry& WaveSoundEntry = CastEntry(*WaveInstanceEntry);
@@ -778,6 +807,7 @@ namespace UE::Audio::Insights
 
 			CumulativePeakAmplitude = AccumulateDataPoints(CumulativePeakAmplitude, PeekLastValue(WaveSoundEntry.AmplitudeDataPoints));
 			CumulativePitch = AccumulateDataPoints(CumulativePitch, PeekLastValue(WaveSoundEntry.PitchDataPoints));
+			CumulativeRelativeRenderCost = AccumulateDataPoints(CumulativeRelativeRenderCost, PeekLastValue(WaveSoundEntry.RelativeRenderCostDataPoints));
 		}
 
 		AveragePeakAmp = AverageDataPoints(CumulativePeakAmplitude, NumWaveInstances);
@@ -788,6 +818,7 @@ namespace UE::Audio::Insights
 		AddSingleValueToBuffer(ActiveSoundEntry.PriorityDataPoints, MaxPriority);
 		AddSingleValueToBuffer(ActiveSoundEntry.AmplitudeDataPoints, AveragePeakAmp);
 		AddSingleValueToBuffer(ActiveSoundEntry.PitchDataPoints, AveragePitch);
+		AddSingleValueToBuffer(ActiveSoundEntry.RelativeRenderCostDataPoints, CumulativeRelativeRenderCost);
 	}
 
 #if !WITH_EDITOR
@@ -931,6 +962,7 @@ namespace UE::Audio::Insights
 				FMixerSourceEnvelopeMessage AmplitudeMessage;
 				FMixerSourceVolumeMessage VolumeMessage;
 				FMixerSourcePitchMessage PitchMessage;
+				FMixerSourceRelativeRenderCostMessage RelativeRenderCostMessage;
 			};
 
 			TArray<uint32> PlayOrderArray;
@@ -979,6 +1011,13 @@ namespace UE::Audio::Insights
 				{
 					CachedEntryInfos[Index].PitchMessage = *FoundPitchCachedMessage;
 				}
+
+				// Relative Render Cost
+				const FMixerSourceRelativeRenderCostMessage* FoundRelativeRenderCostCachedMessage = FindClosestMessageToTimestamp(SessionCachedMessages->RelativeRenderCostCachedMessages, TimeMarker, PlayOrder);
+				if (FoundRelativeRenderCostCachedMessage)
+				{
+					CachedEntryInfos[Index].RelativeRenderCostMessage = *FoundRelativeRenderCostCachedMessage;
+				}
 			});
 
 			// Update the device entries with the collected info
@@ -1002,6 +1041,7 @@ namespace UE::Audio::Insights
 				SoundEntry->AmplitudeDataPoints.Push({ CachedEntryInfo.AmplitudeMessage.Timestamp, CachedEntryInfo.AmplitudeMessage.Envelope });
 				SoundEntry->VolumeDataPoints.Push({ CachedEntryInfo.VolumeMessage.Timestamp, CachedEntryInfo.VolumeMessage.Volume });
 				SoundEntry->PitchDataPoints.Push({ CachedEntryInfo.PitchMessage.Timestamp, CachedEntryInfo.PitchMessage.Pitch });
+				SoundEntry->RelativeRenderCostDataPoints.Push({ CachedEntryInfo.RelativeRenderCostMessage.Timestamp, CachedEntryInfo.RelativeRenderCostMessage.RelativeRenderCost });
 
 				SoundEntry->Timestamp = CachedEntryInfo.PriorityMessage.Timestamp;
 			}
diff --git a/Engine/Plugins/AudioInsights/Source/AudioInsights/Private/Views/SoundDashboardViewFactory.cpp b/Engine/Plugins/AudioInsights/Source/AudioInsights/Private/Views/SoundDashboardViewFactory.cpp
--- a/Engine/Plugins/AudioInsights/Source/AudioInsights/Private/Views/SoundDashboardViewFactory.cpp
+++ b/Engine/Plugins/AudioInsights/Source/AudioInsights/Private/Views/SoundDashboardViewFactory.cpp
@@ -422,6 +422,7 @@ namespace UE::Audio::Insights
 			AddLastValueInBuffer(PinnedEntry.AmplitudeDataPoints, OriginalEntry.AmplitudeDataPoints);
 			AddLastValueInBuffer(PinnedEntry.VolumeDataPoints, OriginalEntry.VolumeDataPoints);
 			AddLastValueInBuffer(PinnedEntry.PitchDataPoints, OriginalEntry.PitchDataPoints);
+			AddLastValueInBuffer(PinnedEntry.RelativeRenderCostDataPoints, OriginalEntry.RelativeRenderCostDataPoints);
 		};
 	} // namespace FPinnedSoundEntryWrapperPrivate
 
@@ -1331,6 +1332,17 @@ namespace UE::Audio::Insights
 					.FillWidth        = 0.1f,
 					.Alignment        = EHorizontalAlignment::HAlign_Left
 				}
+			},
+			{
+				"RelativeRenderCost",
+				{
+					.DisplayName = LOCTEXT("SoundDashboard_RelativeRenderCostColumnDisplayName", "Rel. Render Cost"),
+					.IconName = NAME_None,
+					.bShowDisplayName = true,
+					.bDefaultHidden = true,
+					.FillWidth = 0.1f,
+					.Alignment = EHorizontalAlignment::HAlign_Left
+				}
 			}
 		};
 
@@ -1480,6 +1492,21 @@ namespace UE::Audio::Insights
 						return FText::AsNumber(GetLastEntryArrayValue(SoundDashboardEntry.PitchDataPoints), FSlateStyle::Get().GetPitchFloatFormat());
 					}
 				}
+			},
+			{
+				"RelativeRenderCost",
+				{
+					.GetDisplayValue = [](const IDashboardDataTreeViewEntry& InData)
+					{
+						const FSoundDashboardEntry& SoundDashboardEntry = CastEntry(InData);
+						if (SoundDashboardEntry.bIsCategory)
+						{
+							return FText::GetEmpty();
+						}
+
+						return FText::AsNumber(GetLastEntryArrayValue(SoundDashboardEntry.RelativeRenderCostDataPoints), FSlateStyle::Get().GetRelativeRenderCostFormat());
+					}
+				}
 			}
 		};
 
@@ -1563,6 +1590,17 @@ namespace UE::Audio::Insights
 			return ComparisonDiff < 0.0f;
 		};
 
+		auto SortByRelativeRenderCost = [&SortByPlayOrder](const FSoundDashboardEntry& First, const FSoundDashboardEntry& Second)
+		{
+			const float ComparisonDiff = GetLastEntryArrayValue(First.RelativeRenderCostDataPoints) - GetLastEntryArrayValue(Second.RelativeRenderCostDataPoints);
+			if (FMath::IsNearlyZero(ComparisonDiff, UE_KINDA_SMALL_NUMBER))
+			{
+				return SortByPlayOrder(First, Second);
+			}
+
+			return ComparisonDiff < 0.0f;
+		};
+
 		if (SortByColumn == "Name")
 		{
 			SortByPredicate(SortByName);
@@ -1587,6 +1625,10 @@ namespace UE::Audio::Insights
 		{
 			SortByPredicate(SortByPitch);
 		}
+		else if (SortByColumn == "RelativeRenderCost")
+		{
+			SortByPredicate(SortByRelativeRenderCost);
+		}
 
 		FullTree.Reset();
 
diff --git a/Engine/Plugins/AudioInsights/Source/AudioInsights/Public/AudioInsightsStyle.h b/Engine/Plugins/AudioInsights/Source/AudioInsights/Public/AudioInsightsStyle.h
--- a/Engine/Plugins/AudioInsights/Source/AudioInsights/Public/AudioInsightsStyle.h
+++ b/Engine/Plugins/AudioInsights/Source/AudioInsights/Public/AudioInsightsStyle.h
@@ -77,6 +77,16 @@ namespace UE::Audio::Insights
 			return &FloatFormat;
 		};
 
+		const FNumberFormattingOptions* GetRelativeRenderCostFormat() const
+		{
+			static FNumberFormattingOptions FloatFormat;
+			FloatFormat.MinimumIntegralDigits = 1;
+			FloatFormat.MaximumIntegralDigits = 3;
+			FloatFormat.MinimumFractionalDigits = 2;
+			FloatFormat.MaximumFractionalDigits = 2;
+			return &FloatFormat;
+		};
+
 		FText FormatSecondsAsTime(float InTimeSec) const
 		{
 			return FText::Format(LOCTEXT("TimeInSecondsFormat", "{0}s"), FText::AsNumber(InTimeSec, GetTimeFormat()));
diff --git a/Engine/Plugins/AudioInsights/Source/AudioInsights/Public/Messages/MixerSourceTraceMessages.h b/Engine/Plugins/AudioInsights/Source/AudioInsights/Public/Messages/MixerSourceTraceMessages.h
--- a/Engine/Plugins/AudioInsights/Source/AudioInsights/Public/Messages/MixerSourceTraceMessages.h
+++ b/Engine/Plugins/AudioInsights/Source/AudioInsights/Public/Messages/MixerSourceTraceMessages.h
@@ -83,6 +83,7 @@ namespace UE::Audio::Insights
 	DEFINE_MIXERSOURCE_PARAM_MESSAGE(FMixerSourceEnvelopeMessage, FMixerSourceParameterMessage, Envelope, float, 0.0f)
 	DEFINE_MIXERSOURCE_PARAM_MESSAGE(FMixerSourcePitchMessage, FMixerSourceParameterMessage, Pitch, float, 1.0f)
 	DEFINE_MIXERSOURCE_PARAM_MESSAGE(FMixerSourceVolumeMessage, FMixerSourceParameterMessage, Volume, float, 1.0f)
+	DEFINE_MIXERSOURCE_PARAM_MESSAGE(FMixerSourceRelativeRenderCostMessage, FMixerSourceParameterMessage, RelativeRenderCost, float, 0.0f)
 #undef DEFINE_MIXERSOURCE_PARAM_MESSAGE
 
 	using FDataPoint = TPair<double, float>; // (Timestamp, Value)
@@ -112,6 +113,7 @@ namespace UE::Audio::Insights
 		::Audio::TCircularAudioBuffer<FDataPoint> HPFFreqDataPoints;
 		::Audio::TCircularAudioBuffer<FDataPoint> EnvelopeDataPoints;
 		::Audio::TCircularAudioBuffer<FDataPoint> DistanceAttenuationDataPoints;
+		::Audio::TCircularAudioBuffer<FDataPoint> RelativeRenderCostDataPoints;
 	};
 
 	class FMixerSourceMessages
@@ -125,6 +127,7 @@ namespace UE::Audio::Insights
 		TAnalyzerMessageQueue<FMixerSourceStartMessage> StartMessages { 0.1 };
 		TAnalyzerMessageQueue<FMixerSourceStopMessage> StopMessages { 0.1 };
 		TAnalyzerMessageQueue<FMixerSourceVolumeMessage> VolumeMessages { 0.1 };
+		TAnalyzerMessageQueue<FMixerSourceRelativeRenderCostMessage> RelativeRenderCostMessages { 0.1 };
 	};
 
 #if !WITH_EDITOR
@@ -139,6 +142,7 @@ namespace UE::Audio::Insights
 			, StartCachedMessages(InSession.GetLinearAllocator(), 16384)
 			, StopCachedMessages(InSession.GetLinearAllocator(), 4096)
 			, VolumeCachedMessages(InSession.GetLinearAllocator(), 16384)
+			, RelativeRenderCostCachedMessages(InSession.GetLinearAllocator(), 16384)
 		{
 
 		}
@@ -151,6 +155,7 @@ namespace UE::Audio::Insights
 		TraceServices::TPagedArray<FMixerSourceStartMessage> StartCachedMessages;
 		TraceServices::TPagedArray<FMixerSourceStopMessage> StopCachedMessages;
 		TraceServices::TPagedArray<FMixerSourceVolumeMessage> VolumeCachedMessages;
+		TraceServices::TPagedArray<FMixerSourceRelativeRenderCostMessage> RelativeRenderCostCachedMessages;
 
 		TMap<uint32, FMixerSourceStartMessage> PlayOrderToStartMessageMap;
 	};
diff --git a/Engine/Plugins/AudioInsights/Source/AudioInsights/Public/Messages/SoundTraceMessages.h b/Engine/Plugins/AudioInsights/Source/AudioInsights/Public/Messages/SoundTraceMessages.h
--- a/Engine/Plugins/AudioInsights/Source/AudioInsights/Public/Messages/SoundTraceMessages.h
+++ b/Engine/Plugins/AudioInsights/Source/AudioInsights/Public/Messages/SoundTraceMessages.h
@@ -101,6 +101,7 @@ namespace UE::Audio::Insights
 		TAnalyzerMessageQueue<FMixerSourceEnvelopeMessage> AmplitudeMessages { 0.1 };
 		TAnalyzerMessageQueue<FMixerSourceVolumeMessage> VolumeMessages { 0.1 };
 		TAnalyzerMessageQueue<FMixerSourcePitchMessage> PitchMessages { 0.1 };
+		TAnalyzerMessageQueue<FMixerSourceRelativeRenderCostMessage> RelativeRenderCostMessages{ 0.1 };
 		TAnalyzerMessageQueue<FSoundStopMessage> StopMessages { 0.1 };
 	};
 
@@ -115,6 +116,7 @@ namespace UE::Audio::Insights
 			, AmplitudeCachedMessages(InSession.GetLinearAllocator(), 16384)
 			, VolumeCachedMessages(InSession.GetLinearAllocator(), 16384)
 			, PitchCachedMessages(InSession.GetLinearAllocator(), 16384)
+			, RelativeRenderCostCachedMessages(InSession.GetLinearAllocator(), 16384)
 			, StopCachedMessages(InSession.GetLinearAllocator(), 4096)
 		{
 
@@ -127,6 +129,7 @@ namespace UE::Audio::Insights
 		TraceServices::TPagedArray<FMixerSourceEnvelopeMessage> AmplitudeCachedMessages;
 		TraceServices::TPagedArray<FMixerSourceVolumeMessage> VolumeCachedMessages;
 		TraceServices::TPagedArray<FMixerSourcePitchMessage> PitchCachedMessages;
+		TraceServices::TPagedArray<FMixerSourceRelativeRenderCostMessage> RelativeRenderCostCachedMessages;
 		TraceServices::TPagedArray<FSoundStopMessage> StopCachedMessages;
 	};
 #endif // !WITH_EDITOR
@@ -178,5 +181,6 @@ namespace UE::Audio::Insights
 		::Audio::TCircularAudioBuffer<FDataPoint> AmplitudeDataPoints;
 		::Audio::TCircularAudioBuffer<FDataPoint> VolumeDataPoints;
 		::Audio::TCircularAudioBuffer<FDataPoint> PitchDataPoints;
+		::Audio::TCircularAudioBuffer<FDataPoint> RelativeRenderCostDataPoints;
 	};
 } // namespace UE::Audio::Insights
diff --git a/Engine/Source/Runtime/AudioMixer/Private/AudioMixerSource.cpp b/Engine/Source/Runtime/AudioMixer/Private/AudioMixerSource.cpp
--- a/Engine/Source/Runtime/AudioMixer/Private/AudioMixerSource.cpp
+++ b/Engine/Source/Runtime/AudioMixer/Private/AudioMixerSource.cpp
@@ -38,6 +38,14 @@ UE_TRACE_EVENT_BEGIN(Audio, MixerSourceStop)
 	UE_TRACE_EVENT_FIELD(uint64, Timestamp)
 	UE_TRACE_EVENT_FIELD(uint32, PlayOrder)
 UE_TRACE_EVENT_END()
+
+UE_TRACE_EVENT_BEGIN(Audio, MixerSourceRelativeRenderCost)
+	UE_TRACE_EVENT_FIELD(uint32, DeviceId)
+	UE_TRACE_EVENT_FIELD(uint64, Timestamp)
+	UE_TRACE_EVENT_FIELD(uint32, PlayOrder)
+	UE_TRACE_EVENT_FIELD(uint32, ActiveSoundPlayOrder)
+	UE_TRACE_EVENT_FIELD(float, RelativeRenderCost)
+UE_TRACE_EVENT_END()
 #endif // UE_AUDIO_PROFILERTRACE_ENABLED
 
 
@@ -1818,6 +1826,19 @@ namespace Audio
 				DebugInfo->RelativeRenderCost = RelativeRenderCost;
 			}
 #endif // if ENABLE_AUDIO_DEBUG
+
+#if UE_AUDIO_PROFILERTRACE_ENABLED
+			const bool bChannelEnabled = UE_TRACE_CHANNELEXPR_IS_ENABLED(AudioMixerChannel);
+			if (bChannelEnabled && WaveInstance->ActiveSound != nullptr)
+			{
+				UE_TRACE_LOG(Audio, MixerSourceRelativeRenderCost, AudioMixerChannel)
+					<< MixerSourceRelativeRenderCost.DeviceId(MixerDevice->DeviceID)
+					<< MixerSourceRelativeRenderCost.Timestamp(FPlatformTime::Cycles64())
+					<< MixerSourceRelativeRenderCost.PlayOrder(WaveInstance->GetPlayOrder())
+					<< MixerSourceRelativeRenderCost.ActiveSoundPlayOrder(WaveInstance->ActiveSound->GetPlayOrder())
+					<< MixerSourceRelativeRenderCost.RelativeRenderCost(RelativeRenderCost);
+			}
+#endif // UE_AUDIO_PROFILERTRACE_ENABLED
 		}
 	}
 
