[Audio Insights] Add new columns for Category and Instance Owner

diff --git a/Engine/Plugins/AudioInsights/Source/AudioInsights/Private/Messages/SoundTraceMessages.cpp b/Engine/Plugins/AudioInsights/Source/AudioInsights/Private/Messages/SoundTraceMessages.cpp
--- a/Engine/Plugins/AudioInsights/Source/AudioInsights/Private/Messages/SoundTraceMessages.cpp
+++ b/Engine/Plugins/AudioInsights/Source/AudioInsights/Private/Messages/SoundTraceMessages.cpp
@@ -55,6 +55,10 @@ namespace UE::Audio::Insights
 		{
 			EntryType = ESoundDashboardEntryType::None;
 		}
+
+#if WITH_EDITOR
+		EventData.GetString("OwnerActorName", OwnerActorName);
+#endif // WITH_EDITOR
 	}
 
 	FSoundDashboardEntry::FSoundDashboardEntry()
diff --git a/Engine/Plugins/AudioInsights/Source/AudioInsights/Private/Providers/SoundTraceProvider.cpp b/Engine/Plugins/AudioInsights/Source/AudioInsights/Private/Providers/SoundTraceProvider.cpp
--- a/Engine/Plugins/AudioInsights/Source/AudioInsights/Private/Providers/SoundTraceProvider.cpp
+++ b/Engine/Plugins/AudioInsights/Source/AudioInsights/Private/Providers/SoundTraceProvider.cpp
@@ -368,7 +368,12 @@ namespace UE::Audio::Insights
 			FSoundDashboardEntry& ActiveSoundEntryRef = *OutActiveSoundEntry->Get();
 			ActiveSoundEntryRef.Name = *Msg.Name;
 			ActiveSoundEntryRef.EntryType = Msg.EntryType;
+			ActiveSoundEntryRef.CategoryName = FText::FromString(FSoundTraceProviderPrivate::GetEntryTypeDisplayName(Msg.EntryType));
 			ActiveSoundEntryRef.bShouldForceExpandChildren = false;
+
+#if WITH_EDITOR
+			ActiveSoundEntryRef.OwnerActorName = FText::FromString(*Msg.OwnerActorName);
+#endif // WITH_EDITOR
 		});
 
 		ProcessMessageQueue<FSoundWaveStartMessage>(TraceMessages.SoundWaveStartMessages, GetActiveSoundEntryFromSoundWaveStartMessage,
@@ -397,9 +402,14 @@ namespace UE::Audio::Insights
 			SoundWaveEntry->Timestamp = Msg.Timestamp;
 			SoundWaveEntry->Name = *Msg.Name;
 			SoundWaveEntry->EntryType = Msg.EntryType;
+			SoundWaveEntry->CategoryName = FText::FromString(FSoundTraceProviderPrivate::GetEntryTypeDisplayName(Msg.EntryType));
 			SoundWaveEntry->PinnedEntryType = (*OutActiveSoundEntry)->PinnedEntryType;
 			SoundWaveEntry->bShouldForceExpandChildren = false;
 
+#if WITH_EDITOR
+			SoundWaveEntry->OwnerActorName = FText::FromString(*Msg.OwnerActorName);
+#endif // WITH_EDITOR
+
 			(*OutActiveSoundEntry)->Children.Add(SoundWaveEntry);
 		});
 
diff --git a/Engine/Plugins/AudioInsights/Source/AudioInsights/Private/Views/SoundDashboardViewFactory.cpp b/Engine/Plugins/AudioInsights/Source/AudioInsights/Private/Views/SoundDashboardViewFactory.cpp
--- a/Engine/Plugins/AudioInsights/Source/AudioInsights/Private/Views/SoundDashboardViewFactory.cpp
+++ b/Engine/Plugins/AudioInsights/Source/AudioInsights/Private/Views/SoundDashboardViewFactory.cpp
@@ -1343,6 +1343,30 @@ namespace UE::Audio::Insights
 					.FillWidth = 0.1f,
 					.Alignment = EHorizontalAlignment::HAlign_Left
 				}
+			},
+#if WITH_EDITOR
+			{
+				"Actor",
+				{
+					.DisplayName = LOCTEXT("SoundDashboard_ActorColumnDisplayName", "Actor"),
+					.IconName = NAME_None,
+					.bShowDisplayName = true,
+					.bDefaultHidden = true,
+					.FillWidth = 0.1f,
+					.Alignment = EHorizontalAlignment::HAlign_Left
+				}
+			},
+#endif // WITH_EDITOR
+			{
+				"Category",
+				{
+					.DisplayName = LOCTEXT("SoundDashboard_CategoryColumnDisplayName", "Category"),
+					.IconName = NAME_None,
+					.bShowDisplayName = true,
+					.bDefaultHidden = true,
+					.FillWidth = 0.1f,
+					.Alignment = EHorizontalAlignment::HAlign_Left
+				}
 			}
 		};
 
@@ -1507,6 +1531,30 @@ namespace UE::Audio::Insights
 						return FText::AsNumber(GetLastEntryArrayValue(SoundDashboardEntry.RelativeRenderCostDataPoints), FSlateStyle::Get().GetRelativeRenderCostFormat());
 					}
 				}
+			},
+#if WITH_EDITOR
+			{
+				"Actor",
+				{
+					.GetDisplayValue = [](const IDashboardDataTreeViewEntry& InData)
+					{
+						const FSoundDashboardEntry& SoundDashboardEntry = CastEntry(InData);
+
+						return SoundDashboardEntry.OwnerActorName;
+					}
+				}
+			},
+#endif // WITH_EDITOR
+			{
+				"Category",
+				{
+					.GetDisplayValue = [](const IDashboardDataTreeViewEntry& InData)
+					{
+						const FSoundDashboardEntry& SoundDashboardEntry = CastEntry(InData);
+
+						return SoundDashboardEntry.CategoryName;
+					}
+				}
 			}
 		};
 
@@ -1601,6 +1649,30 @@ namespace UE::Audio::Insights
 			return ComparisonDiff < 0.0f;
 		};
 
+#if WITH_EDITOR
+		auto SortByActor = [&SortByPlayOrder](const FSoundDashboardEntry& First, const FSoundDashboardEntry& Second)
+		{
+			const int32 Comparison = First.OwnerActorName.CompareToCaseIgnored(Second.OwnerActorName);
+			if (Comparison == 0)
+			{
+				return SortByPlayOrder(First, Second);
+			}
+			
+			return Comparison < 0;
+		};
+#endif // WITH_EDITOR
+
+		auto SortByCategory = [&SortByPlayOrder](const FSoundDashboardEntry& First, const FSoundDashboardEntry& Second)
+		{
+			const int32 Comparison = First.CategoryName.CompareToCaseIgnored(Second.CategoryName);
+			if (Comparison == 0)
+			{
+				return SortByPlayOrder(First, Second);
+			}
+			
+			return Comparison < 0;
+		};
+
 		if (SortByColumn == "Name")
 		{
 			SortByPredicate(SortByName);
@@ -1629,6 +1701,16 @@ namespace UE::Audio::Insights
 		{
 			SortByPredicate(SortByRelativeRenderCost);
 		}
+#if WITH_EDITOR
+		else if (SortByColumn == "Actor")
+		{
+			SortByPredicate(SortByActor);
+		}
+#endif // WITH_EDITOR
+		else if (SortByColumn == "Category")
+		{
+			SortByPredicate(SortByCategory);
+		}
 
 		FullTree.Reset();
 
diff --git a/Engine/Plugins/AudioInsights/Source/AudioInsights/Public/Messages/SoundTraceMessages.h b/Engine/Plugins/AudioInsights/Source/AudioInsights/Public/Messages/SoundTraceMessages.h
--- a/Engine/Plugins/AudioInsights/Source/AudioInsights/Public/Messages/SoundTraceMessages.h
+++ b/Engine/Plugins/AudioInsights/Source/AudioInsights/Public/Messages/SoundTraceMessages.h
@@ -47,6 +47,10 @@ namespace UE::Audio::Insights
 
 		FString Name;
 		ESoundDashboardEntryType EntryType;
+
+#if WITH_EDITOR
+		FString OwnerActorName;
+#endif // WITH_EDITOR
 	};
 
 	struct FSoundWaveStartMessage : public FSoundStartMessage
@@ -170,6 +174,7 @@ namespace UE::Audio::Insights
 		double Timestamp = 0.0;
 		double TimeoutTimestamp = INVALID_TIMEOUT;
 		FString Name;
+		FText CategoryName;
 		ESoundDashboardEntryType EntryType = ESoundDashboardEntryType::None;
 		EPinnedEntryType PinnedEntryType = EPinnedEntryType::None;
 		bool bIsCategory = false;
@@ -182,5 +187,9 @@ namespace UE::Audio::Insights
 		::Audio::TCircularAudioBuffer<FDataPoint> VolumeDataPoints;
 		::Audio::TCircularAudioBuffer<FDataPoint> PitchDataPoints;
 		::Audio::TCircularAudioBuffer<FDataPoint> RelativeRenderCostDataPoints;
+
+#if WITH_EDITOR
+		FText OwnerActorName;
+#endif // WITH_EDITOR
 	};
 } // namespace UE::Audio::Insights
diff --git a/Engine/Source/Runtime/Engine/Private/AudioDevice.cpp b/Engine/Source/Runtime/Engine/Private/AudioDevice.cpp
--- a/Engine/Source/Runtime/Engine/Private/AudioDevice.cpp
+++ b/Engine/Source/Runtime/Engine/Private/AudioDevice.cpp
@@ -245,6 +245,7 @@ UE_TRACE_EVENT_BEGIN(Audio, SoundStart)
 	UE_TRACE_EVENT_FIELD(uint32, PlayOrder)
 	UE_TRACE_EVENT_FIELD(UE::Trace::WideString, Name)
 	UE_TRACE_EVENT_FIELD(UE::Trace::WideString, SoundClassName)
+	UE_TRACE_EVENT_FIELD(UE::Trace::WideString, OwnerActorName)
 UE_TRACE_EVENT_END()
 
 UE_TRACE_EVENT_BEGIN(Audio, SoundStop)
@@ -5082,12 +5083,21 @@ void FAudioDevice::SendSoundStartedTraceMessage(const FActiveSound& InActiveSoun
 
 		const FString SoundBaseAssetName = AudioProfilerTrace::GetSoundBaseAssetName(SoundObjectPtr->GetClass());
 
+		FString OwnerActorName;
+#if WITH_EDITOR
+		if (const TObjectPtr<UAudioComponent> AudioComponent = UAudioComponent::GetAudioComponentFromID(InActiveSound.GetAudioComponentID()))
+		{
+			OwnerActorName = AudioComponent->GetOwner()->GetActorLabel(false);
+		}
+#endif // WITH_EDITOR
+
 		UE_TRACE_LOG(Audio, SoundStart, AudioMixerChannel)
 			<< SoundStart.DeviceId(InActiveSound.AudioDevice->DeviceID)
 			<< SoundStart.Timestamp(FPlatformTime::Cycles64())
 			<< SoundStart.PlayOrder(InActiveSound.PlayOrder)
 			<< SoundStart.Name(*SoundObjectPtr->GetPathName())
-			<< SoundStart.SoundClassName(*SoundBaseAssetName);
+			<< SoundStart.SoundClassName(*SoundBaseAssetName)
+			<< SoundStart.OwnerActorName(*OwnerActorName);
 	}
 }
 #endif // UE_AUDIO_PROFILERTRACE_ENABLED
diff --git a/Engine/Source/Runtime/Engine/Private/SoundWave.cpp b/Engine/Source/Runtime/Engine/Private/SoundWave.cpp
--- a/Engine/Source/Runtime/Engine/Private/SoundWave.cpp
+++ b/Engine/Source/Runtime/Engine/Private/SoundWave.cpp
@@ -116,6 +116,7 @@ UE_TRACE_EVENT_BEGIN(Audio, SoundWaveStart)
 	UE_TRACE_EVENT_FIELD(uint32, ActiveSoundPlayOrder)
 	UE_TRACE_EVENT_FIELD(UE::Trace::WideString, Name)
 	UE_TRACE_EVENT_FIELD(UE::Trace::WideString, SoundClassName)
+	UE_TRACE_EVENT_FIELD(UE::Trace::WideString, OwnerActorName)
 UE_TRACE_EVENT_END()
 
 UE_TRACE_EVENT_BEGIN(Audio, SoundPriority)
@@ -3879,13 +3880,22 @@ void USoundWave::Parse(FAudioDevice* AudioDevice, const UPTRINT NodeWaveInstance
 			{
 				const FString SoundWaveBaseAssetName = AudioProfilerTrace::GetSoundBaseAssetName(SoundWaveObjectPtr->GetClass());
 
+				FString OwnerActorName;
+#if WITH_EDITOR
+				if (const TObjectPtr<UAudioComponent> AudioComponent = UAudioComponent::GetAudioComponentFromID(ActiveSound.GetAudioComponentID()))
+				{
+					OwnerActorName = AudioComponent->GetOwner()->GetActorLabel(false);
+				}
+#endif // WITH_EDITOR
+
 				UE_TRACE_LOG(Audio, SoundWaveStart, AudioMixerChannel)
 					<< SoundWaveStart.DeviceId(AudioDevice->DeviceID)
 					<< SoundWaveStart.Timestamp(FPlatformTime::Cycles64())
 					<< SoundWaveStart.PlayOrder(WaveInstance->GetPlayOrder())
 					<< SoundWaveStart.ActiveSoundPlayOrder(ActiveSound.GetPlayOrder())
 					<< SoundWaveStart.Name(*SoundWaveObjectPtr->GetPathName())
-					<< SoundWaveStart.SoundClassName(*SoundWaveBaseAssetName);
+					<< SoundWaveStart.SoundClassName(*SoundWaveBaseAssetName)
+					<< SoundWaveStart.OwnerActorName(*OwnerActorName);
 			}
 		}
 
