[Audio Insights] Sound dashboard: More optimizations caching entry names in different formats + set column visibility based on row's visibility + only run slate lambdas where really needed + don't retrieve audio device manager and audio debugger on each recursive function call

diff --git a/Engine/Plugins/AudioInsights/Source/AudioInsights/Private/Providers/SoundTraceProvider.cpp b/Engine/Plugins/AudioInsights/Source/AudioInsights/Private/Providers/SoundTraceProvider.cpp
--- a/Engine/Plugins/AudioInsights/Source/AudioInsights/Private/Providers/SoundTraceProvider.cpp
+++ b/Engine/Plugins/AudioInsights/Source/AudioInsights/Private/Providers/SoundTraceProvider.cpp
@@ -366,7 +366,7 @@ namespace UE::Audio::Insights
 #endif // !WITH_EDITOR
 
 			FSoundDashboardEntry& ActiveSoundEntryRef = *OutActiveSoundEntry->Get();
-			ActiveSoundEntryRef.Name = *Msg.Name;
+			ActiveSoundEntryRef.SetName(*Msg.Name);
 			ActiveSoundEntryRef.EntryType = Msg.EntryType;
 			ActiveSoundEntryRef.CategoryName = FText::FromString(FSoundTraceProviderPrivate::GetEntryTypeDisplayName(Msg.EntryType));
 			ActiveSoundEntryRef.bShouldForceExpandChildren = false;
@@ -400,7 +400,7 @@ namespace UE::Audio::Insights
 			SoundWaveEntry->DeviceId = Msg.DeviceId;
 			SoundWaveEntry->PlayOrder = Msg.PlayOrder;
 			SoundWaveEntry->Timestamp = Msg.Timestamp;
-			SoundWaveEntry->Name = *Msg.Name;
+			SoundWaveEntry->SetName(*Msg.Name);
 			SoundWaveEntry->EntryType = Msg.EntryType;
 			SoundWaveEntry->CategoryName = FText::FromString(FSoundTraceProviderPrivate::GetEntryTypeDisplayName(Msg.EntryType));
 			SoundWaveEntry->PinnedEntryType = (*OutActiveSoundEntry)->PinnedEntryType;
@@ -629,7 +629,7 @@ namespace UE::Audio::Insights
 				CategoryEntry = MakeShared<FSoundDashboardEntry>();
 				CategoryEntry->DeviceId = Msg.DeviceId;
 				CategoryEntry->EntryType = Msg.EntryType;
-				CategoryEntry->Name = FSoundTraceProviderPrivate::GetEntryTypeDisplayName(Msg.EntryType);
+				CategoryEntry->SetName(FSoundTraceProviderPrivate::GetEntryTypeDisplayName(Msg.EntryType));
 				CategoryEntry->bIsCategory = true;
 			}
 
@@ -892,7 +892,7 @@ namespace UE::Audio::Insights
 			}
 
 			FSoundDashboardEntry& EntryRef = *ActiveSoundEntry->Get();
-			EntryRef.Name = *StartCachedMessage.Name;
+			EntryRef.SetName(*StartCachedMessage.Name);
 			EntryRef.EntryType = StartCachedMessage.EntryType;
 		}
 
@@ -922,7 +922,7 @@ namespace UE::Audio::Insights
 			ChildEntry->DeviceId = SoundWaveStartCachedMessage.DeviceId;
 			ChildEntry->PlayOrder = SoundWaveStartCachedMessage.PlayOrder;
 			ChildEntry->Timestamp = SoundWaveStartCachedMessage.Timestamp;
-			ChildEntry->Name = *SoundWaveStartCachedMessage.Name;
+			ChildEntry->SetName(*SoundWaveStartCachedMessage.Name);
 			ChildEntry->EntryType = SoundWaveStartCachedMessage.EntryType;
 
 			ParentSoundEntryReturn->Children.Add(ChildEntry);
diff --git a/Engine/Plugins/AudioInsights/Source/AudioInsights/Private/Views/SoundDashboardViewFactory.cpp b/Engine/Plugins/AudioInsights/Source/AudioInsights/Private/Views/SoundDashboardViewFactory.cpp
--- a/Engine/Plugins/AudioInsights/Source/AudioInsights/Private/Views/SoundDashboardViewFactory.cpp
+++ b/Engine/Plugins/AudioInsights/Source/AudioInsights/Private/Views/SoundDashboardViewFactory.cpp
@@ -50,7 +50,7 @@ namespace UE::Audio::Insights
 		{
 			FSoundDashboardEntry& SoundEntry = FSoundDashboardViewFactoryPrivate::CastEntry(InEntry);
 
-			bool bEntryMatchesTextFilter = SoundEntry.GetDisplayName().ToString().Contains(InFilterString);
+			bool bEntryMatchesTextFilter = SoundEntry.GetDisplayNameStr().Contains(InFilterString);
 
 			if (bEntryMatchesTextFilter)
 			{
@@ -214,8 +214,6 @@ namespace UE::Audio::Insights
 					return;
 				}
 
-				const FName SoundAssetDisplayName(SoundEntry.GetDisplayName().ToString());
-
 				const bool bIsSoundCueType = SoundEntry.EntryType == ESoundDashboardEntryType::SoundCue || SoundEntry.EntryType == ESoundDashboardEntryType::SoundCueTemplate;
 
 				if (!bIsSoundCueType)
@@ -223,11 +221,11 @@ namespace UE::Audio::Insights
 					switch (InMuteSoloMode)
 					{
 						case EMuteSoloMode::Mute:
-							AudioDebugger.SetMuteSoundWave(SoundAssetDisplayName, bInOnOff);
+							AudioDebugger.SetMuteSoundWave(SoundEntry.GetDisplayFName(), bInOnOff);
 							break;
 
 						case EMuteSoloMode::Solo:
-							AudioDebugger.SetSoloSoundWave(SoundAssetDisplayName, bInOnOff);
+							AudioDebugger.SetSoloSoundWave(SoundEntry.GetDisplayFName(), bInOnOff);
 							break;
 
 						default:
@@ -255,8 +253,6 @@ namespace UE::Audio::Insights
 
 				const FSoundDashboardEntry& SoundEntry = FSoundDashboardViewFactoryPrivate::CastEntry(InEntry);
 
-				const FName SoundAssetDisplayName(SoundEntry.GetDisplayName().ToString());
-
 				const bool bIsSoundCueType = SoundEntry.EntryType == ESoundDashboardEntryType::SoundCue || SoundEntry.EntryType == ESoundDashboardEntryType::SoundCueTemplate;
 
 				if (!bIsSoundCueType)
@@ -264,11 +260,11 @@ namespace UE::Audio::Insights
 					switch (InMuteSoloMode)
 					{
 						case EMuteSoloMode::Mute:
-							AudioDebugger.ToggleMuteSoundWave(SoundAssetDisplayName);
+							AudioDebugger.ToggleMuteSoundWave(SoundEntry.GetDisplayFName());
 							break;
 
 						case EMuteSoloMode::Solo:
-							AudioDebugger.ToggleSoloSoundWave(SoundAssetDisplayName);
+							AudioDebugger.ToggleSoloSoundWave(SoundEntry.GetDisplayFName());
 							break;
 
 						default:
@@ -287,70 +283,63 @@ namespace UE::Audio::Insights
 #endif // ENABLE_AUDIO_DEBUG
 		}
 
-		bool IsMuteSolo(const IDashboardDataTreeViewEntry& InEntry, const bool bInCheckChildren, const EMuteSoloMode InMuteSoloMode)
+		bool IsMuteSolo(::Audio::FAudioDebugger& InAudioDebugger, const IDashboardDataTreeViewEntry& InEntry, const bool bInCheckChildren, const EMuteSoloMode InMuteSoloMode)
 		{
 #if ENABLE_AUDIO_DEBUG
-			if (FAudioDeviceManager* AudioDeviceManager = FAudioDeviceManager::Get())
-			{
-				::Audio::FAudioDebugger& AudioDebugger = AudioDeviceManager->GetDebugger();
-
-				const FSoundDashboardEntry& SoundEntry = FSoundDashboardViewFactoryPrivate::CastEntry(InEntry);
+			const FSoundDashboardEntry& SoundEntry = FSoundDashboardViewFactoryPrivate::CastEntry(InEntry);
 
-				// Treat hidden original entries as muted/soloed to ensure the parent category reflects the correct state
-				if (SoundEntry.PinnedEntryType == FSoundDashboardEntry::EPinnedEntryType::HiddenOriginalEntry)
-				{
-					return true;
-				}
-
-				const FName SoundAssetDisplayName(SoundEntry.GetDisplayName().ToString());
+			// Treat hidden original entries as muted/soloed to ensure the parent category reflects the correct state
+			if (SoundEntry.PinnedEntryType == FSoundDashboardEntry::EPinnedEntryType::HiddenOriginalEntry)
+			{
+				return true;
+			}
 
-				const bool bIsSoundCueType = SoundEntry.EntryType == ESoundDashboardEntryType::SoundCue || SoundEntry.EntryType == ESoundDashboardEntryType::SoundCueTemplate;
+			const bool bIsSoundCueType = SoundEntry.EntryType == ESoundDashboardEntryType::SoundCue || SoundEntry.EntryType == ESoundDashboardEntryType::SoundCueTemplate;
 
-				if (!bIsSoundCueType)
+			if (!bIsSoundCueType)
+			{
+				switch (InMuteSoloMode)
 				{
-					switch (InMuteSoloMode)
+					case EMuteSoloMode::Mute:
 					{
-						case EMuteSoloMode::Mute:
+						if (InAudioDebugger.IsMuteSoundWave(SoundEntry.GetDisplayFName()))
 						{
-							if (AudioDebugger.IsMuteSoundWave(SoundAssetDisplayName))
-							{
-								return true;
-							}
-
-							break;
+							return true;
 						}
 
-						case EMuteSoloMode::Solo:
-						{
-							if (AudioDebugger.IsSoloSoundWave(SoundAssetDisplayName))
-							{
-								return true;
-							}
+						break;
+					}
 
-							break;
+					case EMuteSoloMode::Solo:
+					{
+						if (InAudioDebugger.IsSoloSoundWave(SoundEntry.GetDisplayFName()))
+						{
+							return true;
 						}
 
-						default:
-							break;
+						break;
 					}
+
+					default:
+						break;
 				}
+			}
 
-				if (bInCheckChildren)
-				{
-					uint32 NumChildrenMuteSolo = 0;
+			if (bInCheckChildren)
+			{
+				uint32 NumChildrenMuteSolo = 0;
 
-					for (const TSharedPtr<IDashboardDataTreeViewEntry>& SoundEntryChild : SoundEntry.Children)
+				for (const TSharedPtr<IDashboardDataTreeViewEntry>& SoundEntryChild : SoundEntry.Children)
+				{
+					if (SoundEntryChild.IsValid() && IsMuteSolo(InAudioDebugger, *SoundEntryChild, true /*bInCheckChildren*/, InMuteSoloMode))
 					{
-						if (SoundEntryChild.IsValid() && IsMuteSolo(*SoundEntryChild, true /*bInCheckChildren*/, InMuteSoloMode))
-						{
-							++NumChildrenMuteSolo;
-						}
+						++NumChildrenMuteSolo;
 					}
+				}
 
-					const bool bAllChildrenMuteSolo = !SoundEntry.Children.IsEmpty() && NumChildrenMuteSolo == SoundEntry.Children.Num();
+				const bool bAllChildrenMuteSolo = !SoundEntry.Children.IsEmpty() && NumChildrenMuteSolo == SoundEntry.Children.Num();
 
-					return bAllChildrenMuteSolo;
-				}
+				return bAllChildrenMuteSolo;
 			}
 #endif // ENABLE_AUDIO_DEBUG
 
@@ -925,6 +914,10 @@ namespace UE::Audio::Insights
 			};
 
 			return SNew(SHorizontalBox)
+				.Visibility_Lambda([InRowWidget]()
+				{
+					return InRowWidget->GetVisibility();
+				})
 				// Tree expander arrow
 				+ SHorizontalBox::Slot()
 				.AutoWidth()
@@ -990,30 +983,35 @@ namespace UE::Audio::Insights
 #if WITH_EDITOR
 		if (InColumn == "Mute")
 		{
-			return CreateMuteSoloButton(InRowData, InColumn, 
+			return CreateMuteSoloButton(InRowWidget, InRowData, InColumn,
 				[this](const TArray<TSharedPtr<IDashboardDataTreeViewEntry>>& InEntries)
 				{
 					ToggleMuteSoloEntries(InEntries, EMuteSoloMode::Mute);
 				},
 				[](const IDashboardDataTreeViewEntry& InEntry, const bool bInCheckChildren)
 				{
-					return IsMuteSolo(InEntry, bInCheckChildren, EMuteSoloMode::Mute);
+					FAudioDeviceManager* AudioDeviceManager = FAudioDeviceManager::Get();
+
+					return AudioDeviceManager ? IsMuteSolo(AudioDeviceManager->GetDebugger(), InEntry, bInCheckChildren, EMuteSoloMode::Mute) : false;
 				});
 		}
 		else if (InColumn == "Solo")
 		{
-			return CreateMuteSoloButton(InRowData, InColumn, 
+			return CreateMuteSoloButton(InRowWidget, InRowData, InColumn,
 				[this](const TArray<TSharedPtr<IDashboardDataTreeViewEntry>>& InEntries)
 				{
 					ToggleMuteSoloEntries(InEntries, EMuteSoloMode::Solo);
 				}, 
 				[](const IDashboardDataTreeViewEntry& InEntry, const bool bInCheckChildren)
 				{
-					return IsMuteSolo(InEntry, bInCheckChildren, EMuteSoloMode::Solo);
+					FAudioDeviceManager* AudioDeviceManager = FAudioDeviceManager::Get();
+
+					return AudioDeviceManager ? IsMuteSolo(AudioDeviceManager->GetDebugger(), InEntry, bInCheckChildren, EMuteSoloMode::Solo) : false;
 				});
 		}
 		else
 #endif // WITH_EDITOR
+		if (InColumn == "Name")
 		{
 			if (IsRootItem(InRowData))
 			{
@@ -1021,60 +1019,38 @@ namespace UE::Audio::Insights
 			}
 
 			return SNew(SHorizontalBox)
-				// Tree expander arrow (only for Name column)
+				.Visibility_Lambda([InRowWidget]()
+				{
+					return InRowWidget->GetVisibility();
+				})
+				// Tree expander arrow
 				+ SHorizontalBox::Slot()
 				.AutoWidth()
 				[
-					InColumn == "Name" ? SNew(SExpanderArrow, InRowWidget) : SNullWidget::NullWidget
+					SNew(SExpanderArrow, InRowWidget)
 				]
-				// Icon (optional)
+				// Icon
 				+ SHorizontalBox::Slot()
 				.Padding(2.0f, 2.0f, 2.0f, 2.0f)
 				.AutoWidth()
 				[
 					SNew(SImage)
-					.Image_Lambda([this, InRowData, InColumn]()
-					{
-						const FColumnData& ColumnData = GetColumns()[InColumn];
-						const FName IconName = ColumnData.GetIconName.IsSet() ? ColumnData.GetIconName(InRowData.Get()) : NAME_None;
-						
-						return IconName != NAME_None ? FSlateStyle::Get().GetBrush(IconName) : nullptr;
-					})
-					.ColorAndOpacity_Lambda([this, InRowData, InColumn]()
-					{
-						const FColumnData& ColumnData = GetColumns()[InColumn];
-						const FSlateColor TextColor = ColumnData.GetTextColorValue.IsSet() ? ColumnData.GetTextColorValue(InRowData.Get()) : FSlateColor::UseForeground();
-						
-						return TextColor;
-					})
+					.Image(FSlateStyle::Get().GetBrush(ColumnData.GetIconName(InRowData.Get())))
 				]
-				// Text (optional)
+				// Text
 				+ SHorizontalBox::Slot()
 				.Padding(10.0f, 2.0f, 2.0f, 2.0f)
 				.AutoWidth()
 				[
 					SNew(STextBlock)
-					.Text_Lambda([this, InRowData, InColumn]()
-					{
-						const FColumnData& ColumnData = GetColumns()[InColumn];
-						const FText ValueText = ColumnData.GetDisplayValue.IsSet() ? ColumnData.GetDisplayValue(InRowData.Get()) : FText::GetEmpty();
-
-						return ValueText;
-					})
-					.ColorAndOpacity_Lambda([this, InRowData, InColumn]()
-					{
-						const FColumnData& ColumnData = GetColumns()[InColumn];
-						const FSlateColor TextColor = ColumnData.GetTextColorValue.IsSet() ? ColumnData.GetTextColorValue(InRowData.Get()) : FSlateColor::UseForeground();
-
-						return TextColor;
-					})
+					.Text(ValueText)
 				]
 				// Number of children text (if it is not leaf item)
 				+ SHorizontalBox::Slot()
 				.Padding(6.0f, 2.0f, 0.0f, 2.0f)
 				.AutoWidth()
 				[
-					InColumn == "Name" && !InRowData->Children.IsEmpty()
+					!InRowData->Children.IsEmpty()
 						? SNew(STextBlock)
 						  .ColorAndOpacity(FSlateColor(FColor::White.WithAlpha(128)))
 						  .Text_Lambda([this, InRowData]()
@@ -1085,17 +1061,43 @@ namespace UE::Audio::Insights
 						: SNullWidget::NullWidget
 				];
 		}
+		else
+		{
+			if (IsRootItem(InRowData))
+			{
+				return SNullWidget::NullWidget;
+			}
+
+			return SNew(SHorizontalBox)
+				.Visibility_Lambda([InRowWidget]()
+				{
+					return InRowWidget->GetVisibility();
+				})
+				// Text
+				+ SHorizontalBox::Slot()
+				.Padding(10.0f, 2.0f, 2.0f, 2.0f)
+				.AutoWidth()
+				[
+					SNew(STextBlock)
+					.Text_Lambda([InRowData, ColumnData]()
+					{
+						return ColumnData.GetDisplayValue(InRowData.Get());
+					})
+				];
+		}
 	}
 
 	TSharedRef<ITableRow> FSoundDashboardViewFactory::OnGenerateRow(TSharedPtr<IDashboardDataTreeViewEntry> Item, const TSharedRef<STableViewBase>& OwnerTable)
 	{
 		using namespace FSoundDashboardViewFactoryPrivate;
 
-		return SNew(SRowWidget, OwnerTable, Item, AsShared())
+		const FSoundDashboardEntry& SoundDashboardEntry = CastEntry(*Item);
+
+		return SoundDashboardEntry.bIsVisible ?
+			SNew(SRowWidget, OwnerTable, Item, AsShared())
 			.RenderOpacity_Lambda([Item]
 			{
 				const FSoundDashboardEntry& SoundDashboardEntry = CastEntry(*Item);
-
 				return SoundDashboardEntry.TimeoutTimestamp == INVALID_TIMEOUT ? 1.0f : 0.25f;
 			})
 			.Visibility_Lambda([this, Item]()
@@ -1106,7 +1108,7 @@ namespace UE::Audio::Insights
 				{
 					if (!bIsPinnedCategoryFilterEnabled || (SoundEntry.bIsCategory && CountNumChildren(*Item, bShowRecentlyStoppedSounds, true /*bIncludeTimingOutSounds*/) == 0))
 					{
-						return EVisibility::Hidden;
+						return EVisibility::Collapsed;
 					}
 				}
 
@@ -1117,8 +1119,10 @@ namespace UE::Audio::Insights
 												&& SoundEntry.PinnedEntryType != FSoundDashboardEntry::EPinnedEntryType::HiddenOriginalEntry 
 												&& (!IsCategoryItem(*Item) || (NumUnpinnedChildren > 0));
 
-				return bRowShouldBeVisible ? EVisibility::Visible : EVisibility::Hidden;
-			});
+				return bRowShouldBeVisible ? EVisibility::Visible : EVisibility::Collapsed;
+			})
+			: SNew(STableRow<TSharedPtr<IDashboardDataTreeViewEntry>>, OwnerTable)
+			  .Visibility(EVisibility::Collapsed);
 	}
 
 	void FSoundDashboardViewFactory::ProcessEntries(FTraceTreeDashboardViewFactory::EProcessReason InReason)
@@ -1164,6 +1168,11 @@ namespace UE::Audio::Insights
 
 			return bEntryTypePassesFilter;
 		});
+	}
+
+	void FSoundDashboardViewFactory::FilterText()
+	{
+		using namespace FSoundDashboardViewFactoryPrivate;
 
 		// Filter by text
 		const FString FilterString = GetSearchFilterText().ToString();
@@ -1565,6 +1574,8 @@ namespace UE::Audio::Insights
 	{
 		using namespace FSoundDashboardViewFactoryPrivate;
 
+		FilterText();
+
 		UpdatePinnedSection();
 
 		auto SortByPlayOrder = [](const FSoundDashboardEntry& First, const FSoundDashboardEntry& Second)
@@ -1844,7 +1855,8 @@ namespace UE::Audio::Insights
 	}
 
 #if WITH_EDITOR
-	TSharedRef<SWidget> FSoundDashboardViewFactory::CreateMuteSoloButton(const TSharedRef<IDashboardDataTreeViewEntry>& InRowData, 
+	TSharedRef<SWidget> FSoundDashboardViewFactory::CreateMuteSoloButton(const TSharedRef<FTraceTreeDashboardViewFactory::SRowWidget>& InRowWidget, 
+		const TSharedRef<IDashboardDataTreeViewEntry>& InRowData,
 		const FName& InColumn, 
 		TFunction<void(const TArray<TSharedPtr<IDashboardDataTreeViewEntry>>&)> MuteSoloToggleFunc, 
 		TFunctionRef<bool(const IDashboardDataTreeViewEntry&, const bool)> IsMuteSoloFunc)
@@ -1852,6 +1864,10 @@ namespace UE::Audio::Insights
 		using namespace FSoundDashboardViewFactoryPrivate;
 
 		return SNew(SHorizontalBox)
+			.Visibility_Lambda([this, InRowWidget, InRowData]()
+			{
+				return InRowWidget->GetVisibility();
+			})
 			+ SHorizontalBox::Slot()
 			.AutoWidth()
 			.Padding(2.0f, 0.0f, 0.0f, 0.0f)
@@ -1869,7 +1885,7 @@ namespace UE::Audio::Insights
 					.Image_Lambda([this, InRowData, InColumn, IsMuteSoloFunc]()
 					{
 						const FColumnData& ColumnData = GetColumns()[InColumn];
-						const FName IconName = IsMuteSoloFunc(*InRowData, EntryCanHaveChildren(InRowData) /*bInCheckChildren*/) && ColumnData.GetIconName.IsSet() ? ColumnData.GetIconName(InRowData.Get()) : "AudioInsights.Icon.SoundDashboard.Transparent";
+						const FName IconName = IsMuteSoloFunc(*InRowData, EntryCanHaveChildren(InRowData) /*bInCheckChildren*/) ? ColumnData.GetIconName(InRowData.Get()) : "AudioInsights.Icon.SoundDashboard.Transparent";
 
 						return FSlateStyle::Get().GetBrush(IconName);
 					})
@@ -1881,6 +1897,14 @@ namespace UE::Audio::Insights
 	{
 		using namespace FSoundDashboardViewFactoryPrivate;
 
+		FAudioDeviceManager* AudioDeviceManager = FAudioDeviceManager::Get();
+		if (AudioDeviceManager == nullptr)
+		{
+			return;
+		}
+
+		::Audio::FAudioDebugger& AudioDebugger = AudioDeviceManager->GetDebugger();
+
 		TArray<TSharedPtr<IDashboardDataTreeViewEntry>> EntriesToMuteSolo;
 
 		// In multiple selection we need to discard children entries to avoid double mute/solo toggling
@@ -1926,7 +1950,7 @@ namespace UE::Audio::Insights
 
 			if (EntryCanHaveChildren(Entry.ToSharedRef()))
 			{
-				const bool bAreChildrenMuteSolo = FSoundDashboardViewFactoryPrivate::IsMuteSolo(*Entry, true /*bInCheckChildren*/, InMuteSoloMode);
+				const bool bAreChildrenMuteSolo = FSoundDashboardViewFactoryPrivate::IsMuteSolo(AudioDebugger, *Entry, true /*bInCheckChildren*/, InMuteSoloMode);
 
 				for (const TSharedPtr<IDashboardDataTreeViewEntry>& ChildEntry : Entry->Children)
 				{
@@ -2203,7 +2227,7 @@ namespace UE::Audio::Insights
 		}
 
 		TSharedPtr<FSoundDashboardEntry> PinnedCategory = MakeShared<FSoundDashboardEntry>();
-		PinnedCategory->Name = FSoundDashboardViewFactoryPrivate::PinnedCategoryName.ToString();
+		PinnedCategory->SetName(FSoundDashboardViewFactoryPrivate::PinnedCategoryName.ToString());
 		PinnedCategory->EntryType = ESoundDashboardEntryType::Pinned;
 		PinnedCategory->PinnedEntryType = FSoundDashboardEntry::EPinnedEntryType::PinnedCopy;
 		PinnedCategory->bIsCategory = true;
diff --git a/Engine/Plugins/AudioInsights/Source/AudioInsights/Public/Messages/SoundTraceMessages.h b/Engine/Plugins/AudioInsights/Source/AudioInsights/Public/Messages/SoundTraceMessages.h
--- a/Engine/Plugins/AudioInsights/Source/AudioInsights/Public/Messages/SoundTraceMessages.h
+++ b/Engine/Plugins/AudioInsights/Source/AudioInsights/Public/Messages/SoundTraceMessages.h
@@ -148,13 +148,33 @@ namespace UE::Audio::Insights
 		FSoundDashboardEntry();
 		virtual ~FSoundDashboardEntry() = default;
 
-		virtual TObjectPtr<UObject> GetObject() override { return FSoftObjectPath(Name).ResolveObject(); }
-		virtual const TObjectPtr<UObject> GetObject() const override { return FSoftObjectPath(Name).ResolveObject(); }
+		virtual TObjectPtr<UObject> GetObject() override { return FSoftObjectPath(FullName).ResolveObject(); }
+		virtual const TObjectPtr<UObject> GetObject() const override { return FSoftObjectPath(FullName).ResolveObject(); }
 
-		virtual FText GetDisplayName() const override
+		void SetName(const FString& InName)
 		{
-			const FSoftObjectPath AssetPath(Name);
-			return FText::FromString(AssetPath.IsValid() ? AssetPath.GetAssetName() : Name);
+			FullName = InName;
+
+			const FSoftObjectPath AssetPath(InName);
+
+			DisplayNameStr  = AssetPath.IsValid() ? AssetPath.GetAssetName() : InName;
+			DisplayName     = AssetPath.IsValid() ? AssetPath.GetAssetFName() : FName(InName);
+			DisplayNameText = FText::FromString(AssetPath.IsValid() ? DisplayNameStr : InName);
+		}
+
+		virtual const FText& GetDisplayName() const override
+		{
+			return DisplayNameText;
+		}
+
+		const FName& GetDisplayFName() const
+		{
+			return DisplayName;
+		}
+
+		const FString& GetDisplayNameStr() const
+		{
+			return DisplayNameStr;
 		}
 
 		virtual bool IsValid() const override { return PlayOrder != INDEX_NONE; }
@@ -173,7 +193,6 @@ namespace UE::Audio::Insights
 		uint32 PlayOrder = INDEX_NONE;
 		double Timestamp = 0.0;
 		double TimeoutTimestamp = INVALID_TIMEOUT;
-		FString Name;
 		FText CategoryName;
 		ESoundDashboardEntryType EntryType = ESoundDashboardEntryType::None;
 		EPinnedEntryType PinnedEntryType = EPinnedEntryType::None;
@@ -191,5 +210,11 @@ namespace UE::Audio::Insights
 #if WITH_EDITOR
 		FText OwnerActorName;
 #endif // WITH_EDITOR
+
+	private:
+		FString FullName;
+		FName DisplayName;
+		FString DisplayNameStr;
+		FText DisplayNameText;
 	};
 } // namespace UE::Audio::Insights
diff --git a/Engine/Plugins/AudioInsights/Source/AudioInsights/Public/Views/SoundDashboardViewFactory.h b/Engine/Plugins/AudioInsights/Source/AudioInsights/Public/Views/SoundDashboardViewFactory.h
--- a/Engine/Plugins/AudioInsights/Source/AudioInsights/Public/Views/SoundDashboardViewFactory.h
+++ b/Engine/Plugins/AudioInsights/Source/AudioInsights/Public/Views/SoundDashboardViewFactory.h
@@ -129,6 +129,7 @@ namespace UE::Audio::Insights
 		virtual TArray<TSharedPtr<IDashboardDataTreeViewEntry>>& GetTreeItemsSource() override { return FullTree; }
 		virtual bool ResetTreeData() override;
 
+		void FilterText();
 		void RecursiveSort(TArray<TSharedPtr<IDashboardDataTreeViewEntry>>& OutTree, TFunctionRef<bool(const FSoundDashboardEntry&, const FSoundDashboardEntry&)> Predicate);
 		void SortByPredicate(TFunctionRef<bool(const FSoundDashboardEntry&, const FSoundDashboardEntry&)> Predicate);
 
@@ -137,7 +138,8 @@ namespace UE::Audio::Insights
 #if WITH_EDITOR
 		TSharedRef<SWidget> MakeMuteSoloWidget();
 
-		TSharedRef<SWidget> CreateMuteSoloButton(const TSharedRef<IDashboardDataTreeViewEntry>& InRowData,
+		TSharedRef<SWidget> CreateMuteSoloButton(const TSharedRef<FTraceTreeDashboardViewFactory::SRowWidget>& InRowWidget, 
+			const TSharedRef<IDashboardDataTreeViewEntry>& InRowData,
 			const FName& InColumn,
 			TFunction<void(const TArray<TSharedPtr<IDashboardDataTreeViewEntry>>& /*InEntries*/)> MuteSoloToggleFunc,
 			TFunctionRef<bool(const IDashboardDataTreeViewEntry& /*InEntry*/, const bool /*bInCheckChildren*/)> IsMuteSoloFunc);
diff --git a/Engine/Plugins/AudioInsights/Source/AudioInsights/Public/Views/TreeDashboardViewFactory.h b/Engine/Plugins/AudioInsights/Source/AudioInsights/Public/Views/TreeDashboardViewFactory.h
--- a/Engine/Plugins/AudioInsights/Source/AudioInsights/Public/Views/TreeDashboardViewFactory.h
+++ b/Engine/Plugins/AudioInsights/Source/AudioInsights/Public/Views/TreeDashboardViewFactory.h
@@ -25,7 +25,7 @@ namespace UE::Audio::Insights
 		virtual TObjectPtr<UObject> GetObject() = 0;
 		virtual const TObjectPtr<UObject> GetObject() const = 0;
 
-		virtual FText GetDisplayName() const = 0;
+		virtual const FText& GetDisplayName() const = 0;
 	};
 
 	/** 
