[Audio Insights] Fixes for build error triggered by CL 42823616

diff --git a/Engine/Plugins/AudioInsights/Source/AudioInsights/Private/Views/SoundDashboardViewFactory.cpp b/Engine/Plugins/AudioInsights/Source/AudioInsights/Private/Views/SoundDashboardViewFactory.cpp
--- a/Engine/Plugins/AudioInsights/Source/AudioInsights/Private/Views/SoundDashboardViewFactory.cpp
+++ b/Engine/Plugins/AudioInsights/Source/AudioInsights/Private/Views/SoundDashboardViewFactory.cpp
@@ -283,9 +283,9 @@ namespace UE::Audio::Insights
 #endif // ENABLE_AUDIO_DEBUG
 		}
 
+#if ENABLE_AUDIO_DEBUG
 		bool IsMuteSolo(::Audio::FAudioDebugger& InAudioDebugger, const IDashboardDataTreeViewEntry& InEntry, const bool bInCheckChildren, const EMuteSoloMode InMuteSoloMode)
 		{
-#if ENABLE_AUDIO_DEBUG
 			const FSoundDashboardEntry& SoundEntry = FSoundDashboardViewFactoryPrivate::CastEntry(InEntry);
 
 			// Treat hidden original entries as muted/soloed to ensure the parent category reflects the correct state
@@ -341,10 +341,10 @@ namespace UE::Audio::Insights
 
 				return bAllChildrenMuteSolo;
 			}
-#endif // ENABLE_AUDIO_DEBUG
 
 			return false;
 		}
+#endif // ENABLE_AUDIO_DEBUG
 
 		void ClearMutesAndSolos()
 		{
@@ -991,8 +991,11 @@ namespace UE::Audio::Insights
 				[](const IDashboardDataTreeViewEntry& InEntry, const bool bInCheckChildren)
 				{
 					FAudioDeviceManager* AudioDeviceManager = FAudioDeviceManager::Get();
-
+#if ENABLE_AUDIO_DEBUG
 					return AudioDeviceManager ? IsMuteSolo(AudioDeviceManager->GetDebugger(), InEntry, bInCheckChildren, EMuteSoloMode::Mute) : false;
+#else
+					return false;
+#endif // ENABLE_AUDIO_DEBUG
 				});
 		}
 		else if (InColumn == "Solo")
@@ -1005,8 +1008,11 @@ namespace UE::Audio::Insights
 				[](const IDashboardDataTreeViewEntry& InEntry, const bool bInCheckChildren)
 				{
 					FAudioDeviceManager* AudioDeviceManager = FAudioDeviceManager::Get();
-
+#if ENABLE_AUDIO_DEBUG
 					return AudioDeviceManager ? IsMuteSolo(AudioDeviceManager->GetDebugger(), InEntry, bInCheckChildren, EMuteSoloMode::Solo) : false;
+#else
+					return false;
+#endif // ENABLE_AUDIO_DEBUG
 				});
 		}
 		else
@@ -1895,6 +1901,7 @@ namespace UE::Audio::Insights
 
 	void FSoundDashboardViewFactory::ToggleMuteSoloEntries(const TArray<TSharedPtr<IDashboardDataTreeViewEntry>>& InEntries, const EMuteSoloMode InMuteSoloMode)
 	{
+#if ENABLE_AUDIO_DEBUG
 		using namespace FSoundDashboardViewFactoryPrivate;
 
 		FAudioDeviceManager* AudioDeviceManager = FAudioDeviceManager::Get();
@@ -1965,6 +1972,7 @@ namespace UE::Audio::Insights
 				FSoundDashboardViewFactoryPrivate::ToggleMuteSolo(*Entry, InMuteSoloMode);
 			}
 		}
+#endif // ENABLE_AUDIO_DEBUG
 	}
 
 	TArray<TObjectPtr<UObject>> FSoundDashboardViewFactory::GetSelectedEditableAssets() const
