[Audio Insights] Added missing nullptr check in AudioBus dashboard + avoid initializing audio insights if restricted mode is enabled

diff --git a/Engine/Plugins/AudioInsights/Source/AudioInsightsEditor/Private/Providers/AudioBusProvider.cpp b/Engine/Plugins/AudioInsights/Source/AudioInsightsEditor/Private/Providers/AudioBusProvider.cpp
--- a/Engine/Plugins/AudioInsights/Source/AudioInsightsEditor/Private/Providers/AudioBusProvider.cpp
+++ b/Engine/Plugins/AudioInsights/Source/AudioInsightsEditor/Private/Providers/AudioBusProvider.cpp
@@ -105,7 +105,7 @@ namespace UE::Audio::Insights
 			AudioBusAssetDashboardEntry->EntryType  = EAudioBusEntryType::AssetBased;
 			AudioBusAssetDashboardEntry->Name       = InAssetData.GetObjectPathString();
 			AudioBusAssetDashboardEntry->AudioBus   = Cast<UAudioBus>(InAssetData.GetAsset());
-			AudioBusAssetDashboardEntry->AudioBusId = AudioBusAssetDashboardEntry->AudioBus->GetUniqueID();
+			AudioBusAssetDashboardEntry->AudioBusId = AudioBusAssetDashboardEntry->AudioBus.IsValid() ? AudioBusAssetDashboardEntry->AudioBus->GetUniqueID() : INDEX_NONE;
 
 			AudioBusDataViewEntries.Add(MoveTemp(AudioBusAssetDashboardEntry));
 
