[Audio Insights] - Fix Control Buses and Modulation Matrix Windows have missing icons in the View tab

diff --git a/Engine/Plugins/AudioInsights/Source/AudioInsightsEditor/Private/AudioInsightsEditorDashboardFactory.cpp b/Engine/Plugins/AudioInsights/Source/AudioInsightsEditor/Private/AudioInsightsEditorDashboardFactory.cpp
--- a/Engine/Plugins/AudioInsights/Source/AudioInsightsEditor/Private/AudioInsightsEditorDashboardFactory.cpp
+++ b/Engine/Plugins/AudioInsights/Source/AudioInsightsEditor/Private/AudioInsightsEditorDashboardFactory.cpp
@@ -312,7 +312,7 @@ namespace UE::Audio::Insights
 
 					MenuBuilder.AddMenuEntry(Factory->GetDisplayName(),
 						FText::GetEmpty(),
-						FSlateStyle::Get().CreateIcon(Factory->GetIcon().GetStyleName()),
+						Factory->GetIcon(),
 						FUIAction(FExecuteAction::CreateLambda([this, FactoryName]()
 						{
 							if (DashboardTabManager.IsValid())
