[Audio Insights] Sound dashboard: Fixed issue where rows were staying invisible due to an optimization code change.

diff --git a/Engine/Plugins/AudioInsights/Source/AudioInsights/Private/Views/SoundDashboardViewFactory.cpp b/Engine/Plugins/AudioInsights/Source/AudioInsights/Private/Views/SoundDashboardViewFactory.cpp
--- a/Engine/Plugins/AudioInsights/Source/AudioInsights/Private/Views/SoundDashboardViewFactory.cpp
+++ b/Engine/Plugins/AudioInsights/Source/AudioInsights/Private/Views/SoundDashboardViewFactory.cpp
@@ -1099,8 +1099,7 @@ namespace UE::Audio::Insights
 
 		const FSoundDashboardEntry& SoundDashboardEntry = CastEntry(*Item);
 
-		return SoundDashboardEntry.bIsVisible ?
-			SNew(SRowWidget, OwnerTable, Item, AsShared())
+		return SNew(SRowWidget, OwnerTable, Item, AsShared())
 			.RenderOpacity_Lambda([Item]
 			{
 				const FSoundDashboardEntry& SoundDashboardEntry = CastEntry(*Item);
@@ -1126,9 +1125,7 @@ namespace UE::Audio::Insights
 												&& (!IsCategoryItem(*Item) || (NumUnpinnedChildren > 0));
 
 				return bRowShouldBeVisible ? EVisibility::Visible : EVisibility::Collapsed;
-			})
-			: SNew(STableRow<TSharedPtr<IDashboardDataTreeViewEntry>>, OwnerTable)
-			  .Visibility(EVisibility::Collapsed);
+			});
 	}
 
 	void FSoundDashboardViewFactory::ProcessEntries(FTraceTreeDashboardViewFactory::EProcessReason InReason)
