[Audio Insights] Allow users to strip away the categories inside the Sounds Dashboard

diff --git a/Engine/Plugins/AudioInsights/Source/AudioInsights/Private/Views/SoundDashboardViewFactory.cpp b/Engine/Plugins/AudioInsights/Source/AudioInsights/Private/Views/SoundDashboardViewFactory.cpp
--- a/Engine/Plugins/AudioInsights/Source/AudioInsights/Private/Views/SoundDashboardViewFactory.cpp
+++ b/Engine/Plugins/AudioInsights/Source/AudioInsights/Private/Views/SoundDashboardViewFactory.cpp
@@ -102,6 +102,36 @@ namespace UE::Audio::Insights
 			return SoundEntry.bIsVisible && (bShowRecentlyStoppedSounds || SoundEntry.TimeoutTimestamp == INVALID_TIMEOUT);
 		}
 
+		FSlateColor GetIconColor(const TSharedPtr<IDashboardDataTreeViewEntry>& InEntry)
+		{
+			const FSoundDashboardEntry& SoundDashboardEntry = FSoundDashboardViewFactoryPrivate::CastEntry(*InEntry);
+
+			switch (SoundDashboardEntry.EntryType)
+			{
+				case ESoundDashboardEntryType::MetaSound:
+					return FSlateColor(FSlateStyle::Get().GetColor("SoundDashboard.MetaSoundColor"));
+
+				case ESoundDashboardEntryType::SoundCue:
+					return FSlateColor(FSlateStyle::Get().GetColor("SoundDashboard.SoundCueColor"));
+
+				case ESoundDashboardEntryType::ProceduralSource:
+					return FSlateColor(FSlateStyle::Get().GetColor("SoundDashboard.ProceduralSourceColor"));
+
+				case ESoundDashboardEntryType::SoundWave:
+					return FSlateColor(FSlateStyle::Get().GetColor("SoundDashboard.SoundWaveColor"));
+
+				case ESoundDashboardEntryType::SoundCueTemplate:
+					return FSlateColor(FSlateStyle::Get().GetColor("SoundDashboard.SoundCueTemplateColor"));
+
+				case ESoundDashboardEntryType::Pinned:
+					return FSlateColor(FSlateStyle::Get().GetColor("SoundDashboard.PinnedColor"));
+
+				case ESoundDashboardEntryType::None:
+				default:
+					return FSlateColor(FColor::White);
+			}
+		};
+
 		void CacheInitExpandStateRecursive(FSoundDashboardEntry& SoundEntry)
 		{
 			for (TSharedPtr<IDashboardDataTreeViewEntry>& ChildEntry : SoundEntry.Children)
@@ -882,37 +912,6 @@ namespace UE::Audio::Insights
 		{
 			const FName IconName = ColumnData.GetIconName.IsSet() ? ColumnData.GetIconName(InRowData.Get()) : NAME_None;
 
-			// Custom color for root item icons
-			const auto GetIconColor = [](const TSharedPtr<IDashboardDataTreeViewEntry>& InEntry)
-			{
-				const FSoundDashboardEntry& SoundDashboardEntry = FSoundDashboardViewFactoryPrivate::CastEntry(*InEntry);
-
-				switch (SoundDashboardEntry.EntryType)
-				{
-				case ESoundDashboardEntryType::MetaSound:
-					return FSlateColor(FSlateStyle::Get().GetColor("SoundDashboard.MetaSoundColor"));
-
-				case ESoundDashboardEntryType::SoundCue:
-					return FSlateColor(FSlateStyle::Get().GetColor("SoundDashboard.SoundCueColor"));
-
-				case ESoundDashboardEntryType::ProceduralSource:
-					return FSlateColor(FSlateStyle::Get().GetColor("SoundDashboard.ProceduralSourceColor"));
-
-				case ESoundDashboardEntryType::SoundWave:
-					return FSlateColor(FSlateStyle::Get().GetColor("SoundDashboard.SoundWaveColor"));
-
-				case ESoundDashboardEntryType::SoundCueTemplate:
-					return FSlateColor(FSlateStyle::Get().GetColor("SoundDashboard.SoundCueTemplateColor"));
-
-				case ESoundDashboardEntryType::Pinned:
-					return FSlateColor(FSlateStyle::Get().GetColor("SoundDashboard.PinnedColor"));
-
-				case ESoundDashboardEntryType::None:
-				default:
-					return FSlateColor(FColor::White);
-				}
-			};
-
 			return SNew(SHorizontalBox)
 				.Visibility_Lambda([InRowWidget]()
 				{
@@ -930,7 +929,12 @@ namespace UE::Audio::Insights
 				.AutoWidth()
 				[
 					IconName != NAME_None 
-						? SNew(SImage).ColorAndOpacity(GetIconColor(InRowData)).Image(FSlateStyle::Get().GetBrush(IconName)) 
+						? SNew(SImage)
+						.ColorAndOpacity_Lambda([InRowData]() 
+						{
+							return GetIconColor(InRowData);
+						})
+						.Image(FSlateStyle::Get().GetBrush(IconName)) 
 						: SNullWidget::NullWidget
 				]
 				// Text
@@ -939,7 +943,11 @@ namespace UE::Audio::Insights
 				.AutoWidth()
 				[
 					SNew(STextBlock)
-					.Font(FCoreStyle::Get().GetFontStyle("BoldFont"))
+					.Font_Lambda([this, InRowData]()
+					{
+						const FSoundDashboardEntry& SoundDashboardEntry = FSoundDashboardViewFactoryPrivate::CastEntry(*InRowData);
+						return SoundDashboardEntry.bIsCategory ? FCoreStyle::Get().GetFontStyle("BoldFont") : FCoreStyle::Get().GetFontStyle("NormalFont");
+					})
 					.Text(InValueText)
 					.ColorAndOpacity(FSlateColor::UseForeground())
 				]
@@ -953,6 +961,12 @@ namespace UE::Audio::Insights
 					.Text_Lambda([this, InRowData]()
 					{
 						const uint32 TotalNumChildren = CountNumChildren(*InRowData, bShowRecentlyStoppedSounds);
+
+						if (TreeViewingMode != ETreeViewingOptions::FullTree && TotalNumChildren == 0)
+						{
+							return FText::GetEmpty();
+						}
+
 						return FText::FromString("(" + FString::FromInt(TotalNumChildren) + ")");
 					})
 				];
@@ -1041,6 +1055,14 @@ namespace UE::Audio::Insights
 				.AutoWidth()
 				[
 					SNew(SImage)
+					.ColorAndOpacity_Lambda([this, InRowData]()
+					{
+						if (IsRootItem(InRowData))
+						{
+							return GetIconColor(InRowData);
+						}
+						return FSlateColor(FColor::White);
+					})
 					.Image(FSlateStyle::Get().GetBrush(ColumnData.GetIconName(InRowData.Get())))
 				]
 				// Text
@@ -1069,7 +1091,7 @@ namespace UE::Audio::Insights
 		}
 		else
 		{
-			if (IsRootItem(InRowData))
+			if (IsRootItem(InRowData) && TreeViewingMode == ETreeViewingOptions::FullTree)
 			{
 				return SNullWidget::NullWidget;
 			}
@@ -1109,6 +1131,11 @@ namespace UE::Audio::Insights
 			{
 				const FSoundDashboardEntry& SoundEntry = CastEntry(*Item);
 
+				if (TreeViewingMode != ETreeViewingOptions::FullTree && SoundEntry.bIsCategory && SoundEntry.EntryType != ESoundDashboardEntryType::Pinned)
+				{
+					return EVisibility::Collapsed;
+				}
+
 				if (SoundEntry.PinnedEntryType == FSoundDashboardEntry::EPinnedEntryType::PinnedCopy)
 				{
 					if (!bIsPinnedCategoryFilterEnabled || (SoundEntry.bIsCategory && CountNumChildren(*Item, bShowRecentlyStoppedSounds, true /*bIncludeTimingOutSounds*/) == 0))
@@ -1132,6 +1159,15 @@ namespace UE::Audio::Insights
 	{
 		using namespace FSoundDashboardViewFactoryPrivate;
 
+		if (bTreeViewModeChanged)
+		{
+			for (const TSharedPtr<IDashboardDataTreeViewEntry>& Entry : DataViewEntries)
+			{
+				FSoundDashboardEntry& SoundEntry = CastEntry(*Entry);
+				CacheInitExpandStateRecursive(SoundEntry);
+			}
+		}
+
 		// Filter by category
 		FTraceTreeDashboardViewFactory::FilterEntries<FSoundTraceProvider>([this](IDashboardDataTreeViewEntry& InEntry)
 		{
@@ -1190,13 +1226,20 @@ namespace UE::Audio::Insights
 
 			ResetVisibility(*CategoryEntry);
 
-			for (const TSharedPtr<IDashboardDataTreeViewEntry>& Entry : CategoryEntry->Children)
+			if (TreeViewingMode == ETreeViewingOptions::FullTree)
 			{
-				if (Entry.IsValid() && bFilterHasText)
+				for (const TSharedPtr<IDashboardDataTreeViewEntry>& Entry : CategoryEntry->Children)
 				{
-					SetFilteredVisibility(*Entry, FilterString);
+					if (Entry.IsValid() && bFilterHasText)
+					{
+						SetFilteredVisibility(*Entry, FilterString);
+					}
 				}
 			}
+			else if (bFilterHasText)
+			{
+				SetFilteredVisibility(*CategoryEntry, FilterString);
+			}
 		}
 	}
 
@@ -1734,6 +1777,12 @@ namespace UE::Audio::Insights
 		}
 		
 		FullTree.Append(DataViewEntries);
+
+		if (bTreeViewModeChanged)
+		{
+			FilteredEntriesListView->RebuildList();
+			bTreeViewModeChanged = false;
+		}
 	}
 
 	bool FSoundDashboardViewFactory::ResetTreeData()
@@ -1814,7 +1863,52 @@ namespace UE::Audio::Insights
 
 	TSharedRef<SWidget> FSoundDashboardViewFactory::MakeShowRecentlyStoppedSoundsWidget()
 	{
+		using namespace FSoundDashboardViewFactoryPrivate;
+
 		return SNew(SHorizontalBox)
+			+ SHorizontalBox::Slot()
+			.AutoWidth()
+			.Padding(8.0f, 8.0f)
+			[
+				SNew(SCheckBox)
+				.HAlign(EHorizontalAlignment::HAlign_Center)
+				.Style(&FSlateStyle::Get().GetWidgetStyle<FCheckBoxStyle>("SoundDashboard.MuteSoloButton"))
+				.ToolTipText(LOCTEXT("SoundsDashboard_ShowCategoriesTooltip", "Shows categories"))
+				.IsChecked(TreeViewingMode == ETreeViewingOptions::FullTree ? ECheckBoxState::Checked : ECheckBoxState::Unchecked)
+				.OnCheckStateChanged_Lambda([this](ECheckBoxState NewState)
+				{
+					TreeViewingMode = NewState == ECheckBoxState::Checked ? ETreeViewingOptions::FullTree : ETreeViewingOptions::ExcludingCategories;
+					UpdateFilterReason = EProcessReason::FilterUpdated;
+					ClearSelection();
+
+					bTreeViewModeChanged = true;
+				})
+				[
+					SNew(SHorizontalBox)
+					+ SHorizontalBox::Slot()
+					.HAlign(HAlign_Left)
+					.VAlign(VAlign_Center)
+					.AutoWidth()
+					.Padding(2.0f, 0.0f)
+					[
+						SNew(SImage)
+				 		.Image_Lambda([this]()
+				 		{
+				 			const FName IconName = TreeViewingMode == ETreeViewingOptions::FullTree ? "AudioInsights.Icon.SoundDashboard.Visible" : "AudioInsights.Icon.SoundDashboard.Invisible";
+				 			return FSlateStyle::Get().GetBrush(IconName);
+				 		})
+					]
+					+ SHorizontalBox::Slot()
+					.HAlign(HAlign_Right)
+					.VAlign(VAlign_Center)
+					.AutoWidth()
+					.Padding(2.0f, 0.0f)
+					[
+						SNew(STextBlock)
+						.Text(LOCTEXT("SoundsDashboard_ShowStoppedSounds", "Show Categories"))
+					]
+				]
+			]
 			+ SHorizontalBox::Slot()
 			.AutoWidth()
 			.Padding(8.0f, 8.0f)
@@ -2033,6 +2127,42 @@ namespace UE::Audio::Insights
 		return false;
 	}
 
+	bool FSoundDashboardViewFactory::SelectionIncludesUnpinnedItem() const
+	{
+		using namespace FSoundDashboardViewFactoryPrivate;
+
+		if (!FilteredEntriesListView.IsValid())
+		{
+			return false;
+		}
+
+		const TArray<TSharedPtr<IDashboardDataTreeViewEntry>> SelectedItems = FilteredEntriesListView->GetSelectedItems();
+		for (const TSharedPtr<IDashboardDataTreeViewEntry>& SelectedItem : SelectedItems)
+		{
+			if (!SelectedItem.IsValid())
+			{
+				continue;
+			}
+
+			const FSoundDashboardEntry& SelectedSoundEntry = CastEntry(*SelectedItem);
+			if (SelectedSoundEntry.PinnedEntryType == FSoundDashboardEntry::EPinnedEntryType::None)
+			{
+				return true;
+			}
+		}
+
+		return false;
+	}
+
+	void FSoundDashboardViewFactory::ClearSelection()
+	{
+		// Make sure to clear the selection and the internal selector in SListView
+		// to ensure no shared references keep the entry alive past the point
+		// it has been removed from the dashboard
+		FilteredEntriesListView->ClearSelection();
+		FilteredEntriesListView->SetSelection(nullptr);
+	}
+
 	void FSoundDashboardViewFactory::PinSound()
 	{
 		if (FilteredEntriesListView.IsValid())
@@ -2057,11 +2187,7 @@ namespace UE::Audio::Insights
 				}
 			}
 			
-			// Make sure to clear the selection and the internal selector in SListView
-			// to ensure no shared references keep the entry alive past the point
-			// it has been removed from the dashboard
-			FilteredEntriesListView->ClearSelection();
-			FilteredEntriesListView->SetSelection(nullptr);
+			ClearSelection();
 		}
 	}
 
@@ -2087,62 +2213,46 @@ namespace UE::Audio::Insights
 				UnpinSelectedItems(PinnedItemEntries, SelectedItems, bSelectionContainsAssets);
 			}
 			
-			// Make sure to clear the selection and the internal selector in SListView
-			// to ensure no shared references keep the entry alive past the point
-			// it has been removed from the dashboard
-			FilteredEntriesListView->ClearSelection();
-			FilteredEntriesListView->SetSelection(nullptr);
+			ClearSelection();
 		}
 	}
 
-	bool FSoundDashboardViewFactory::SelectionIncludesUnpinnedItem() const
+	void FSoundDashboardViewFactory::PinSelectedItems(const TSharedPtr<IDashboardDataTreeViewEntry>& Entry, const TArray<TSharedPtr<IDashboardDataTreeViewEntry>>& SelectedItems)
 	{
 		using namespace FSoundDashboardViewFactoryPrivate;
 
-		if (!FilteredEntriesListView.IsValid())
+		if (SelectedItems.Contains(Entry) && (!IsRootItem(Entry.ToSharedRef()) || TreeViewingMode != ETreeViewingOptions::FullTree))
 		{
-			return false;
-		}
+			const FSoundDashboardEntry& SoundEntry = CastEntry(*Entry);
 
-		const TArray<TSharedPtr<IDashboardDataTreeViewEntry>> SelectedItems = FilteredEntriesListView->GetSelectedItems();
-		for (const TSharedPtr<IDashboardDataTreeViewEntry>& SelectedItem : SelectedItems)
-		{
-			if (!SelectedItem.IsValid())
+			if (SoundEntry.PinnedEntryType != FSoundDashboardEntry::EPinnedEntryType::HiddenOriginalEntry)
 			{
-				continue;
-			}
-
-			const FSoundDashboardEntry& SelectedSoundEntry = CastEntry(*SelectedItem);
-			if (SelectedSoundEntry.PinnedEntryType == FSoundDashboardEntry::EPinnedEntryType::None)
-			{
-				return true;
+				MarkBranchAsPinned(Entry, true /*bIsPinned*/);
+				CreatePinnedEntry(Entry);
 			}
 		}
-
-		return false;
-	}
-
-	void FSoundDashboardViewFactory::PinSelectedItems(const TSharedPtr<IDashboardDataTreeViewEntry>& Entry, const TArray<TSharedPtr<IDashboardDataTreeViewEntry>>& SelectedItems)
-	{
-		for (const TSharedPtr<IDashboardDataTreeViewEntry>& Child : Entry->Children)
+		else 
 		{
-			if (SelectedItems.Contains(Child))
+			for (const TSharedPtr<IDashboardDataTreeViewEntry>& Child : Entry->Children)
 			{
-				if (IsRootItem(Entry.ToSharedRef()))
+				if (SelectedItems.Contains(Child))
 				{
-					MarkBranchAsPinned(Child, true /*bIsPinned*/);
-					CreatePinnedEntry(Child);
+					if (TreeViewingMode == ETreeViewingOptions::FullTree && IsRootItem(Entry.ToSharedRef()))
+					{
+						MarkBranchAsPinned(Child, true /*bIsPinned*/);
+						CreatePinnedEntry(Child);
+					}
+					else
+					{
+						MarkBranchAsPinned(Entry, true /*bIsPinned*/);
+						CreatePinnedEntry(Entry);
+					}
 				}
 				else
 				{
-					MarkBranchAsPinned(Entry, true /*bIsPinned*/);
-					CreatePinnedEntry(Entry);
+					PinSelectedItems(Child, SelectedItems);
 				}
 			}
-			else
-			{
-				PinSelectedItems(Child, SelectedItems);
-			}
 		}
 	}
 
@@ -2213,6 +2323,41 @@ namespace UE::Audio::Insights
 		}
 	}
 
+	void FSoundDashboardViewFactory::RecreatePinnedEntries(const TSharedPtr<IDashboardDataTreeViewEntry>& Entry)
+	{
+		FSoundDashboardEntry& SoundEntry = FSoundDashboardViewFactoryPrivate::CastEntry(*Entry);
+
+		if (Entry->Children.IsEmpty() && SoundEntry.PinnedEntryType == FSoundDashboardEntry::EPinnedEntryType::HiddenOriginalEntry)
+		{
+			MarkBranchAsPinned(Entry, true /*bIsPinned*/);
+			CreatePinnedEntry(Entry);
+		}
+		else
+		{
+			for (const TSharedPtr<IDashboardDataTreeViewEntry>& Child : Entry->Children)
+			{
+				if (SoundEntry.PinnedEntryType == FSoundDashboardEntry::EPinnedEntryType::HiddenOriginalEntry)
+				{
+					if (TreeViewingMode == ETreeViewingOptions::FullTree && SoundEntry.bIsCategory)
+					{
+						MarkBranchAsPinned(Child, true /*bIsPinned*/);
+						CreatePinnedEntry(Child);
+					}
+					else
+					{
+						MarkBranchAsPinned(Entry, true /*bIsPinned*/);
+						CreatePinnedEntry(Entry);
+						break;
+					}
+				}
+				else
+				{
+					RecreatePinnedEntries(Child);
+				}
+			}
+		}
+	}
+
 	void FSoundDashboardViewFactory::MarkBranchAsPinned(const TSharedPtr<IDashboardDataTreeViewEntry> Entry, const bool bIsPinned)
 	{
 		FSoundDashboardEntry& SoundEntry = FSoundDashboardViewFactoryPrivate::CastEntry(*Entry);
@@ -2248,58 +2393,65 @@ namespace UE::Audio::Insights
 		// The pinned area will delete itself once empty
 		InitPinnedItemEntries();
 
-		const FSoundDashboardEntry& SoundEntry = CastEntry(*Entry);
-
-		// Check if category is already in the list, if so we need to merge
-		bool bFoundExistingCategory = false;
-		for (TSharedPtr<FPinnedSoundEntryWrapper> PinnedCategoryEntry : PinnedItemEntries->PinnedWrapperChildren)
+		if (TreeViewingMode != ETreeViewingOptions::FullTree)
 		{
-			if (!PinnedCategoryEntry.IsValid())
-			{
-				continue;
-			}
-
-			const FSoundDashboardEntry& PinnedCategorySoundDashboardEntry = CastEntry(*PinnedCategoryEntry->GetPinnedSectionEntry());
+			PinnedItemEntries->AddChildEntry(Entry);
+		}
+		else
+		{
+			const FSoundDashboardEntry& SoundEntry = CastEntry(*Entry);
 
-			if (PinnedCategorySoundDashboardEntry.EntryType == SoundEntry.EntryType)
+			// Check if category is already in the list, if so we need to merge
+			bool bFoundExistingCategory = false;
+			for (TSharedPtr<FPinnedSoundEntryWrapper> PinnedCategoryEntry : PinnedItemEntries->PinnedWrapperChildren)
 			{
-				bFoundExistingCategory = true;
+				if (!PinnedCategoryEntry.IsValid())
+				{
+					continue;
+				}
+
+				const FSoundDashboardEntry& PinnedCategorySoundDashboardEntry = CastEntry(*PinnedCategoryEntry->GetPinnedSectionEntry());
 
-				const TSharedPtr<IDashboardDataTreeViewEntry>* FoundExisingEntry = PinnedCategorySoundDashboardEntry.Children.FindByPredicate([&SoundEntry](const TSharedPtr<IDashboardDataTreeViewEntry> PinnedEntry)
+				if (PinnedCategorySoundDashboardEntry.EntryType == SoundEntry.EntryType)
 				{
-					if (!PinnedEntry.IsValid())
+					bFoundExistingCategory = true;
+
+					const TSharedPtr<IDashboardDataTreeViewEntry>* FoundExisingEntry = PinnedCategorySoundDashboardEntry.Children.FindByPredicate([&SoundEntry](const TSharedPtr<IDashboardDataTreeViewEntry> PinnedEntry)
 					{
-						return false;
-					}
+						if (!PinnedEntry.IsValid())
+						{
+							return false;
+						}
 
-					const FSoundDashboardEntry& ExistingPinnedEntry = CastEntry(*PinnedEntry);
-					return ExistingPinnedEntry.PlayOrder == SoundEntry.PlayOrder;
-				});
+						const FSoundDashboardEntry& ExistingPinnedEntry = CastEntry(*PinnedEntry);
+						return ExistingPinnedEntry.PlayOrder == SoundEntry.PlayOrder;
+					});
 
-				// If we didn't find this entry already inside the pinned area, add it here
-				if (FoundExisingEntry == nullptr)
-				{
-					PinnedCategoryEntry->AddChildEntry(Entry);
+					// If we didn't find this entry already inside the pinned area, add it here
+					if (FoundExisingEntry == nullptr)
+					{
+						PinnedCategoryEntry->AddChildEntry(Entry);
+					}
+					break;
 				}
-				break;
 			}
-		}
 
-		// if we haven't found an existing pinned category, create a new one and add this item
-		if (!bFoundExistingCategory)
-		{
-			for (const TSharedPtr<IDashboardDataTreeViewEntry>& DataCategoryEntry : DataViewEntries)
+			// if we haven't found an existing pinned category, create a new one and add this item
+			if (!bFoundExistingCategory)
 			{
-				if (SoundEntry.EntryType == CastEntry(*DataCategoryEntry).EntryType)
+				for (const TSharedPtr<IDashboardDataTreeViewEntry>& DataCategoryEntry : DataViewEntries)
 				{
-					TSharedPtr<FSoundDashboardEntry> NewPinnedCategoryEntry = MakeShared<FSoundDashboardEntry>(CastEntry(*DataCategoryEntry));
-					NewPinnedCategoryEntry->Children.Reset();
-					NewPinnedCategoryEntry->bShouldForceExpandChildren = true;
+					if (SoundEntry.EntryType == CastEntry(*DataCategoryEntry).EntryType)
+					{
+						TSharedPtr<FSoundDashboardEntry> NewPinnedCategoryEntry = MakeShared<FSoundDashboardEntry>(CastEntry(*DataCategoryEntry));
+						NewPinnedCategoryEntry->Children.Reset();
+						NewPinnedCategoryEntry->bShouldForceExpandChildren = true;
 
-					TSharedPtr<FPinnedSoundEntryWrapper> PinnedCategory = PinnedItemEntries->AddChildEntry(NewPinnedCategoryEntry);
-					PinnedCategory->AddChildEntry(Entry);
+						TSharedPtr<FPinnedSoundEntryWrapper> PinnedCategory = PinnedItemEntries->AddChildEntry(NewPinnedCategoryEntry);
+						PinnedCategory->AddChildEntry(Entry);
 
-					break;
+						break;
+					}
 				}
 			}
 		}
@@ -2312,6 +2464,11 @@ namespace UE::Audio::Insights
 			return;
 		}
 
+		if (bTreeViewModeChanged)
+		{
+			RebuildPinnedSection();
+		}
+
 		PinnedItemEntries->CleanUp();
 
 		if (FPinnedSoundEntryWrapperPrivate::CanBeDeleted(PinnedItemEntries))
@@ -2324,6 +2481,36 @@ namespace UE::Audio::Insights
 		}
 	}
 
+	void FSoundDashboardViewFactory::RebuildPinnedSection()
+	{
+		using namespace FSoundDashboardViewFactoryPrivate;
+
+		if (!PinnedItemEntries.IsValid())
+		{
+			return;
+		}
+
+		PinnedItemEntries.Reset();
+		InitPinnedItemEntries();
+
+		for (TSharedPtr<IDashboardDataTreeViewEntry>& Entry : DataViewEntries)
+		{
+			FSoundDashboardEntry& SoundEntry = CastEntry(*Entry);
+			if (SoundEntry.PinnedEntryType == FSoundDashboardEntry::EPinnedEntryType::HiddenOriginalEntry && SoundEntry.bIsCategory)
+			{
+				for (TSharedPtr<IDashboardDataTreeViewEntry>& ChildEntry : Entry->Children)
+				{
+					MarkBranchAsPinned(ChildEntry, true /*bIsPinned*/);
+					CreatePinnedEntry(ChildEntry);
+				}
+			}
+			else
+			{
+				RecreatePinnedEntries(Entry);
+			}
+		}
+	}
+
 #if WITH_EDITOR
 	void FSoundDashboardViewFactory::BrowseSoundAsset() const
 	{
diff --git a/Engine/Plugins/AudioInsights/Source/AudioInsights/Private/Views/TreeDashboardViewFactory.cpp b/Engine/Plugins/AudioInsights/Source/AudioInsights/Private/Views/TreeDashboardViewFactory.cpp
--- a/Engine/Plugins/AudioInsights/Source/AudioInsights/Private/Views/TreeDashboardViewFactory.cpp
+++ b/Engine/Plugins/AudioInsights/Source/AudioInsights/Private/Views/TreeDashboardViewFactory.cpp
@@ -269,11 +269,7 @@ namespace UE::Audio::Insights
 	void FTraceTreeDashboardViewFactory::RequestSort()
 	{
 		SortTable();
-
-		if (FilteredEntriesListView.IsValid())
-		{
-			FilteredEntriesListView->RequestTreeRefresh();
-		}
+		RefreshFilteredEntriesListView();
 	}
 
 	void FTraceTreeDashboardViewFactory::OnColumnSortModeChanged(const EColumnSortPriority::Type InSortPriority, const FName& InColumnId, const EColumnSortMode::Type InSortMode)
@@ -300,6 +296,26 @@ namespace UE::Audio::Insights
 		FilteredEntriesListView->SetItemExpansion(Item, bIsItemExpanded);
 	}
 
+	void FTraceTreeDashboardViewFactory::TransformDataRecursive(const TSharedPtr<IDashboardDataTreeViewEntry>& OriginalEntry, TFunctionRef<bool(IDashboardDataTreeViewEntry&)> FilterPredicate, const int32 DepthToDisplayFrom)
+	{
+		if (!OriginalEntry.IsValid() || !FilterPredicate(*OriginalEntry))
+		{
+			return;
+		}
+
+		if (DepthToDisplayFrom == 0 || OriginalEntry->Children.IsEmpty())
+		{
+			DataViewEntries.Add(OriginalEntry);
+		}
+		else
+		{
+			for (const TSharedPtr<IDashboardDataTreeViewEntry>& Child : OriginalEntry->Children)
+			{
+				TransformDataRecursive(Child, FilterPredicate, DepthToDisplayFrom - 1);
+			}
+		}
+	}
+
 	void FTraceTreeDashboardViewFactory::Tick(float InElapsed)
 	{
 		for (const TSharedPtr<FTraceProviderBase>& Provider : Providers)
diff --git a/Engine/Plugins/AudioInsights/Source/AudioInsights/Public/Views/SoundDashboardViewFactory.h b/Engine/Plugins/AudioInsights/Source/AudioInsights/Public/Views/SoundDashboardViewFactory.h
--- a/Engine/Plugins/AudioInsights/Source/AudioInsights/Public/Views/SoundDashboardViewFactory.h
+++ b/Engine/Plugins/AudioInsights/Source/AudioInsights/Public/Views/SoundDashboardViewFactory.h
@@ -129,6 +129,8 @@ namespace UE::Audio::Insights
 		virtual TArray<TSharedPtr<IDashboardDataTreeViewEntry>>& GetTreeItemsSource() override { return FullTree; }
 		virtual bool ResetTreeData() override;
 
+		virtual int32 GetDisplayedTreeDepth() const override { return static_cast<int32>(TreeViewingMode); }
+
 		void FilterText();
 		void RecursiveSort(TArray<TSharedPtr<IDashboardDataTreeViewEntry>>& OutTree, TFunctionRef<bool(const FSoundDashboardEntry&, const FSoundDashboardEntry&)> Predicate);
 		void SortByPredicate(TFunctionRef<bool(const FSoundDashboardEntry&, const FSoundDashboardEntry&)> Predicate);
@@ -150,17 +152,22 @@ namespace UE::Audio::Insights
 #endif // WITH_EDITOR
 
 		bool SelectedItemsIncludesAnAsset() const;
+		bool SelectionIncludesUnpinnedItem() const;
+		void ClearSelection();
+
 		void PinSound();
 		void UnpinSound();
-		bool SelectionIncludesUnpinnedItem() const;
 
 		void PinSelectedItems(const TSharedPtr<IDashboardDataTreeViewEntry>& Entry, const TArray<TSharedPtr<IDashboardDataTreeViewEntry>>& SelectedItems);
 		void UnpinSelectedItems(const TSharedPtr<FPinnedSoundEntryWrapper>& PinnedWrapperEntry, const TArray<TSharedPtr<IDashboardDataTreeViewEntry>>& SelectedItems, const bool bSelectionContainsAssets);
 
+		void RecreatePinnedEntries(const TSharedPtr<IDashboardDataTreeViewEntry>& Entry);
+
 		void MarkBranchAsPinned(const TSharedPtr<IDashboardDataTreeViewEntry> Entry, const bool bIsPinned);
 		void InitPinnedItemEntries();
 		void CreatePinnedEntry(TSharedPtr<IDashboardDataTreeViewEntry> Entry);
 		void UpdatePinnedSection();
+		void RebuildPinnedSection();
 
 #if WITH_EDITOR
 		void BrowseSoundAsset() const;
@@ -186,8 +193,16 @@ namespace UE::Audio::Insights
 		TSharedPtr<SWidget> SoundsFilterBar;
 		TSharedPtr<SWidget> SoundsFilterBarButton;
 
+		enum class ETreeViewingOptions : uint8
+		{
+			FullTree = 0,
+			ExcludingCategories
+		};
+
+		ETreeViewingOptions TreeViewingMode = ETreeViewingOptions::FullTree;
 		ESoundDashboardFilterFlags SelectedFilterFlags = AllFilterFlags;
 		bool bIsPinnedCategoryFilterEnabled = true;
 		bool bShowRecentlyStoppedSounds = false;
+		bool bTreeViewModeChanged = false;
 	};
 } // namespace UE::Audio::Insights
diff --git a/Engine/Plugins/AudioInsights/Source/AudioInsights/Public/Views/TreeDashboardViewFactory.h b/Engine/Plugins/AudioInsights/Source/AudioInsights/Public/Views/TreeDashboardViewFactory.h
--- a/Engine/Plugins/AudioInsights/Source/AudioInsights/Public/Views/TreeDashboardViewFactory.h
+++ b/Engine/Plugins/AudioInsights/Source/AudioInsights/Public/Views/TreeDashboardViewFactory.h
@@ -107,6 +107,8 @@ namespace UE::Audio::Insights
 
 		virtual TArray<TSharedPtr<IDashboardDataTreeViewEntry>>& GetTreeItemsSource() { return DataViewEntries; }
 
+		virtual int32 GetDisplayedTreeDepth() const { return 0; }
+
 		virtual bool ResetTreeData()
 		{
 			if (!DataViewEntries.IsEmpty())
@@ -146,7 +148,18 @@ namespace UE::Audio::Insights
 						return InPredicate(*Pair.Value);
 					};
 
-					Algo::TransformIf(*DeviceData, DataViewEntries, FilterEntry, TransformEntry);
+					const int32 DepthToDisplay = GetDisplayedTreeDepth();
+					if (DepthToDisplay == 0)
+					{
+						Algo::TransformIf(*DeviceData, DataViewEntries, FilterEntry, TransformEntry);
+					}
+					else
+					{
+						for (const typename TableProviderType::FEntryPair& Pair : *DeviceData)
+						{
+							TransformDataRecursive(TransformEntry(Pair), InPredicate, DepthToDisplay);
+						}
+					}
 
 					// Sort list
 					RequestSort();
@@ -184,6 +197,8 @@ namespace UE::Audio::Insights
 		void OnColumnSortModeChanged(const EColumnSortPriority::Type InSortPriority, const FName& InColumnId, const EColumnSortMode::Type InSortMode);
 		void HandleRecursiveExpansion(TSharedPtr<IDashboardDataTreeViewEntry> Item, bool bIsItemExpanded);
 
+		void TransformDataRecursive(const TSharedPtr<IDashboardDataTreeViewEntry>& OriginalEntry, TFunctionRef<bool(IDashboardDataTreeViewEntry&)> FilterPredicate, const int32 DepthToDisplayFrom);
+
 		TSharedPtr<SSearchBox> SearchBoxWidget;
 		FText SearchBoxFilterText;
 	};
