[Audio Insights] Fix for loc text conflicts in Sounds Dashboard

diff --git a/Engine/Plugins/AudioInsights/Source/AudioInsights/Private/Views/SoundDashboardViewFactory.cpp b/Engine/Plugins/AudioInsights/Source/AudioInsights/Private/Views/SoundDashboardViewFactory.cpp
--- a/Engine/Plugins/AudioInsights/Source/AudioInsights/Private/Views/SoundDashboardViewFactory.cpp
+++ b/Engine/Plugins/AudioInsights/Source/AudioInsights/Private/Views/SoundDashboardViewFactory.cpp
@@ -1873,7 +1873,7 @@ namespace UE::Audio::Insights
 				SNew(SCheckBox)
 				.HAlign(EHorizontalAlignment::HAlign_Center)
 				.Style(&FSlateStyle::Get().GetWidgetStyle<FCheckBoxStyle>("SoundDashboard.MuteSoloButton"))
-				.ToolTipText(LOCTEXT("SoundsDashboard_ShowCategoriesTooltip", "Shows categories"))
+				.ToolTipText(LOCTEXT("SoundsDashboard_ShowCategoriesTooltip", "Toggle whether categories are visible"))
 				.IsChecked(TreeViewingMode == ETreeViewingOptions::FullTree ? ECheckBoxState::Checked : ECheckBoxState::Unchecked)
 				.OnCheckStateChanged_Lambda([this](ECheckBoxState NewState)
 				{
@@ -1905,7 +1905,7 @@ namespace UE::Audio::Insights
 					.Padding(2.0f, 0.0f)
 					[
 						SNew(STextBlock)
-						.Text(LOCTEXT("SoundsDashboard_ShowStoppedSounds", "Show Categories"))
+						.Text(LOCTEXT("SoundsDashboard_ShowCategories", "Show Categories"))
 					]
 				]
 			]
