[Audio Insights] Added dB/Lin toggle for Amp (Peak) values in Sound and Mixer Source dashboards (also applies for MixerSource plot Amp (Peak) values). For Sound dashboard we put this option in a new Settings area, Show Stopped Sound option was moved there too + misc. small style adjustments.

diff --git a/Engine/Plugins/AudioInsights/Source/AudioInsights/Private/Views/MixerSourceDashboardViewFactory.cpp b/Engine/Plugins/AudioInsights/Source/AudioInsights/Private/Views/MixerSourceDashboardViewFactory.cpp
--- a/Engine/Plugins/AudioInsights/Source/AudioInsights/Private/Views/MixerSourceDashboardViewFactory.cpp
+++ b/Engine/Plugins/AudioInsights/Source/AudioInsights/Private/Views/MixerSourceDashboardViewFactory.cpp
@@ -4,6 +4,7 @@
 #include "Async/Async.h"
 #include "AudioInsightsModule.h"
 #include "AudioInsightsStyle.h"
+#include "AudioInsightsUtils.h"
 #include "DSP/Dsp.h"
 #include "Features/IModularFeatures.h"
 #include "Framework/Docking/LayoutService.h"
@@ -15,6 +16,7 @@
 #include "Templates/SharedPointer.h"
 #include "UObject/SoftObjectPath.h"
 #include "Widgets/Input/SComboBox.h"
+#include "Widgets/Input/SSegmentedControl.h"
 
 #if WITH_EDITOR
 #include "Editor.h"
@@ -111,7 +113,7 @@ namespace UE::Audio::Insights
 
 	const TMap<FName, FTraceTableDashboardViewFactory::FColumnData>& FMixerSourceDashboardViewFactory::GetColumns() const
 	{
-		auto CreateColumnData = []()
+		auto CreateColumnData = [this]()
 		{
 			return TMap<FName, FTraceTableDashboardViewFactory::FColumnData>
 			{
@@ -137,10 +139,14 @@ namespace UE::Audio::Insights
 					"Amplitude",
 					{
 						LOCTEXT("Source_EnvColumnDisplayName", "Amp (Peak)"),
-						[](const IDashboardDataViewEntry& InData)
+						[this](const IDashboardDataViewEntry& InData)
 						{
 							const ::Audio::TCircularAudioBuffer<FDataPoint>& EnvelopeDataPoints = MixerSourcePrivate::CastEntry(InData).EnvelopeDataPoints;
-							return FText::AsNumber(MixerSourcePrivate::GetLastEntryArrayValue(EnvelopeDataPoints), FSlateStyle::Get().GetAmpFloatFormat());
+							const float Amplitude = MixerSourcePrivate::GetLastEntryArrayValue(EnvelopeDataPoints);
+
+							return bDisplayAmpPeakInDb 
+								? AudioInsightsUtils::ConvertToDecibelsText(Amplitude)
+								: FText::AsNumber(Amplitude, FSlateStyle::Get().GetAmpFloatFormat());
 						},
 						false /* bDefaultHidden */,
 						0.12f /* FillWidth */
@@ -544,17 +550,33 @@ namespace UE::Audio::Insights
 
 					const ::Audio::DisjointedArrayView<const FDataPoint> TimeStampedValuesDisjointedArrayView = TimeStampedValues.PeekInPlace(TimeStampedValues.Num());
 					
-					for (const auto& [Timestamp, Value] : TimeStampedValuesDisjointedArrayView.FirstBuffer)
+					if (bDisplayAmpPeakInDb)
 					{
-						CurrentTimestamp = FMath::Max(CurrentTimestamp, Timestamp);
+						for (const auto& [Timestamp, Value] : TimeStampedValuesDisjointedArrayView.FirstBuffer)
+						{
+							CurrentTimestamp = FMath::Max(CurrentTimestamp, Timestamp);
 
 #if WITH_EDITOR
-						const double DataPointTime = Timestamp - BeginTimestamp;
+							const double DataPointTime = Timestamp - BeginTimestamp;
 #else
-						const double DataPointTime = Timestamp;
+							const double DataPointTime = Timestamp;
 #endif // WITH_EDITOR
+							DataPoints->Emplace(DataPointTime, ::Audio::ConvertToDecibels(Value));
+						}
+					}
+					else
+					{
+						for (const auto& [Timestamp, Value] : TimeStampedValuesDisjointedArrayView.FirstBuffer)
+						{
+							CurrentTimestamp = FMath::Max(CurrentTimestamp, Timestamp);
 
-						DataPoints->Emplace(DataPointTime, Value);
+#if WITH_EDITOR
+							const double DataPointTime = Timestamp - BeginTimestamp;
+#else
+							const double DataPointTime = Timestamp;
+#endif // WITH_EDITOR
+							DataPoints->Emplace(DataPointTime, Value);
+						}
 					}
 				}
 			}
@@ -920,7 +942,7 @@ namespace UE::Audio::Insights
 			+ SHorizontalBox::Slot()
 			.AutoWidth()
 			.HAlign(EHorizontalAlignment::HAlign_Center)
-			.Padding(0.0f, 10.0f, 0.0f, 0.0f)
+			.Padding(4.0f, 10.0f, 0.0f, 0.0f)
 			[
 				SAssignNew(MuteToggleButton, SCheckBox)
 				.HAlign(EHorizontalAlignment::HAlign_Center)
@@ -969,6 +991,41 @@ namespace UE::Audio::Insights
 	}
 #endif // WITH_EDITOR
 
+	
+	TSharedRef<SWidget> FMixerSourceDashboardViewFactory::MakeAmpDisplayModeWidget()
+	{
+		return SNew(SHorizontalBox)
+			+ SHorizontalBox::Slot()
+			.AutoWidth()
+			[
+				SNew(STextBlock)
+				.Margin(FMargin(0.0, 6.0, 5.0, 0.0))
+				.Text(LOCTEXT("SoundsDashboard_AmpPeakDisplayMode_Text", "Amp (Peak) Display Mode:"))
+				.ToolTipText(LOCTEXT("SoundsDashboard_AmpPeakDisplayMode_ToolTipText", "Switches the amplitude display mode between decibels or linear scale."))
+			]
+			+ SHorizontalBox::Slot()
+			.AutoWidth()
+			[
+				SNew(SSegmentedControl<bool>)
+				.OnValueChanged_Lambda([this](bool bInValue)
+				{
+					bDisplayAmpPeakInDb = bInValue;
+				})
+				.Value_Lambda([this]
+				{
+					return bDisplayAmpPeakInDb;
+				})
+				// Decibels mode
+				+SSegmentedControl<bool>::Slot(true)
+				.Text(LOCTEXT("SoundsDashboard_AmpDisplayMode_dB_Text", "dB"))
+				.ToolTip(LOCTEXT("SoundsDashboard_AmpDisplayMode_dB_ToolTipText", "Displays amplitude values in decibels."))
+				// Linear mode
+				+ SSegmentedControl<bool>::Slot(false)
+				.Text(LOCTEXT("SoundsDashboard_AmpDisplayMode_Linear_Text", "Lin"))
+				.ToolTip(LOCTEXT("SoundsDashboard_AmpDisplayMode_Linear_Text", "Displays amplitude values in linear scale."))
+			];
+	}
+
 	TSharedRef<SWidget> FMixerSourceDashboardViewFactory::MakePlotsButtonWidget()
 	{
 		return SAssignNew(PlotsButton, SCheckBox)
@@ -1043,6 +1100,14 @@ namespace UE::Audio::Insights
 					+ SHorizontalBox::Slot()
 					.AutoWidth()
 					.HAlign(HAlign_Right)
+					.Padding(0.0f, 6.0f, 30.0f, 0.0f)
+					[
+						MakeAmpDisplayModeWidget()
+					]
+					+ SHorizontalBox::Slot()
+					.AutoWidth()
+					.HAlign(HAlign_Right)
+					.Padding(0.0f, 6.0f, 0.0f, 0.0f)
 					[
 						MakePlotsButtonWidget()
 					]
diff --git a/Engine/Plugins/AudioInsights/Source/AudioInsights/Private/Views/SoundDashboardViewFactory.cpp b/Engine/Plugins/AudioInsights/Source/AudioInsights/Private/Views/SoundDashboardViewFactory.cpp
--- a/Engine/Plugins/AudioInsights/Source/AudioInsights/Private/Views/SoundDashboardViewFactory.cpp
+++ b/Engine/Plugins/AudioInsights/Source/AudioInsights/Private/Views/SoundDashboardViewFactory.cpp
@@ -4,11 +4,13 @@
 #include "Algo/Accumulate.h"
 #include "AudioInsightsModule.h"
 #include "AudioInsightsStyle.h"
+#include "AudioInsightsUtils.h"
 #include "Messages/SoundTraceMessages.h"
 #include "Misc/EnumClassFlags.h"
 #include "Providers/SoundTraceProvider.h"
 #include "SoundDashboardCommands.h"
 #include "Views/SAudioFilterBar.h"
+#include "Widgets/Input/SSegmentedControl.h"
 
 #if WITH_EDITOR
 #include "Audio/AudioDebug.h"
@@ -744,12 +746,12 @@ namespace UE::Audio::Insights
 			// Empty Spacing
 			+ SHorizontalBox::Slot()
 			.FillWidth(1.0f)
-			// Show recently stopped sounds button
+			// Settings button
 			+ SHorizontalBox::Slot()
 			.HAlign(HAlign_Right)
 			.AutoWidth()
 			[
-				MakeShowRecentlyStoppedSoundsWidget()
+				MakeSettingsButtonWidget()
 			];
 	}
 #endif // WITH_EDITOR
@@ -892,7 +894,7 @@ namespace UE::Audio::Insights
 			.HAlign(HAlign_Right)
 			.Padding(0.0f, 2.0f, 0.0f, 0.0f)
 			[
-				MakeShowRecentlyStoppedSoundsWidget()
+				MakeSettingsButtonWidget()
 			]
 #endif // WITH_EDITOR
 			+ SVerticalBox::Slot()
@@ -1530,7 +1532,7 @@ namespace UE::Audio::Insights
 			{
 				"Amplitude",
 				{
-					.GetDisplayValue = [](const IDashboardDataTreeViewEntry& InData)
+					.GetDisplayValue = [this](const IDashboardDataTreeViewEntry& InData)
 					{
 						const FSoundDashboardEntry& SoundDashboardEntry = CastEntry(InData);
 						if (SoundDashboardEntry.bIsCategory)
@@ -1538,7 +1540,11 @@ namespace UE::Audio::Insights
 							return FText::GetEmpty();
 						}
 
-						return FText::AsNumber(GetLastEntryArrayValue(SoundDashboardEntry.AmplitudeDataPoints), FSlateStyle::Get().GetAmpFloatFormat());
+						const float AmplitudeValue = GetLastEntryArrayValue(SoundDashboardEntry.AmplitudeDataPoints);
+
+						return bDisplayAmpPeakInDb 
+							? AudioInsightsUtils::ConvertToDecibelsText(AmplitudeValue) 
+							: FText::AsNumber(AmplitudeValue, FSlateStyle::Get().GetAmpFloatFormat());
 					}
 				}
 			},
@@ -1861,92 +1867,158 @@ namespace UE::Audio::Insights
 		RecursiveSort(DataViewEntries, Predicate);
 	}
 
-	TSharedRef<SWidget> FSoundDashboardViewFactory::MakeShowRecentlyStoppedSoundsWidget()
+	TSharedRef<SWidget> FSoundDashboardViewFactory::OnGetSettingsMenuContent()
 	{
-		using namespace FSoundDashboardViewFactoryPrivate;
+		FMenuBuilder MenuBuilder(false /*bShouldCloseWindowAfterMenuSelection*/, nullptr /*CommandList*/);
 
-		return SNew(SHorizontalBox)
-			+ SHorizontalBox::Slot()
-			.AutoWidth()
-			.Padding(8.0f, 8.0f)
-			[
-				SNew(SCheckBox)
-				.HAlign(EHorizontalAlignment::HAlign_Center)
-				.Style(&FSlateStyle::Get().GetWidgetStyle<FCheckBoxStyle>("SoundDashboard.MuteSoloButton"))
-				.ToolTipText(LOCTEXT("SoundsDashboard_ShowCategoriesTooltip", "Toggle whether categories are visible"))
-				.IsChecked(TreeViewingMode == ETreeViewingOptions::FullTree ? ECheckBoxState::Checked : ECheckBoxState::Unchecked)
-				.OnCheckStateChanged_Lambda([this](ECheckBoxState NewState)
+		MenuBuilder.BeginSection("SoundDashboardSettingActions", LOCTEXT("SoundDashboard_SettingActions_HeaderText", "Settings"));
+
+		{
+			// Amp (Peak) Display Mode Display Mode toggle
+			MenuBuilder.AddMenuEntry(
+				FUIAction(FExecuteAction::CreateLambda([this]()
+				{
+					bDisplayAmpPeakInDb = !bDisplayAmpPeakInDb;
+				})),
+				SNew(SHorizontalBox)
+				+ SHorizontalBox::Slot()
+				.AutoWidth()
+				[
+					SNew(STextBlock)
+					.Margin(FMargin(0.0, 4.0, 5.0, 0.0))
+					.Text(LOCTEXT("SoundsDashboard_AmpPeakDisplayMode_Text", "Amp (Peak) Display Mode"))
+				]
+				+ SHorizontalBox::Slot()
+				.AutoWidth()
+				[
+					SNew(SSegmentedControl<bool>)
+					.OnValueChanged_Lambda([this](bool bInValue)
+					{
+						bDisplayAmpPeakInDb = bInValue;
+					})
+					.Value_Lambda([this]
+					{
+						return bDisplayAmpPeakInDb;
+					})
+					// Decibels mode
+					+SSegmentedControl<bool>::Slot(true)
+					.Text(LOCTEXT("SoundsDashboard_AmpDisplayMode_dB_Text", "dB"))
+					.ToolTip(LOCTEXT("SoundsDashboard_AmpDisplayMode_dB_ToolTipText", "Displays amplitude values in decibels."))
+					// Linear mode
+					+ SSegmentedControl<bool>::Slot(false)
+					.Text(LOCTEXT("SoundsDashboard_AmpDisplayMode_Linear_Text", "Lin"))
+					.ToolTip(LOCTEXT("SoundsDashboard_AmpDisplayMode_Linear_Text", "Displays amplitude values in linear scale."))
+				],
+				NAME_None /*InExtensionHook*/,
+				LOCTEXT("SoundsDashboard_AmpPeakDisplayMode_ToolTipText", "Switches the amplitude display mode between decibels or linear scale."),
+				EUserInterfaceActionType::CollapsedButton
+			);
+
+			// Show Stopped Sounds checkbox
+			MenuBuilder.AddMenuEntry(
+				FUIAction(FExecuteAction::CreateLambda([this]()
 				{
-					TreeViewingMode = NewState == ECheckBoxState::Checked ? ETreeViewingOptions::FullTree : ETreeViewingOptions::ExcludingCategories;
+					bShowRecentlyStoppedSounds = !bShowRecentlyStoppedSounds;
+				})),
+				SNew(SHorizontalBox)
+				+ SHorizontalBox::Slot()
+				.HAlign(HAlign_Right)
+				.VAlign(VAlign_Center)
+				.AutoWidth()
+				[
+					SNew(STextBlock)
+					.Text(LOCTEXT("SoundsDashboard_ShowStoppedSounds", "Show Stopped Sounds"))
+				]
+				+ SHorizontalBox::Slot()
+				.HAlign(HAlign_Left)
+				.VAlign(VAlign_Center)
+				.AutoWidth()
+				.Padding(27.0, 0.0f, 0.0f, 0.0f)
+				[
+					SNew(SImage)
+				 	.Image_Lambda([this]()
+				 	{
+				 		const FName IconName = bShowRecentlyStoppedSounds ? "AudioInsights.Icon.SoundDashboard.Visible" : "AudioInsights.Icon.SoundDashboard.Invisible";
+		 
+				 		return FSlateStyle::Get().GetBrush(IconName);
+				 	})
+				],
+				NAME_None /*InExtensionHook*/,
+				LOCTEXT("SoundsDashboard_ShowStoppedSoundsTooltip", "Shows sounds that have recently stopped playing"),
+				EUserInterfaceActionType::CollapsedButton
+			);
+
+			// Show/Hide categories checkbox
+			MenuBuilder.AddMenuEntry(
+				FUIAction(FExecuteAction::CreateLambda([this]()
+				{
+					if (TreeViewingMode == ETreeViewingOptions::FullTree)
+					{
+						TreeViewingMode = ETreeViewingOptions::ExcludingCategories;
+					}
+					else if (TreeViewingMode == ETreeViewingOptions::ExcludingCategories)
+					{
+						TreeViewingMode = ETreeViewingOptions::FullTree;
+					}
+
 					UpdateFilterReason = EProcessReason::FilterUpdated;
 					ClearSelection();
 
 					bTreeViewModeChanged = true;
-				})
+				})),
+				SNew(SHorizontalBox)
+				+ SHorizontalBox::Slot()
+				.HAlign(HAlign_Right)
+				.VAlign(VAlign_Center)
+				.AutoWidth()
 				[
-					SNew(SHorizontalBox)
-					+ SHorizontalBox::Slot()
-					.HAlign(HAlign_Left)
-					.VAlign(VAlign_Center)
-					.AutoWidth()
-					.Padding(2.0f, 0.0f)
-					[
-						SNew(SImage)
-				 		.Image_Lambda([this]()
-				 		{
-				 			const FName IconName = TreeViewingMode == ETreeViewingOptions::FullTree ? "AudioInsights.Icon.SoundDashboard.Visible" : "AudioInsights.Icon.SoundDashboard.Invisible";
-				 			return FSlateStyle::Get().GetBrush(IconName);
-				 		})
-					]
-					+ SHorizontalBox::Slot()
-					.HAlign(HAlign_Right)
-					.VAlign(VAlign_Center)
-					.AutoWidth()
-					.Padding(2.0f, 0.0f)
-					[
-						SNew(STextBlock)
-						.Text(LOCTEXT("SoundsDashboard_ShowCategories", "Show Categories"))
-					]
+					SNew(STextBlock)
+					.Text(LOCTEXT("SoundsDashboard_ShowCategories", "Show Categories"))
 				]
-			]
-			+ SHorizontalBox::Slot()
-			.AutoWidth()
-			.Padding(8.0f, 8.0f)
+				+ SHorizontalBox::Slot()
+				.HAlign(HAlign_Left)
+				.VAlign(VAlign_Center)
+				.AutoWidth()
+				.Padding(27.0, 0.0f, 0.0f, 0.0f)
+				[
+					SNew(SImage)
+				 	.Image_Lambda([this]()
+				 	{
+				 		const FName IconName = TreeViewingMode == ETreeViewingOptions::FullTree ? "AudioInsights.Icon.SoundDashboard.Visible" : "AudioInsights.Icon.SoundDashboard.Invisible";
+				 		
+						return FSlateStyle::Get().GetBrush(IconName);
+				 	})
+				],
+				NAME_None /*InExtensionHook*/,
+				LOCTEXT("SoundsDashboard_ShowCategoriesTooltip", "Toggle whether categories are visible"),
+				EUserInterfaceActionType::CollapsedButton
+			);
+		}
+
+		MenuBuilder.EndSection();
+
+		return MenuBuilder.MakeWidget();
+	}
+
+	TSharedRef<SWidget> FSoundDashboardViewFactory::MakeSettingsButtonWidget()
+	{
+		return SNew(SComboButton)
+			.ComboButtonStyle(&FAppStyle::Get().GetWidgetStyle<FComboButtonStyle>("SimpleComboButton"))
+			.OnGetMenuContent(this, &FSoundDashboardViewFactory::OnGetSettingsMenuContent)
+			.MenuPlacement(EMenuPlacement::MenuPlacement_ComboBoxRight)
+			.HasDownArrow(false)
+			.ToolTipText(LOCTEXT("SoundsDashboard_Settings_TooltipText", "Sound dashboard settings"))
+			.ButtonContent()
 			[
-				SNew(SCheckBox)
-				.HAlign(EHorizontalAlignment::HAlign_Center)
-				.Style(&FSlateStyle::Get().GetWidgetStyle<FCheckBoxStyle>("SoundDashboard.MuteSoloButton"))
-				.ToolTipText(LOCTEXT("SoundsDashboard_ShowStoppedSoundsTooltip", "Shows sounds that have recently stopped playing"))
-				.IsChecked(bShowRecentlyStoppedSounds ? ECheckBoxState::Checked : ECheckBoxState::Unchecked)
-				.OnCheckStateChanged_Lambda([this](ECheckBoxState NewState)
-				{
-					bShowRecentlyStoppedSounds = NewState == ECheckBoxState::Checked;
-				})
+				SNew(SHorizontalBox)
+				+ SHorizontalBox::Slot()
+				.AutoWidth()
+				.VAlign(VAlign_Center)
+				.Padding(5.0, 0.0f)
 				[
-					SNew(SHorizontalBox)
-					+ SHorizontalBox::Slot()
-					.HAlign(HAlign_Left)
-					.VAlign(VAlign_Center)
-					.AutoWidth()
-					.Padding(2.0f, 0.0f)
-					[
-						SNew(SImage)
-				 		.Image_Lambda([this]()
-				 		{
-				 			const FName IconName = bShowRecentlyStoppedSounds ? "AudioInsights.Icon.SoundDashboard.Visible" : "AudioInsights.Icon.SoundDashboard.Invisible";
-		 
-				 			return FSlateStyle::Get().GetBrush(IconName);
-				 		})
-					]
-					+ SHorizontalBox::Slot()
-					.HAlign(HAlign_Right)
-					.VAlign(VAlign_Center)
-					.AutoWidth()
-					.Padding(2.0f, 0.0f)
-					[
-						SNew(STextBlock)
-						.Text(LOCTEXT("SoundsDashboard_ShowStoppedSounds", "Show Stopped"))
-					]
+					SNew(SImage)
+					.ColorAndOpacity(FSlateColor::UseForeground())
+					.Image(FSlateStyle::Get().GetBrush("AudioInsights.Icon.Settings"))
 				]
 			];
 	}
diff --git a/Engine/Plugins/AudioInsights/Source/AudioInsights/Public/AudioInsightsStyle.h b/Engine/Plugins/AudioInsights/Source/AudioInsights/Public/AudioInsightsStyle.h
--- a/Engine/Plugins/AudioInsights/Source/AudioInsights/Public/AudioInsightsStyle.h
+++ b/Engine/Plugins/AudioInsights/Source/AudioInsights/Public/AudioInsightsStyle.h
@@ -147,7 +147,7 @@ namespace UE::Audio::Insights
 			Set("AudioInsights.Icon.Event", new IMAGE_BRUSH_SVG(TEXT("Icons/audio_event"), Icon16));
 			Set("AudioInsights.Icon.Log", new IMAGE_BRUSH_SVG(TEXT("Icons/audio_log"), Icon16));
 			Set("AudioInsights.Icon.Sources", new IMAGE_BRUSH_SVG(TEXT("Icons/audio_sources"), Icon16));
-			Set("AudioInsights.Icon.Sources.Plots", new IMAGE_BRUSH(TEXT("Icons/audio_sources_plots"), Icon24));
+			Set("AudioInsights.Icon.Sources.Plots", new IMAGE_BRUSH_SVG(TEXT("Icons/audio_sources_plots"), Icon20));
 			Set("AudioInsights.Icon.Submix", new IMAGE_BRUSH_SVG(TEXT("Icons/audio_submix"), Icon16));
 			Set("AudioInsights.Icon.VirtualLoop", new IMAGE_BRUSH_SVG(TEXT("Icons/audio_virtualloop"), Icon16));
 			Set("AudioInsights.Icon.Viewport", new IMAGE_BRUSH_SVG(TEXT("Icons/viewport"), Icon16));
@@ -157,6 +157,7 @@ namespace UE::Audio::Insights
 			Set("AudioInsights.Icon.Start.Inactive", new IMAGE_BRUSH_SVG(TEXT("Icons/start_inactive"), Icon20));
 			Set("AudioInsights.Icon.Record.Active", new IMAGE_BRUSH_SVG(TEXT("Icons/record_active"), Icon20));
 			Set("AudioInsights.Icon.Record.Inactive", new IMAGE_BRUSH_SVG(TEXT("Icons/record_inactive"), Icon20));
+			Set("AudioInsights.Icon.Settings", new IMAGE_BRUSH_SVG(TEXT("Icons/settings"), Icon20));
 			Set("AudioInsights.Icon.Stop.Active", new IMAGE_BRUSH_SVG(TEXT("Icons/stop_active"), Icon20));
 			Set("AudioInsights.Icon.Stop.Inactive", new IMAGE_BRUSH_SVG(TEXT("Icons/stop_inactive"), Icon20));
 
diff --git a/Engine/Plugins/AudioInsights/Source/AudioInsights/Public/AudioInsightsUtils.h b/Engine/Plugins/AudioInsights/Source/AudioInsights/Public/AudioInsightsUtils.h
new file mode 100644
--- /dev/null
+++ b/Engine/Plugins/AudioInsights/Source/AudioInsights/Public/AudioInsightsUtils.h
@@ -0,0 +1,29 @@
+// Copyright Epic Games, Inc. All Rights Reserved.
+#pragma once
+
+#include "AudioInsightsStyle.h"
+#include "DSP/Dsp.h"
+#include "Internationalization/Text.h"
+#include "Math/UnrealMathUtility.h"
+
+namespace UE::Audio::Insights
+{
+	class AUDIOINSIGHTS_API AudioInsightsUtils
+	{
+	public:
+		static FText ConvertToDecibelsText(const float InAmplitude)
+		{
+			if (FMath::IsNearlyEqual(InAmplitude, 1.0f))
+			{
+				return FText::AsNumber(0.0f, FSlateStyle::Get().GetAmpFloatFormat());
+			}
+			else if (FMath::IsNearlyZero(InAmplitude))
+			{
+				static const FText MinusInifinityDbText = INVTEXT("-∞");
+				return MinusInifinityDbText;
+			}
+
+			return FText::AsNumber(::Audio::ConvertToDecibels(InAmplitude), FSlateStyle::Get().GetAmpFloatFormat());
+		}
+	};
+} // namespace UE::Audio::Insights
diff --git a/Engine/Plugins/AudioInsights/Source/AudioInsights/Public/Views/MixerSourceDashboardViewFactory.h b/Engine/Plugins/AudioInsights/Source/AudioInsights/Public/Views/MixerSourceDashboardViewFactory.h
--- a/Engine/Plugins/AudioInsights/Source/AudioInsights/Public/Views/MixerSourceDashboardViewFactory.h
+++ b/Engine/Plugins/AudioInsights/Source/AudioInsights/Public/Views/MixerSourceDashboardViewFactory.h
@@ -89,6 +89,7 @@ namespace UE::Audio::Insights
 		void SaveLayoutToConfig();
 #endif // WITH EDITOR
 
+		TSharedRef<SWidget> MakeAmpDisplayModeWidget();
 		TSharedRef<SWidget> MakePlotsButtonWidget();
 		void TogglePlotsTabVisibility(ECheckBoxState InCheckboxState);
 
@@ -137,6 +138,8 @@ namespace UE::Audio::Insights
 		TSharedPtr<FTabManager> MixerSourcesTabManager;
 		TSharedPtr<FWorkspaceItem> MixerSourcesWorkspace;
 
+		bool bDisplayAmpPeakInDb = true;
+
 #if WITH_EDITOR
 		// State of the mute and solo buttons
 		ECheckBoxState MuteState = ECheckBoxState::Unchecked;
diff --git a/Engine/Plugins/AudioInsights/Source/AudioInsights/Public/Views/SoundDashboardViewFactory.h b/Engine/Plugins/AudioInsights/Source/AudioInsights/Public/Views/SoundDashboardViewFactory.h
--- a/Engine/Plugins/AudioInsights/Source/AudioInsights/Public/Views/SoundDashboardViewFactory.h
+++ b/Engine/Plugins/AudioInsights/Source/AudioInsights/Public/Views/SoundDashboardViewFactory.h
@@ -135,7 +135,8 @@ namespace UE::Audio::Insights
 		void RecursiveSort(TArray<TSharedPtr<IDashboardDataTreeViewEntry>>& OutTree, TFunctionRef<bool(const FSoundDashboardEntry&, const FSoundDashboardEntry&)> Predicate);
 		void SortByPredicate(TFunctionRef<bool(const FSoundDashboardEntry&, const FSoundDashboardEntry&)> Predicate);
 
-		TSharedRef<SWidget> MakeShowRecentlyStoppedSoundsWidget();
+		TSharedRef<SWidget> OnGetSettingsMenuContent();
+		TSharedRef<SWidget> MakeSettingsButtonWidget();
 
 #if WITH_EDITOR
 		TSharedRef<SWidget> MakeMuteSoloWidget();
@@ -204,5 +205,6 @@ namespace UE::Audio::Insights
 		bool bIsPinnedCategoryFilterEnabled = true;
 		bool bShowRecentlyStoppedSounds = false;
 		bool bTreeViewModeChanged = false;
+		bool bDisplayAmpPeakInDb = true;
 	};
 } // namespace UE::Audio::Insights
