[Audio Insights] Renamed class from FRewindDebugger to FRewindDebuggerAudioInsightsRuntime

diff --git a/Engine/Plugins/AudioInsights/Source/AudioInsights/Private/AudioInsightsModule.cpp b/Engine/Plugins/AudioInsights/Source/AudioInsights/Private/AudioInsightsModule.cpp
--- a/Engine/Plugins/AudioInsights/Source/AudioInsights/Private/AudioInsightsModule.cpp
+++ b/Engine/Plugins/AudioInsights/Source/AudioInsights/Private/AudioInsightsModule.cpp
@@ -31,7 +31,7 @@ namespace UE::Audio::Insights
 		if (!IsRunningCommandlet())
 		{
 			TraceModule = MakeUnique<FTraceModule>();
-			RewindDebuggerExtension = MakeUnique<FRewindDebugger>();
+			RewindDebuggerExtension = MakeUnique<FRewindDebuggerAudioInsightsRuntime>();
 
 			IModularFeatures::Get().RegisterModularFeature(TraceServices::ModuleFeatureName, TraceModule.Get());
 			IModularFeatures::Get().RegisterModularFeature(IRewindDebuggerRuntimeExtension::ModularFeatureName, RewindDebuggerExtension.Get());
diff --git a/Engine/Plugins/AudioInsights/Source/AudioInsights/Private/AudioInsightsModule.h b/Engine/Plugins/AudioInsights/Source/AudioInsights/Private/AudioInsightsModule.h
--- a/Engine/Plugins/AudioInsights/Source/AudioInsights/Private/AudioInsightsModule.h
+++ b/Engine/Plugins/AudioInsights/Source/AudioInsights/Private/AudioInsightsModule.h
@@ -51,7 +51,7 @@ namespace UE::Audio::Insights
 	private:
 		TSharedPtr<FDashboardFactory> DashboardFactory;
 		TUniquePtr<FTraceModule> TraceModule;
-		TUniquePtr<FRewindDebugger> RewindDebuggerExtension;
+		TUniquePtr<FRewindDebuggerAudioInsightsRuntime> RewindDebuggerExtension;
 
 #if !WITH_EDITOR
 		TSharedPtr<FAudioInsightsComponent> AudioInsightsComponent;
diff --git a/Engine/Plugins/AudioInsights/Source/AudioInsights/Private/AudioInsightsTraceModule.cpp b/Engine/Plugins/AudioInsights/Source/AudioInsights/Private/AudioInsightsTraceModule.cpp
--- a/Engine/Plugins/AudioInsights/Source/AudioInsights/Private/AudioInsightsTraceModule.cpp
+++ b/Engine/Plugins/AudioInsights/Source/AudioInsights/Private/AudioInsightsTraceModule.cpp
@@ -24,7 +24,7 @@ namespace UE::Audio::Insights
 		static const TCHAR* AudioMixerChannelName = TEXT("AudioMixer");
 	}
 
-	void FRewindDebugger::RecordingStarted()
+	void FRewindDebuggerAudioInsightsRuntime::RecordingStarted()
 	{
 		using namespace FTraceModulePrivate;
 
diff --git a/Engine/Plugins/AudioInsights/Source/AudioInsights/Public/AudioInsightsTraceModule.h b/Engine/Plugins/AudioInsights/Source/AudioInsights/Public/AudioInsightsTraceModule.h
--- a/Engine/Plugins/AudioInsights/Source/AudioInsights/Public/AudioInsightsTraceModule.h
+++ b/Engine/Plugins/AudioInsights/Source/AudioInsights/Public/AudioInsightsTraceModule.h
@@ -13,7 +13,7 @@
 
 namespace UE::Audio::Insights
 {
-	class FRewindDebugger : public IRewindDebuggerRuntimeExtension
+	class FRewindDebuggerAudioInsightsRuntime : public IRewindDebuggerRuntimeExtension
 	{
 	public:
 		virtual void RecordingStarted() override;
@@ -79,7 +79,7 @@ namespace UE::Audio::Insights
 
 		TMap<FName, TSharedPtr<FTraceProviderBase>> TraceProviders;
 		TArray<FString> ChannelsToRestore;
-		FRewindDebugger RewindDebugger;
+		FRewindDebuggerAudioInsightsRuntime RewindDebugger;
 
 		bool bTraceAnalysisHasStarted = false;
 		bool bStopTracingAfterAudioInsightsIsFinished = false;
