[Audio Insights] Remove platform restrictions on AudioInsights.uplugin

diff --git a/Engine/Plugins/AudioInsights/AudioInsights.uplugin b/Engine/Plugins/AudioInsights/AudioInsights.uplugin
--- a/Engine/Plugins/AudioInsights/AudioInsights.uplugin
+++ b/Engine/Plugins/AudioInsights/AudioInsights.uplugin
@@ -24,22 +24,12 @@
 			"ProgramAllowList":
 			[
 				"UnrealInsights"
-			],
-			"PlatformAllowList":
-			[
-				"Win64",
-				"Mac"
 			]
 		},
 		{
 			"Name": "AudioInsightsEditor",
 			"Type": "Editor",
-			"LoadingPhase": "Default",
-			"PlatformAllowList":
-			[
-				"Win64",
-				"Mac"
-			]
+			"LoadingPhase": "Default"
 		}
 	],
 	"Plugins":
