[Audio Insights] AudioInsightsStyle.h: added missing FSlateStyleRegistry::UnRegisterSlateStyle in destructor (potential fix for: Opening Audio Insights causes Editor to crash on Mac)

diff --git a/Engine/Plugins/AudioInsights/Source/AudioInsights/Public/AudioInsightsStyle.h b/Engine/Plugins/AudioInsights/Source/AudioInsights/Public/AudioInsightsStyle.h
--- a/Engine/Plugins/AudioInsights/Source/AudioInsights/Public/AudioInsightsStyle.h
+++ b/Engine/Plugins/AudioInsights/Source/AudioInsights/Public/AudioInsightsStyle.h
@@ -215,6 +215,11 @@ namespace UE::Audio::Insights
 
 			FSlateStyleRegistry::RegisterSlateStyle(*this);
 		}
+
+		~FSlateStyle()
+		{
+			FSlateStyleRegistry::UnRegisterSlateStyle(*this);
+		}
 	};
 } // namespace UE::Audio::Insights
 #undef LOCTEXT_NAMESPACE
