[Audio Insights] Sources dashboard: cache name for dashboard columns + only apply dB transform to Amplitude plot

diff --git a/Engine/Plugins/AudioInsights/Source/AudioInsights/Private/Views/MixerSourceDashboardViewFactory.cpp b/Engine/Plugins/AudioInsights/Source/AudioInsights/Private/Views/MixerSourceDashboardViewFactory.cpp
--- a/Engine/Plugins/AudioInsights/Source/AudioInsights/Private/Views/MixerSourceDashboardViewFactory.cpp
+++ b/Engine/Plugins/AudioInsights/Source/AudioInsights/Private/Views/MixerSourceDashboardViewFactory.cpp
@@ -56,6 +56,14 @@ namespace UE::Audio::Insights
 		const FName MixerSourcesTableTabName = "MixerSourcesTableTab";
 		const FName MixerSourcesPlotsTabName = "MixerSourcesPlotsTab";
 
+		const FName PlayOrderColumnName = "PlayOrder";
+		const FName NameColumnName = "Name";
+		const FName AmplitudeColumnName = "Amplitude";
+		const FName VolumeColumnName = "Volume";
+		const FName DistanceAttenuationColumnName = "DistanceAttenuation";
+		const FName PitchColumnName = "Pitch";
+		const FName LPFColumnName = "LPF";
+		const FName HPFColumnName = "HPF";
 	} // namespace MixerSourcePrivate
 
 	constexpr double FMixerSourceDashboardViewFactory::MaxPlotHistorySeconds = 5.0;
@@ -113,12 +121,14 @@ namespace UE::Audio::Insights
 
 	const TMap<FName, FTraceTableDashboardViewFactory::FColumnData>& FMixerSourceDashboardViewFactory::GetColumns() const
 	{
+		using namespace MixerSourcePrivate;
+
 		auto CreateColumnData = [this]()
 		{
 			return TMap<FName, FTraceTableDashboardViewFactory::FColumnData>
 			{
 				{
-					"PlayOrder",
+					PlayOrderColumnName,
 					{
 						LOCTEXT("PlayOrder_PlayOrderColumnDisplayName", "Play Order"),
 						[](const IDashboardDataViewEntry& InData) { return FText::AsNumber(MixerSourcePrivate::CastEntry(InData).PlayOrder); },
@@ -127,7 +137,7 @@ namespace UE::Audio::Insights
 					}
 				},
 				{
-					"Name",
+					NameColumnName,
 					{
 						LOCTEXT("Source_NameColumnDisplayName", "Name"),
 						[](const IDashboardDataViewEntry& InData) { return FText::FromString(*FSoftObjectPath(MixerSourcePrivate::CastEntry(InData).Name).GetAssetName()); },
@@ -136,7 +146,7 @@ namespace UE::Audio::Insights
 					}
 				},
 				{
-					"Amplitude",
+					AmplitudeColumnName,
 					{
 						LOCTEXT("Source_EnvColumnDisplayName", "Amp (Peak)"),
 						[this](const IDashboardDataViewEntry& InData)
@@ -153,7 +163,7 @@ namespace UE::Audio::Insights
 					}
 				},
 				{
-					"Volume",
+					VolumeColumnName,
 					{
 						LOCTEXT("Source_VolumeColumnDisplayName", "Volume"),
 						[](const IDashboardDataViewEntry& InData)
@@ -166,7 +176,7 @@ namespace UE::Audio::Insights
 					}
 				},
 				{
-					"DistanceAttenuation",
+					DistanceAttenuationColumnName,
 					{
 						LOCTEXT("Source_AttenuationColumnDisplayName", "Distance Attenuation"),
 						[](const IDashboardDataViewEntry& InData)
@@ -179,7 +189,7 @@ namespace UE::Audio::Insights
 					}
 				},
 				{
-					"Pitch",
+					PitchColumnName,
 					{
 						LOCTEXT("Source_PitchColumnDisplayName", "Pitch"),
 						[](const IDashboardDataViewEntry& InData)
@@ -192,7 +202,7 @@ namespace UE::Audio::Insights
 					}
 				},
 				{
-					"LPF",
+					LPFColumnName,
 					{
 						LOCTEXT("Source_LPFColumnDisplayName", "LPF Freq (Hz)"),
 						[](const IDashboardDataViewEntry& InData)
@@ -205,7 +215,7 @@ namespace UE::Audio::Insights
 					}
 				},
 				{
-					"HPF",
+					HPFColumnName,
 					{
 						LOCTEXT("Source_HPFColumnDisplayName", "HPF Freq (Hz)"),
 						[](const IDashboardDataViewEntry& InData)
@@ -225,7 +235,9 @@ namespace UE::Audio::Insights
 
 	void FMixerSourceDashboardViewFactory::SortTable()
 	{
-		if (SortByColumn == "PlayOrder")
+		using namespace MixerSourcePrivate;
+
+		if (SortByColumn == PlayOrderColumnName)
 		{
 			if (SortMode == EColumnSortMode::Ascending)
 			{
@@ -248,7 +260,7 @@ namespace UE::Audio::Insights
 				});
 			}
 		}
-		else if (SortByColumn == "Name")
+		else if (SortByColumn == NameColumnName)
 		{
 			if (SortMode == EColumnSortMode::Ascending)
 			{
@@ -271,7 +283,7 @@ namespace UE::Audio::Insights
 				});
 			}
 		}
-		else if (SortByColumn == "Amplitude")
+		else if (SortByColumn == AmplitudeColumnName)
 		{
 			if (SortMode == EColumnSortMode::Ascending)
 			{
@@ -294,7 +306,7 @@ namespace UE::Audio::Insights
 				});
 			}
 		}
-		else if (SortByColumn == "Volume")
+		else if (SortByColumn == VolumeColumnName)
 		{
 			if (SortMode == EColumnSortMode::Ascending)
 			{
@@ -317,7 +329,7 @@ namespace UE::Audio::Insights
 				});
 			}
 		}
-		else if (SortByColumn == "DistanceAttenuation")
+		else if (SortByColumn == DistanceAttenuationColumnName)
 		{
 			if (SortMode == EColumnSortMode::Ascending)
 			{
@@ -340,7 +352,7 @@ namespace UE::Audio::Insights
 				});
 			}
 		}
-		else if (SortByColumn == "Pitch")
+		else if (SortByColumn == PitchColumnName)
 		{
 			if (SortMode == EColumnSortMode::Ascending)
 			{
@@ -363,7 +375,7 @@ namespace UE::Audio::Insights
 				});
 			}
 		}
-		else if (SortByColumn == "LPF")
+		else if (SortByColumn == LPFColumnName)
 		{
 			if (SortMode == EColumnSortMode::Ascending)
 			{
@@ -386,7 +398,7 @@ namespace UE::Audio::Insights
 				});
 			}
 		}
-		else if (SortByColumn == "HPF")
+		else if (SortByColumn == HPFColumnName)
 		{
 			if (SortMode == EColumnSortMode::Ascending)
 			{
@@ -512,6 +524,8 @@ namespace UE::Audio::Insights
 
 	void FMixerSourceDashboardViewFactory::UpdatePlotsWidgetsData()
 	{
+		using namespace MixerSourcePrivate;
+
 		if (!PlotWidgetMetadataPerCurve.IsValid() || DataViewEntries.Num() <= 0)
 		{
 			return;
@@ -549,8 +563,8 @@ namespace UE::Audio::Insights
 					const ::Audio::TCircularAudioBuffer<FDataPoint>& TimeStampedValues = (DataFunc)(SourceDataPoint);
 
 					const ::Audio::DisjointedArrayView<const FDataPoint> TimeStampedValuesDisjointedArrayView = TimeStampedValues.PeekInPlace(TimeStampedValues.Num());
-					
-					if (bDisplayAmpPeakInDb)
+
+					if (bDisplayAmpPeakInDb && ColumnName == AmplitudeColumnName)
 					{
 						for (const auto& [Timestamp, Value] : TimeStampedValuesDisjointedArrayView.FirstBuffer)
 						{
@@ -658,19 +672,21 @@ namespace UE::Audio::Insights
 
 	const TMap<FName, FMixerSourceDashboardViewFactory::FPlotColumnInfo>& FMixerSourceDashboardViewFactory::GetPlotColumnInfo()
 	{
+		using namespace MixerSourcePrivate;
+
 		auto CreatePlotColumnInfo = []()
 		{
 			return TMap<FName, FMixerSourceDashboardViewFactory::FPlotColumnInfo>
 			{
 				{
-					"Amplitude",
+					AmplitudeColumnName,
 					{
 						[](const IDashboardDataViewEntry& InData) -> const ::Audio::TCircularAudioBuffer<FDataPoint>& { return MixerSourcePrivate::CastEntry(InData).EnvelopeDataPoints; },
 						FSlateStyle::Get().GetAmpFloatFormat()
 					}
 				},
 				{
-					"Volume",
+					VolumeColumnName,
 					{
 
 						[](const IDashboardDataViewEntry& InData) -> const ::Audio::TCircularAudioBuffer<FDataPoint>& { return MixerSourcePrivate::CastEntry(InData).VolumeDataPoints; },
@@ -678,28 +694,28 @@ namespace UE::Audio::Insights
 					}
 				},
 				{
-					"DistanceAttenuation",
+					DistanceAttenuationColumnName,
 					{
 						[](const IDashboardDataViewEntry& InData) -> const ::Audio::TCircularAudioBuffer<FDataPoint>& { return MixerSourcePrivate::CastEntry(InData).DistanceAttenuationDataPoints; },
 						FSlateStyle::Get().GetAmpFloatFormat()
 					}
 				},
 				{
-					"Pitch",
+					PitchColumnName,
 					{
 						[](const IDashboardDataViewEntry& InData) -> const ::Audio::TCircularAudioBuffer<FDataPoint>& { return MixerSourcePrivate::CastEntry(InData).PitchDataPoints; },
 						FSlateStyle::Get().GetPitchFloatFormat()
 					}
 				},
 				{
-					"LPF",
+					LPFColumnName,
 					{
 						[](const IDashboardDataViewEntry& InData) -> const ::Audio::TCircularAudioBuffer<FDataPoint>& { return MixerSourcePrivate::CastEntry(InData).LPFFreqDataPoints; },
 						FSlateStyle::Get().GetFreqFloatFormat()
 					}
 				},
 				{
-					"HPF",
+					HPFColumnName,
 					{
 						[](const IDashboardDataViewEntry& InData) -> const ::Audio::TCircularAudioBuffer<FDataPoint>& { return MixerSourcePrivate::CastEntry(InData).HPFFreqDataPoints; },
 						FSlateStyle::Get().GetFreqFloatFormat()
