[Audio Insights] Sounds Dashboard - allow user to view dashboard as just a list of soundwaves (i.e. just the children). Move viewing options to sub menu inside the settings and added keyboard shortcuts.

diff --git a/Engine/Plugins/AudioInsights/Source/AudioInsights/Private/SoundDashboardCommands.cpp b/Engine/Plugins/AudioInsights/Source/AudioInsights/Private/SoundDashboardCommands.cpp
--- a/Engine/Plugins/AudioInsights/Source/AudioInsights/Private/SoundDashboardCommands.cpp
+++ b/Engine/Plugins/AudioInsights/Source/AudioInsights/Private/SoundDashboardCommands.cpp
@@ -21,6 +21,10 @@ namespace UE::Audio::Insights
 		UI_COMMAND(Unpin, "Unpin", "Removes the selected sound from the Pinned category.", EUserInterfaceActionType::Button, FInputChord(EKeys::U, EModifierKey::Control));
 		UI_COMMAND(Browse, "Browse To Asset", "Browses to the selected sound asset in the content browser.", EUserInterfaceActionType::Button, FInputChord(EKeys::B, EModifierKey::Control));
 		UI_COMMAND(Edit, "Edit", "Opens the selected sound for edit.", EUserInterfaceActionType::Button, FInputChord(EKeys::E, EModifierKey::Control));
+
+		UI_COMMAND(ViewFullTree, "Tree View", "Organize sounds into categories.", EUserInterfaceActionType::RadioButton, FInputChord(EKeys::F1, EModifierKey::Control));
+		UI_COMMAND(ViewActiveSounds, "Active Sounds", "Organize sounds into Active Sounds.", EUserInterfaceActionType::RadioButton, FInputChord(EKeys::F2, EModifierKey::Control));
+		UI_COMMAND(ViewFlatList, "Flat List", "Display sounds as individual waves in a flat list.", EUserInterfaceActionType::RadioButton, FInputChord(EKeys::F3, EModifierKey::Control));
 		// @TODO UE-250399: Hide category pending to implement
 		//UI_COMMAND(Hide, "Hide", "Hides the selected sound.", EUserInterfaceActionType::Button, FInputChord(EKeys::H, EModifierKey::Control));
 	}
diff --git a/Engine/Plugins/AudioInsights/Source/AudioInsights/Private/SoundDashboardCommands.h b/Engine/Plugins/AudioInsights/Source/AudioInsights/Private/SoundDashboardCommands.h
--- a/Engine/Plugins/AudioInsights/Source/AudioInsights/Private/SoundDashboardCommands.h
+++ b/Engine/Plugins/AudioInsights/Source/AudioInsights/Private/SoundDashboardCommands.h
@@ -19,6 +19,9 @@ namespace UE::Audio::Insights
 		TSharedPtr<const FUICommandInfo> GetUnpinCommand() const  { return Unpin; }
 		TSharedPtr<const FUICommandInfo> GetBrowseCommand() const { return Browse; }
 		TSharedPtr<const FUICommandInfo> GetEditCommand() const   { return Edit; }
+		TSharedPtr<const FUICommandInfo> GetViewFullTreeCommand() const { return ViewFullTree; }
+		TSharedPtr<const FUICommandInfo> GetViewActiveSoundsCommand() const { return ViewActiveSounds; }
+		TSharedPtr<const FUICommandInfo> GetViewFlatList() const { return ViewFlatList; }
 		// @TODO UE-250399: Hide category pending to implement
 		//TSharedPtr<const FUICommandInfo> GetHideCommand() const   { return Hide; }
 		
@@ -27,6 +30,9 @@ namespace UE::Audio::Insights
 		TSharedPtr<FUICommandInfo> Unpin;
 		TSharedPtr<FUICommandInfo> Browse;
 		TSharedPtr<FUICommandInfo> Edit;
+		TSharedPtr<FUICommandInfo> ViewFullTree;
+		TSharedPtr<FUICommandInfo> ViewActiveSounds;
+		TSharedPtr<FUICommandInfo> ViewFlatList;
 		// @TODO UE-250399: Hide category pending to implement
 		//TSharedPtr<FUICommandInfo> Hide;
 	};
diff --git a/Engine/Plugins/AudioInsights/Source/AudioInsights/Private/Views/SoundDashboardViewFactory.cpp b/Engine/Plugins/AudioInsights/Source/AudioInsights/Private/Views/SoundDashboardViewFactory.cpp
--- a/Engine/Plugins/AudioInsights/Source/AudioInsights/Private/Views/SoundDashboardViewFactory.cpp
+++ b/Engine/Plugins/AudioInsights/Source/AudioInsights/Private/Views/SoundDashboardViewFactory.cpp
@@ -632,6 +632,24 @@ namespace UE::Audio::Insights
 		CommandList->MapAction(Commands.GetEditCommand(),   FExecuteAction::CreateRaw(this, &FSoundDashboardViewFactory::OpenSoundAsset), FCanExecuteAction::CreateRaw(this, &FSoundDashboardViewFactory::SelectedItemsIncludesAnAsset));
 #endif // WITH_EDITOR
 
+		CommandList->MapAction(
+			Commands.GetViewFullTreeCommand(),
+			FExecuteAction::CreateLambda([this]() { ChangeTreeViewingMode(ETreeViewingOptions::FullTree); }),
+			FCanExecuteAction(),
+			FIsActionChecked::CreateLambda([this]() { return TreeViewingMode == ETreeViewingOptions::FullTree; }));
+
+		CommandList->MapAction(
+			Commands.GetViewActiveSoundsCommand(),
+			FExecuteAction::CreateLambda([this]() { ChangeTreeViewingMode(ETreeViewingOptions::ActiveSounds); }),
+			FCanExecuteAction(),
+			FIsActionChecked::CreateLambda([this]() { return TreeViewingMode == ETreeViewingOptions::ActiveSounds; }));
+
+		CommandList->MapAction(
+			Commands.GetViewFlatList(),
+			FExecuteAction::CreateLambda([this]() { ChangeTreeViewingMode(ETreeViewingOptions::FlatList); }),
+			FCanExecuteAction(),
+			FIsActionChecked::CreateLambda([this]() { return TreeViewingMode == ETreeViewingOptions::FlatList; }));
+
 		// @TODO UE-250399: Hide category pending to implement
 		//CommandList->MapAction(Commands.GetHideCommand(), FExecuteAction::CreateRaw(this, &FSoundDashboardViewFactory::HideSound));
 	}
@@ -1211,6 +1229,29 @@ namespace UE::Audio::Insights
 		});
 	}
 
+	void FSoundDashboardViewFactory::ChangeTreeViewingMode(const ETreeViewingOptions SelectedMode)
+	{
+		if (TreeViewingMode == SelectedMode)
+		{
+			return;
+		}
+
+		TreeViewingMode = SelectedMode;
+		UpdateFilterReason = EProcessReason::FilterUpdated;
+		ClearSelection();
+
+		bTreeViewModeChanged = true;
+	}
+
+	void FSoundDashboardViewFactory::BuildViewOptionsMenuContent(FMenuBuilder& MenuBuilder)
+	{
+		const FSoundDashboardCommands& Commands = FSoundDashboardCommands::Get();
+
+		MenuBuilder.AddMenuEntry(Commands.GetViewFullTreeCommand());
+		MenuBuilder.AddMenuEntry(Commands.GetViewActiveSoundsCommand());
+		MenuBuilder.AddMenuEntry(Commands.GetViewFlatList());
+	}
+
 	void FSoundDashboardViewFactory::FilterText()
 	{
 		using namespace FSoundDashboardViewFactoryPrivate;
@@ -1869,7 +1910,7 @@ namespace UE::Audio::Insights
 
 	TSharedRef<SWidget> FSoundDashboardViewFactory::OnGetSettingsMenuContent()
 	{
-		FMenuBuilder MenuBuilder(false /*bShouldCloseWindowAfterMenuSelection*/, nullptr /*CommandList*/);
+		FMenuBuilder MenuBuilder(false /*bShouldCloseWindowAfterMenuSelection*/, CommandList);
 
 		MenuBuilder.BeginSection("SoundDashboardSettingActions", LOCTEXT("SoundDashboard_SettingActions_HeaderText", "Settings"));
 
@@ -1948,50 +1989,13 @@ namespace UE::Audio::Insights
 				EUserInterfaceActionType::CollapsedButton
 			);
 
-			// Show/Hide categories checkbox
-			MenuBuilder.AddMenuEntry(
-				FUIAction(FExecuteAction::CreateLambda([this]()
-				{
-					if (TreeViewingMode == ETreeViewingOptions::FullTree)
-					{
-						TreeViewingMode = ETreeViewingOptions::ExcludingCategories;
-					}
-					else if (TreeViewingMode == ETreeViewingOptions::ExcludingCategories)
-					{
-						TreeViewingMode = ETreeViewingOptions::FullTree;
-					}
-
-					UpdateFilterReason = EProcessReason::FilterUpdated;
-					ClearSelection();
-
-					bTreeViewModeChanged = true;
-				})),
-				SNew(SHorizontalBox)
-				+ SHorizontalBox::Slot()
-				.HAlign(HAlign_Right)
-				.VAlign(VAlign_Center)
-				.AutoWidth()
-				[
-					SNew(STextBlock)
-					.Text(LOCTEXT("SoundsDashboard_ShowCategories", "Show Categories"))
-				]
-				+ SHorizontalBox::Slot()
-				.HAlign(HAlign_Left)
-				.VAlign(VAlign_Center)
-				.AutoWidth()
-				.Padding(27.0, 0.0f, 0.0f, 0.0f)
-				[
-					SNew(SImage)
-				 	.Image_Lambda([this]()
-				 	{
-				 		const FName IconName = TreeViewingMode == ETreeViewingOptions::FullTree ? "AudioInsights.Icon.SoundDashboard.Visible" : "AudioInsights.Icon.SoundDashboard.Invisible";
-				 		
-						return FSlateStyle::Get().GetBrush(IconName);
-				 	})
-				],
-				NAME_None /*InExtensionHook*/,
-				LOCTEXT("SoundsDashboard_ShowCategoriesTooltip", "Toggle whether categories are visible"),
-				EUserInterfaceActionType::CollapsedButton
+			// View options
+			MenuBuilder.AddSubMenu
+			(
+				LOCTEXT("SoundsDashboard_ViewOptions", "View"),
+				LOCTEXT("SoundsDashboard_ViewOptionsTooltip", "Choose how items in the Sounds Dashboard is displayed."),
+				FNewMenuDelegate::CreateRaw(this, &FSoundDashboardViewFactory::BuildViewOptionsMenuContent),
+				false
 			);
 		}
 
@@ -2406,9 +2410,10 @@ namespace UE::Audio::Insights
 		}
 		else
 		{
+			const bool bSoundEntryIsPinned = SoundEntry.PinnedEntryType == FSoundDashboardEntry::EPinnedEntryType::HiddenOriginalEntry;
 			for (const TSharedPtr<IDashboardDataTreeViewEntry>& Child : Entry->Children)
 			{
-				if (SoundEntry.PinnedEntryType == FSoundDashboardEntry::EPinnedEntryType::HiddenOriginalEntry)
+				if (bSoundEntryIsPinned)
 				{
 					if (TreeViewingMode == ETreeViewingOptions::FullTree && SoundEntry.bIsCategory)
 					{
@@ -2424,7 +2429,19 @@ namespace UE::Audio::Insights
 				}
 				else
 				{
-					RecreatePinnedEntries(Child);
+					const FSoundDashboardEntry& ChildEntry = FSoundDashboardViewFactoryPrivate::CastEntry(*Child);
+					const bool bChildEntryIsPinned = ChildEntry.PinnedEntryType == FSoundDashboardEntry::EPinnedEntryType::HiddenOriginalEntry;
+
+					if (bChildEntryIsPinned && TreeViewingMode != ETreeViewingOptions::FlatList && !SoundEntry.bIsCategory)
+					{
+						MarkBranchAsPinned(Entry, true /*bIsPinned*/);
+						CreatePinnedEntry(Entry);
+						break;
+					}
+					else
+					{
+						RecreatePinnedEntries(Child);
+					}
 				}
 			}
 		}
diff --git a/Engine/Plugins/AudioInsights/Source/AudioInsights/Private/Views/TreeDashboardViewFactory.cpp b/Engine/Plugins/AudioInsights/Source/AudioInsights/Private/Views/TreeDashboardViewFactory.cpp
--- a/Engine/Plugins/AudioInsights/Source/AudioInsights/Private/Views/TreeDashboardViewFactory.cpp
+++ b/Engine/Plugins/AudioInsights/Source/AudioInsights/Private/Views/TreeDashboardViewFactory.cpp
@@ -296,14 +296,17 @@ namespace UE::Audio::Insights
 		FilteredEntriesListView->SetItemExpansion(Item, bIsItemExpanded);
 	}
 
-	void FTraceTreeDashboardViewFactory::TransformDataRecursive(const TSharedPtr<IDashboardDataTreeViewEntry>& OriginalEntry, TFunctionRef<bool(IDashboardDataTreeViewEntry&)> FilterPredicate, const int32 DepthToDisplayFrom)
+	void FTraceTreeDashboardViewFactory::TransformDataRecursive(const TSharedPtr<IDashboardDataTreeViewEntry>& OriginalEntry, TFunctionRef<bool(IDashboardDataTreeViewEntry&)> FilterPredicate, const int32 DepthToDisplayFrom, const bool bParentPassedFilter /*= true*/)
 	{
-		if (!OriginalEntry.IsValid() || !FilterPredicate(*OriginalEntry))
+		if (!OriginalEntry.IsValid())
 		{
 			return;
 		}
 
-		if (DepthToDisplayFrom == 0 || OriginalEntry->Children.IsEmpty())
+		const bool bPassesFilter = DepthToDisplayFrom >= 0 ? FilterPredicate(*OriginalEntry) : bParentPassedFilter;
+		const bool bIsAtValidTreeDepth = DepthToDisplayFrom == 0 || OriginalEntry->Children.IsEmpty();
+
+		if (bIsAtValidTreeDepth && bPassesFilter)
 		{
 			DataViewEntries.Add(OriginalEntry);
 		}
@@ -311,7 +314,7 @@ namespace UE::Audio::Insights
 		{
 			for (const TSharedPtr<IDashboardDataTreeViewEntry>& Child : OriginalEntry->Children)
 			{
-				TransformDataRecursive(Child, FilterPredicate, DepthToDisplayFrom - 1);
+				TransformDataRecursive(Child, FilterPredicate, DepthToDisplayFrom - 1, bPassesFilter);
 			}
 		}
 	}
diff --git a/Engine/Plugins/AudioInsights/Source/AudioInsights/Public/Views/SoundDashboardViewFactory.h b/Engine/Plugins/AudioInsights/Source/AudioInsights/Public/Views/SoundDashboardViewFactory.h
--- a/Engine/Plugins/AudioInsights/Source/AudioInsights/Public/Views/SoundDashboardViewFactory.h
+++ b/Engine/Plugins/AudioInsights/Source/AudioInsights/Public/Views/SoundDashboardViewFactory.h
@@ -131,6 +131,16 @@ namespace UE::Audio::Insights
 
 		virtual int32 GetDisplayedTreeDepth() const override { return static_cast<int32>(TreeViewingMode); }
 
+		enum class ETreeViewingOptions : uint8
+		{
+			FullTree = 0,
+			ActiveSounds,
+			FlatList
+		};
+
+		void ChangeTreeViewingMode(const ETreeViewingOptions SelectedMode);
+		void BuildViewOptionsMenuContent(FMenuBuilder& MenuBuilder);
+
 		void FilterText();
 		void RecursiveSort(TArray<TSharedPtr<IDashboardDataTreeViewEntry>>& OutTree, TFunctionRef<bool(const FSoundDashboardEntry&, const FSoundDashboardEntry&)> Predicate);
 		void SortByPredicate(TFunctionRef<bool(const FSoundDashboardEntry&, const FSoundDashboardEntry&)> Predicate);
@@ -194,12 +204,6 @@ namespace UE::Audio::Insights
 		TSharedPtr<SWidget> SoundsFilterBar;
 		TSharedPtr<SWidget> SoundsFilterBarButton;
 
-		enum class ETreeViewingOptions : uint8
-		{
-			FullTree = 0,
-			ExcludingCategories
-		};
-
 		ETreeViewingOptions TreeViewingMode = ETreeViewingOptions::FullTree;
 		ESoundDashboardFilterFlags SelectedFilterFlags = AllFilterFlags;
 		bool bIsPinnedCategoryFilterEnabled = true;
diff --git a/Engine/Plugins/AudioInsights/Source/AudioInsights/Public/Views/TreeDashboardViewFactory.h b/Engine/Plugins/AudioInsights/Source/AudioInsights/Public/Views/TreeDashboardViewFactory.h
--- a/Engine/Plugins/AudioInsights/Source/AudioInsights/Public/Views/TreeDashboardViewFactory.h
+++ b/Engine/Plugins/AudioInsights/Source/AudioInsights/Public/Views/TreeDashboardViewFactory.h
@@ -197,7 +197,7 @@ namespace UE::Audio::Insights
 		void OnColumnSortModeChanged(const EColumnSortPriority::Type InSortPriority, const FName& InColumnId, const EColumnSortMode::Type InSortMode);
 		void HandleRecursiveExpansion(TSharedPtr<IDashboardDataTreeViewEntry> Item, bool bIsItemExpanded);
 
-		void TransformDataRecursive(const TSharedPtr<IDashboardDataTreeViewEntry>& OriginalEntry, TFunctionRef<bool(IDashboardDataTreeViewEntry&)> FilterPredicate, const int32 DepthToDisplayFrom);
+		void TransformDataRecursive(const TSharedPtr<IDashboardDataTreeViewEntry>& OriginalEntry, TFunctionRef<bool(IDashboardDataTreeViewEntry&)> FilterPredicate, const int32 DepthToDisplayFrom, const bool bParentPassedFilter = true);
 
 		TSharedPtr<SSearchBox> SearchBoxWidget;
 		FText SearchBoxFilterText;
