[Audio Insights] Sound and Sources dashboard: fixed LOCTEXT conflict

diff --git a/Engine/Plugins/AudioInsights/Source/AudioInsights/Private/Views/MixerSourceDashboardViewFactory.cpp b/Engine/Plugins/AudioInsights/Source/AudioInsights/Private/Views/MixerSourceDashboardViewFactory.cpp
--- a/Engine/Plugins/AudioInsights/Source/AudioInsights/Private/Views/MixerSourceDashboardViewFactory.cpp
+++ b/Engine/Plugins/AudioInsights/Source/AudioInsights/Private/Views/MixerSourceDashboardViewFactory.cpp
@@ -1016,8 +1016,8 @@ namespace UE::Audio::Insights
 			[
 				SNew(STextBlock)
 				.Margin(FMargin(0.0, 6.0, 5.0, 0.0))
-				.Text(LOCTEXT("SoundsDashboard_AmpPeakDisplayMode_Text", "Amp (Peak) Display Mode:"))
-				.ToolTipText(LOCTEXT("SoundsDashboard_AmpPeakDisplayMode_ToolTipText", "Switches the amplitude display mode between decibels or linear scale."))
+				.Text(LOCTEXT("SourcesDashboard_AmpPeakDisplayMode_Text", "Amp (Peak) Display Mode:"))
+				.ToolTipText(LOCTEXT("SourcesDashboard_AmpPeakDisplayMode_ToolTipText", "Switches the amplitude display mode between decibels or linear scale."))
 			]
 			+ SHorizontalBox::Slot()
 			.AutoWidth()
@@ -1033,12 +1033,12 @@ namespace UE::Audio::Insights
 				})
 				// Decibels mode
 				+SSegmentedControl<bool>::Slot(true)
-				.Text(LOCTEXT("SoundsDashboard_AmpDisplayMode_dB_Text", "dB"))
-				.ToolTip(LOCTEXT("SoundsDashboard_AmpDisplayMode_dB_ToolTipText", "Displays amplitude values in decibels."))
+				.Text(LOCTEXT("SourcesDashboard_AmpDisplayMode_dB_Text", "dB"))
+				.ToolTip(LOCTEXT("SourcesDashboard_AmpDisplayMode_dB_ToolTipText", "Displays amplitude values in decibels."))
 				// Linear mode
 				+ SSegmentedControl<bool>::Slot(false)
-				.Text(LOCTEXT("SoundsDashboard_AmpDisplayMode_Linear_Text", "Lin"))
-				.ToolTip(LOCTEXT("SoundsDashboard_AmpDisplayMode_Linear_Text", "Displays amplitude values in linear scale."))
+				.Text(LOCTEXT("SourcesDashboard_AmpDisplayMode_Linear_Text", "Lin"))
+				.ToolTip(LOCTEXT("SourcesDashboard_AmpDisplayMode_Linear_TooltipText", "Displays amplitude values in linear scale."))
 			];
 	}
 
diff --git a/Engine/Plugins/AudioInsights/Source/AudioInsights/Private/Views/SoundDashboardViewFactory.cpp b/Engine/Plugins/AudioInsights/Source/AudioInsights/Private/Views/SoundDashboardViewFactory.cpp
--- a/Engine/Plugins/AudioInsights/Source/AudioInsights/Private/Views/SoundDashboardViewFactory.cpp
+++ b/Engine/Plugins/AudioInsights/Source/AudioInsights/Private/Views/SoundDashboardViewFactory.cpp
@@ -1948,7 +1948,7 @@ namespace UE::Audio::Insights
 					// Linear mode
 					+ SSegmentedControl<bool>::Slot(false)
 					.Text(LOCTEXT("SoundsDashboard_AmpDisplayMode_Linear_Text", "Lin"))
-					.ToolTip(LOCTEXT("SoundsDashboard_AmpDisplayMode_Linear_Text", "Displays amplitude values in linear scale."))
+					.ToolTip(LOCTEXT("SoundsDashboard_AmpDisplayMode_Linear_TooltipText", "Displays amplitude values in linear scale."))
 				],
 				NAME_None /*InExtensionHook*/,
 				LOCTEXT("SoundsDashboard_AmpPeakDisplayMode_ToolTipText", "Switches the amplitude display mode between decibels or linear scale."),
