Reverted change from 43132987 that was merged wrong

diff --git a/Engine/Plugins/AudioInsights/Source/AudioInsights/Public/Views/MixerSourceDashboardViewFactory.h b/Engine/Plugins/AudioInsights/Source/AudioInsights/Public/Views/MixerSourceDashboardViewFactory.h
--- a/Engine/Plugins/AudioInsights/Source/AudioInsights/Public/Views/MixerSourceDashboardViewFactory.h
+++ b/Engine/Plugins/AudioInsights/Source/AudioInsights/Public/Views/MixerSourceDashboardViewFactory.h
@@ -7,100 +7,99 @@
 #include "Views/TableDashboardViewFactory.h"
 #include "Widgets/Input/SCheckBox.h"
 
-#define UE_API AUDIOINSIGHTS_API
-
 namespace UE::Audio::Insights
 {
-	class FMixerSourceDashboardViewFactory : public FTraceObjectTableDashboardViewFactory
+	class AUDIOINSIGHTS_API FMixerSourceDashboardViewFactory : public FTraceObjectTableDashboardViewFactory
 	{
 	public:
-		UE_API FMixerSourceDashboardViewFactory();
+		FMixerSourceDashboardViewFactory();
 		struct FPlotColumnInfo
 		{
 			const TFunctionRef<const ::Audio::TCircularAudioBuffer<FDataPoint>&(const IDashboardDataViewEntry& InData)> DataFunc;
 			const FNumberFormattingOptions* FormatOptions;
 		};
 
-		UE_API virtual ~FMixerSourceDashboardViewFactory();
+		virtual ~FMixerSourceDashboardViewFactory();
 
-		UE_API virtual FName GetName() const override;
-		UE_API virtual FText GetDisplayName() const override;
-		UE_API virtual FSlateIcon GetIcon() const override;
-		UE_API virtual EDefaultDashboardTabStack GetDefaultTabStack() const override;
-		UE_API virtual void ProcessEntries(FTraceTableDashboardViewFactory::EProcessReason Reason) override;
-		UE_API virtual TSharedRef<SWidget> MakeWidget(TSharedRef<SDockTab> OwnerTab, const FSpawnTabArgs& SpawnTabArgs) override;
+		virtual FName GetName() const override;
+		virtual FText GetDisplayName() const override;
+		virtual FSlateIcon GetIcon() const override;
+		virtual EDefaultDashboardTabStack GetDefaultTabStack() const override;
+		virtual void ProcessEntries(FTraceTableDashboardViewFactory::EProcessReason Reason) override;
+		virtual TSharedRef<SWidget> MakeWidget(TSharedRef<SDockTab> OwnerTab, const FSpawnTabArgs& SpawnTabArgs) override;
 
 		// Maximum amount of data history kept for plots (in seconds)
-		static UE_API const double MaxPlotHistorySeconds;
+		static const double MaxPlotHistorySeconds;
 		// Maximum number of sources to plot at once 
-		static UE_API const int32 MaxPlotSources; 
+		static const int32 MaxPlotSources; 
 
 #if WITH_EDITOR
 		DECLARE_MULTICAST_DELEGATE_ThreeParams(FOnUpdateMuteSoloState, ECheckBoxState /*MuteState*/, ECheckBoxState /*SoloState*/, const FString& /*CurrentFilterString*/);
-		UE_API inline static FOnUpdateMuteSoloState OnUpdateMuteSoloState;
+		inline static FOnUpdateMuteSoloState OnUpdateMuteSoloState;
 #endif // WITH_EDITOR
 
 	protected:
-		UE_API virtual const TMap<FName, FTraceTableDashboardViewFactory::FColumnData>& GetColumns() const override;
-		UE_API virtual void SortTable() override;
+		virtual const TMap<FName, FTraceTableDashboardViewFactory::FColumnData>& GetColumns() const override;
+		virtual void SortTable() override;
 
-		UE_API TSharedRef<SWidget> MakePlotsWidget();
+		TSharedRef<SWidget> MakePlotsWidget();
 
 	private:
 		using FPlotCurvePoint = SAudioCurveView::FCurvePoint;
 		using FPointDataPerCurveMap = TMap<int32, TArray<FPlotCurvePoint>>; // Map of source id to data point array 
 		using FPlotCurveMetadata = SAudioCurveView::FCurveMetadata;
 
-		UE_API void OnAnalysisStarting(const double Timestamp);
+		void OnAnalysisStarting(const double Timestamp);
 
 #if WITH_EDITOR
-		UE_API void OnPIEStarted(bool bSimulating);
-		UE_API void OnPIEStopped(bool bSimulating);
-		UE_API void OnPIEPaused(bool bSimulating);
-		UE_API void OnPIEResumed(bool bSimulating);
+		void OnPIEStarted(bool bSimulating);
+		void OnPIEStopped(bool bSimulating);
+		void OnPIEPaused(bool bSimulating);
+		void OnPIEResumed(bool bSimulating);
 #else
-		UE_API void OnAudioInsightsComponentTabSpawn();
-		UE_API void OnSessionAnalysisCompleted();
-		UE_API void OnTimingViewTimeMarkerChanged(double InTimeMarker);
+		void OnAudioInsightsComponentTabSpawn();
+		void OnSessionAnalysisCompleted();
+		void OnTimingViewTimeMarkerChanged(double InTimeMarker);
 #endif // WITH_EDITOR
 
-		UE_API void ResetPlots();
-		UE_API void UpdatePlotsWidgetsData();
+		void ResetPlots();
+		void UpdatePlotsWidgetsData();
 
 #if !WITH_EDITOR
-		UE_API void FilterOfflinePlots(const FString& InFilterString, TMap<int32, FPlotCurveMetadata>& InPlotWidgetMetadataPerCurve);
+		void FilterOfflinePlots(const FString& InFilterString, TMap<int32, FPlotCurveMetadata>& InPlotWidgetMetadataPerCurve);
 #endif // !WITH_EDITOR
 
 #if WITH_EDITOR
-		UE_API void UpdateMuteSoloState();
+		void UpdateMuteSoloState();
 #endif // WITH_EDITOR
 
 		// Column information used by plot widgets, keyed by column name. These keys should be a subset of the keys in GetColumns(). 
-		UE_API const TMap<FName, FPlotColumnInfo>& GetPlotColumnInfo();
-		UE_API const TFunctionRef<const ::Audio::TCircularAudioBuffer<FDataPoint>&(const IDashboardDataViewEntry& InData)> GetPlotColumnDataFunc(const FName& ColumnName);
-		UE_API const FNumberFormattingOptions* GetPlotColumnNumberFormat(const FName& ColumnName);
-		UE_API const FText GetPlotColumnDisplayName(const FName& ColumnName);
+		const TMap<FName, FPlotColumnInfo>& GetPlotColumnInfo();
+		const TFunctionRef<const ::Audio::TCircularAudioBuffer<FDataPoint>&(const IDashboardDataViewEntry& InData)> GetPlotColumnDataFunc(const FName& ColumnName);
+		const FNumberFormattingOptions* GetPlotColumnNumberFormat(const FName& ColumnName);
+		const FText GetPlotColumnDisplayName(const FName& ColumnName);
 
 #if WITH_EDITOR
-		UE_API TSharedRef<SWidget> MakeMuteSoloWidget();
+		TSharedRef<SWidget> MakeMuteSoloWidget();
 
-		UE_API void ToggleMuteForAllItems(ECheckBoxState NewState);
-		UE_API void ToggleSoloForAllItems(ECheckBoxState NewState);
+		void ToggleMuteForAllItems(ECheckBoxState NewState);
+		void ToggleSoloForAllItems(ECheckBoxState NewState);
 
-		UE_API TSharedRef<FTabManager::FLayout> LoadLayoutFromConfig();
-		UE_API void SaveLayoutToConfig();
+		TSharedRef<FTabManager::FLayout> LoadLayoutFromConfig();
+		void SaveLayoutToConfig();
 #endif // WITH EDITOR
 
-		UE_API TSharedRef<SWidget> MakePlotsButtonWidget();
-		UE_API void TogglePlotsTabVisibility(ECheckBoxState InCheckboxState);
+		TSharedRef<SWidget> MakeAmpDisplayModeWidget();
+		TSharedRef<SWidget> MakePlotsButtonWidget();
+		void TogglePlotsTabVisibility(ECheckBoxState InCheckboxState);
 
-		UE_API TSharedRef<SDockTab> CreateMixerSourcesTab(const FSpawnTabArgs& Args);
-		UE_API TSharedRef<SDockTab> CreatePlotsTab(const FSpawnTabArgs& Args);
+		TSharedRef<SDockTab> CreateMixerSourcesTab(const FSpawnTabArgs& Args);
+		TSharedRef<SDockTab> CreatePlotsTab(const FSpawnTabArgs& Args);
 
-		UE_API void RegisterTabSpawners();
-		UE_API void UnregisterTabSpawners();
+		void RegisterTabSpawners();
+		void UnregisterTabSpawners();
 
-		UE_API TSharedRef<FTabManager::FLayout> GetDefaultTabLayout();
+		TSharedRef<FTabManager::FLayout> GetDefaultTabLayout();
 
 		FCheckBoxStyle MuteToggleButtonStyle;
 		FCheckBoxStyle SoloToggleButtonStyle;
@@ -139,6 +138,8 @@ namespace UE::Audio::Insights
 		TSharedPtr<FTabManager> MixerSourcesTabManager;
 		TSharedPtr<FWorkspaceItem> MixerSourcesWorkspace;
 
+		bool bDisplayAmpPeakInDb = true;
+
 #if WITH_EDITOR
 		// State of the mute and solo buttons
 		ECheckBoxState MuteState = ECheckBoxState::Unchecked;
@@ -150,5 +151,3 @@ namespace UE::Audio::Insights
 #endif // WITH_EDITOR
 	};
 } // namespace UE::Audio::Insights
-
-#undef UE_API
