[Audio Insights] Remove references to Hide feature in Sounds Dashboard - this feature has been scrapped

diff --git a/Engine/Plugins/AudioInsights/Source/AudioInsights/Private/SoundDashboardCommands.cpp b/Engine/Plugins/AudioInsights/Source/AudioInsights/Private/SoundDashboardCommands.cpp
--- a/Engine/Plugins/AudioInsights/Source/AudioInsights/Private/SoundDashboardCommands.cpp
+++ b/Engine/Plugins/AudioInsights/Source/AudioInsights/Private/SoundDashboardCommands.cpp
@@ -25,8 +25,6 @@ namespace UE::Audio::Insights
 		UI_COMMAND(ViewFullTree, "Tree View", "Organize sounds into categories.", EUserInterfaceActionType::RadioButton, FInputChord(EKeys::F1, EModifierKey::Control));
 		UI_COMMAND(ViewActiveSounds, "Active Sounds", "Organize sounds into Active Sounds.", EUserInterfaceActionType::RadioButton, FInputChord(EKeys::F2, EModifierKey::Control));
 		UI_COMMAND(ViewFlatList, "Flat List", "Display sounds as individual waves in a flat list.", EUserInterfaceActionType::RadioButton, FInputChord(EKeys::F3, EModifierKey::Control));
-		// @TODO UE-250399: Hide category pending to implement
-		//UI_COMMAND(Hide, "Hide", "Hides the selected sound.", EUserInterfaceActionType::Button, FInputChord(EKeys::H, EModifierKey::Control));
 	}
 } // namespace UE::Audio::Insights
 
diff --git a/Engine/Plugins/AudioInsights/Source/AudioInsights/Private/SoundDashboardCommands.h b/Engine/Plugins/AudioInsights/Source/AudioInsights/Private/SoundDashboardCommands.h
--- a/Engine/Plugins/AudioInsights/Source/AudioInsights/Private/SoundDashboardCommands.h
+++ b/Engine/Plugins/AudioInsights/Source/AudioInsights/Private/SoundDashboardCommands.h
@@ -22,8 +22,6 @@ namespace UE::Audio::Insights
 		TSharedPtr<const FUICommandInfo> GetViewFullTreeCommand() const { return ViewFullTree; }
 		TSharedPtr<const FUICommandInfo> GetViewActiveSoundsCommand() const { return ViewActiveSounds; }
 		TSharedPtr<const FUICommandInfo> GetViewFlatList() const { return ViewFlatList; }
-		// @TODO UE-250399: Hide category pending to implement
-		//TSharedPtr<const FUICommandInfo> GetHideCommand() const   { return Hide; }
 		
 	private:
 		TSharedPtr<FUICommandInfo> Pin;
@@ -33,7 +31,5 @@ namespace UE::Audio::Insights
 		TSharedPtr<FUICommandInfo> ViewFullTree;
 		TSharedPtr<FUICommandInfo> ViewActiveSounds;
 		TSharedPtr<FUICommandInfo> ViewFlatList;
-		// @TODO UE-250399: Hide category pending to implement
-		//TSharedPtr<FUICommandInfo> Hide;
 	};
 } // namespace UE::Audio::Insights
diff --git a/Engine/Plugins/AudioInsights/Source/AudioInsights/Private/Views/SoundDashboardViewFactory.cpp b/Engine/Plugins/AudioInsights/Source/AudioInsights/Private/Views/SoundDashboardViewFactory.cpp
--- a/Engine/Plugins/AudioInsights/Source/AudioInsights/Private/Views/SoundDashboardViewFactory.cpp
+++ b/Engine/Plugins/AudioInsights/Source/AudioInsights/Private/Views/SoundDashboardViewFactory.cpp
@@ -398,8 +398,6 @@ namespace UE::Audio::Insights
 		const FText SoundWaveCategoryName = LOCTEXT("SoundDashboard_Filter_SoundWaveNameText", "Sound Wave");
 		const FText SoundCueTemplateCategoryName = LOCTEXT("SoundDashboard_Filter_SoundCueTemplateNameText", "Sound Cue Template");
 		const FText PinnedCategoryName = LOCTEXT("SoundDashboard_Filter_PinnedNameText", "Pinned");
-		// @TODO UE-250399: Hide category pending to implement
-		//const FText HiddenCategoryName = LOCTEXT("SoundDashboard_Filter_HiddenNameText", "Hidden");
 	} // namespace FSoundDashboardViewFactoryPrivate
 
 	/////////////////////////////////////////////////////////////////////////////////////////
@@ -649,9 +647,6 @@ namespace UE::Audio::Insights
 			FExecuteAction::CreateLambda([this]() { ChangeTreeViewingMode(ETreeViewingOptions::FlatList); }),
 			FCanExecuteAction(),
 			FIsActionChecked::CreateLambda([this]() { return TreeViewingMode == ETreeViewingOptions::FlatList; }));
-
-		// @TODO UE-250399: Hide category pending to implement
-		//CommandList->MapAction(Commands.GetHideCommand(), FExecuteAction::CreateRaw(this, &FSoundDashboardViewFactory::HideSound));
 	}
 
 #if WITH_EDITOR
@@ -836,16 +831,6 @@ namespace UE::Audio::Insights
 				FText::GetEmpty(),
 				FSlateStyle::Get().GetColor("SoundDashboard.PinnedColor"),
 				FilterCategory)
-
-				// @TODO UE-250399: Hide category pending to implement
-				//MakeShared<FSoundDashboardFilter>(
-				//ESoundDashboardFilterFlags::Hidden,
-				//"Hidden",
-				//HiddenCategoryName,
-				//"AudioInsights.Icon.SoundDashboard.Hide",
-				//FText::GetEmpty(),
-				//FSlateStyle::Get().GetColor("SoundDashboard.HiddenColor"),
-				//FilterCategory)
 			};
 
 			SAssignNew(SoundsFilterBar, SAudioFilterBar<ESoundDashboardFilterFlags>)
@@ -1322,9 +1307,6 @@ namespace UE::Audio::Insights
 			MenuBuilder.AddMenuEntry(Commands.GetBrowseCommand(), NAME_None, TAttribute<FText>(), TAttribute<FText>(), FSlateStyle::Get().CreateIcon("AudioInsights.Icon.SoundDashboard.Browse"));
 			MenuBuilder.AddMenuEntry(Commands.GetEditCommand(),   NAME_None, TAttribute<FText>(), TAttribute<FText>(), FSlateStyle::Get().CreateIcon("AudioInsights.Icon.SoundDashboard.Edit"));
 #endif // WITH_EDITOR
-
-			// @TODO UE-250399: Hide category pending to implement
-			//MenuBuilder.AddMenuEntry(Commands.GetHideCommand(), NAME_None, TAttribute<FText>(), TAttribute<FText>(), FSlateStyle::Get().CreateIcon("AudioInsights.Icon.SoundDashboard.Hide"));
 		}
 
 		MenuBuilder.EndSection();
@@ -2622,12 +2604,6 @@ namespace UE::Audio::Insights
 		}
 	}
 #endif // WITH_EDITOR
-
-	// @TODO UE-250399: Hide category pending to implement
-	//void FSoundDashboardViewFactory::HideSound()
-	//{
-	//
-	//}
 } // namespace UE::Audio::Insights
 
 #undef LOCTEXT_NAMESPACE
diff --git a/Engine/Plugins/AudioInsights/Source/AudioInsights/Public/AudioInsightsStyle.h b/Engine/Plugins/AudioInsights/Source/AudioInsights/Public/AudioInsightsStyle.h
--- a/Engine/Plugins/AudioInsights/Source/AudioInsights/Public/AudioInsightsStyle.h
+++ b/Engine/Plugins/AudioInsights/Source/AudioInsights/Public/AudioInsightsStyle.h
@@ -179,8 +179,6 @@ namespace UE::Audio::Insights
 			Set("AudioInsights.Icon.SoundDashboard.Browse",           new IMAGE_BRUSH_SVG(TEXT("Icons/SoundDashboard/browse"),           Icon16));
 			Set("AudioInsights.Icon.SoundDashboard.Edit",             new IMAGE_BRUSH_SVG(TEXT("Icons/SoundDashboard/edit"),             Icon16));
 			Set("AudioInsights.Icon.SoundDashboard.Filter",           new IMAGE_BRUSH_SVG(TEXT("Icons/SoundDashboard/filter"),           Icon16));
-			// @TODO UE-250399: Hide category pending to implement
-			//Set("AudioInsights.Icon.SoundDashboard.Hide",             new IMAGE_BRUSH_SVG(TEXT("Icons/SoundDashboard/hide"),             Icon16));
 			Set("AudioInsights.Icon.SoundDashboard.Info",             new IMAGE_BRUSH_SVG(TEXT("Icons/SoundDashboard/info"),             Icon16));
 			Set("AudioInsights.Icon.SoundDashboard.MetaSound",        new IMAGE_BRUSH_SVG(TEXT("Icons/SoundDashboard/metasound"),        Icon16));
 			Set("AudioInsights.Icon.SoundDashboard.Mute",             new IMAGE_BRUSH_SVG(TEXT("Icons/SoundDashboard/mute"),             Icon16));
diff --git a/Engine/Plugins/AudioInsights/Source/AudioInsights/Public/Views/SoundDashboardViewFactory.h b/Engine/Plugins/AudioInsights/Source/AudioInsights/Public/Views/SoundDashboardViewFactory.h
--- a/Engine/Plugins/AudioInsights/Source/AudioInsights/Public/Views/SoundDashboardViewFactory.h
+++ b/Engine/Plugins/AudioInsights/Source/AudioInsights/Public/Views/SoundDashboardViewFactory.h
@@ -22,8 +22,6 @@ namespace UE::Audio::Insights
 		SoundWave = 1 << 3,
 		SoundCueTemplate = 1 << 4,
 		Pinned = 1 << 5
-		// @TODO UE-250399: Hide category pending to implement
-		//, Hidden = 1 << 6
 	};
 
 	ENUM_CLASS_FLAGS(ESoundDashboardFilterFlags);
@@ -185,9 +183,7 @@ namespace UE::Audio::Insights
 #if WITH_EDITOR
 		UE_API void BrowseSoundAsset() const;
 		UE_API void OpenSoundAsset() const;
-#endif
-		// @TODO UE-250399: Hide category pending to implement
-		//void HideSound();
+#endif // !WITH_EDITOR
 
 		static constexpr ESoundDashboardFilterFlags AllFilterFlags = 
 			ESoundDashboardFilterFlags::MetaSound        |
@@ -196,8 +192,6 @@ namespace UE::Audio::Insights
 			ESoundDashboardFilterFlags::SoundWave        |
 			ESoundDashboardFilterFlags::SoundCueTemplate |
 			ESoundDashboardFilterFlags::Pinned;
-			// @TODO UE-250399: Hide category pending to implement
-			//| ESoundDashboardFilterFlags::Hidden;
 
 		TSharedPtr<FPinnedSoundEntryWrapper> PinnedItemEntries;
 		TArray<TSharedPtr<IDashboardDataTreeViewEntry>> FullTree;
