[Audio Insights] Fix for plot y-axis in Mixer Sources tab not displaying the left-hand most number when switching between dB and linear scales

diff --git a/Engine/Plugins/AudioInsights/Source/AudioInsights/Public/AudioInsightsStyle.h b/Engine/Plugins/AudioInsights/Source/AudioInsights/Public/AudioInsightsStyle.h
--- a/Engine/Plugins/AudioInsights/Source/AudioInsights/Public/AudioInsightsStyle.h
+++ b/Engine/Plugins/AudioInsights/Source/AudioInsights/Public/AudioInsightsStyle.h
@@ -33,7 +33,7 @@ namespace UE::Audio::Insights
 		{
 			static FNumberFormattingOptions FloatFormat;
 			FloatFormat.MinimumIntegralDigits = 1;
-			FloatFormat.MaximumIntegralDigits = 1;
+			FloatFormat.MaximumIntegralDigits = 3;
 			FloatFormat.MinimumFractionalDigits = 3;
 			FloatFormat.MaximumFractionalDigits = 3;
 			return &FloatFormat;
