[Audio Insights] Fix for crash on opening Audio Insights

diff --git a/Engine/Plugins/AudioInsights/Source/AudioInsights/Public/AudioInsightsStyle.h b/Engine/Plugins/AudioInsights/Source/AudioInsights/Public/AudioInsightsStyle.h
--- a/Engine/Plugins/AudioInsights/Source/AudioInsights/Public/AudioInsightsStyle.h
+++ b/Engine/Plugins/AudioInsights/Source/AudioInsights/Public/AudioInsightsStyle.h
@@ -17,7 +17,7 @@ namespace UE::Audio::Insights
 	class FSlateStyle final : public FSlateStyleSet
 	{
 	public:
-		static FSlateStyle& Get()
+		AUDIOINSIGHTS_API static FSlateStyle& Get()
 		{
 			static FSlateStyle InsightsStyle;
 			return InsightsStyle;
