[Audio Insights] Add Auto-Expand options to the Sounds Dashboard

diff --git a/Engine/Plugins/AudioInsights/Source/AudioInsights/Private/Providers/SoundTraceProvider.cpp b/Engine/Plugins/AudioInsights/Source/AudioInsights/Private/Providers/SoundTraceProvider.cpp
--- a/Engine/Plugins/AudioInsights/Source/AudioInsights/Private/Providers/SoundTraceProvider.cpp
+++ b/Engine/Plugins/AudioInsights/Source/AudioInsights/Private/Providers/SoundTraceProvider.cpp
@@ -369,7 +369,6 @@ namespace UE::Audio::Insights
 			ActiveSoundEntryRef.SetName(*Msg.Name);
 			ActiveSoundEntryRef.EntryType = Msg.EntryType;
 			ActiveSoundEntryRef.CategoryName = FText::FromString(FSoundTraceProviderPrivate::GetEntryTypeDisplayName(Msg.EntryType));
-			ActiveSoundEntryRef.bShouldForceExpandChildren = false;
 
 #if WITH_EDITOR
 			ActiveSoundEntryRef.OwnerActorName = FText::FromString(*Msg.OwnerActorName);
@@ -404,7 +403,6 @@ namespace UE::Audio::Insights
 			SoundWaveEntry->EntryType = Msg.EntryType;
 			SoundWaveEntry->CategoryName = FText::FromString(FSoundTraceProviderPrivate::GetEntryTypeDisplayName(Msg.EntryType));
 			SoundWaveEntry->PinnedEntryType = (*OutActiveSoundEntry)->PinnedEntryType;
-			SoundWaveEntry->bShouldForceExpandChildren = false;
 
 #if WITH_EDITOR
 			SoundWaveEntry->OwnerActorName = FText::FromString(*Msg.OwnerActorName);
diff --git a/Engine/Plugins/AudioInsights/Source/AudioInsights/Private/SoundDashboardCommands.cpp b/Engine/Plugins/AudioInsights/Source/AudioInsights/Private/SoundDashboardCommands.cpp
--- a/Engine/Plugins/AudioInsights/Source/AudioInsights/Private/SoundDashboardCommands.cpp
+++ b/Engine/Plugins/AudioInsights/Source/AudioInsights/Private/SoundDashboardCommands.cpp
@@ -25,6 +25,10 @@ namespace UE::Audio::Insights
 		UI_COMMAND(ViewFullTree, "Tree View", "Organize sounds into categories.", EUserInterfaceActionType::RadioButton, FInputChord(EKeys::F1, EModifierKey::Control));
 		UI_COMMAND(ViewActiveSounds, "Active Sounds", "Organize sounds into Active Sounds.", EUserInterfaceActionType::RadioButton, FInputChord(EKeys::F2, EModifierKey::Control));
 		UI_COMMAND(ViewFlatList, "Flat List", "Display sounds as individual waves in a flat list.", EUserInterfaceActionType::RadioButton, FInputChord(EKeys::F3, EModifierKey::Control));
+
+		UI_COMMAND(AutoExpandCategories, "Categories", "Auto-expand new categories.", EUserInterfaceActionType::RadioButton, FInputChord(EKeys::F1, EModifierKey::Shift));
+		UI_COMMAND(AutoExpandEverything, "Everything", "Auto-expand all new categories and sounds.", EUserInterfaceActionType::RadioButton, FInputChord(EKeys::F2, EModifierKey::Shift));
+		UI_COMMAND(AutoExpandNothing, "Nothing", "Don't auto-expand anything.", EUserInterfaceActionType::RadioButton, FInputChord(EKeys::F3, EModifierKey::Shift));
 	}
 } // namespace UE::Audio::Insights
 
diff --git a/Engine/Plugins/AudioInsights/Source/AudioInsights/Private/SoundDashboardCommands.h b/Engine/Plugins/AudioInsights/Source/AudioInsights/Private/SoundDashboardCommands.h
--- a/Engine/Plugins/AudioInsights/Source/AudioInsights/Private/SoundDashboardCommands.h
+++ b/Engine/Plugins/AudioInsights/Source/AudioInsights/Private/SoundDashboardCommands.h
@@ -22,6 +22,9 @@ namespace UE::Audio::Insights
 		TSharedPtr<const FUICommandInfo> GetViewFullTreeCommand() const { return ViewFullTree; }
 		TSharedPtr<const FUICommandInfo> GetViewActiveSoundsCommand() const { return ViewActiveSounds; }
 		TSharedPtr<const FUICommandInfo> GetViewFlatList() const { return ViewFlatList; }
+		TSharedPtr<const FUICommandInfo> GetAutoExpandCategoriesCommand() const { return AutoExpandCategories; }
+		TSharedPtr<const FUICommandInfo> GetAutoExpandEverythingCommand() const { return AutoExpandEverything; }
+		TSharedPtr<const FUICommandInfo> GetAutoExpandNothingCommand() const { return AutoExpandNothing; }
 		
 	private:
 		TSharedPtr<FUICommandInfo> Pin;
@@ -31,5 +34,8 @@ namespace UE::Audio::Insights
 		TSharedPtr<FUICommandInfo> ViewFullTree;
 		TSharedPtr<FUICommandInfo> ViewActiveSounds;
 		TSharedPtr<FUICommandInfo> ViewFlatList;
+		TSharedPtr<FUICommandInfo> AutoExpandCategories;
+		TSharedPtr<FUICommandInfo> AutoExpandEverything;
+		TSharedPtr<FUICommandInfo> AutoExpandNothing;
 	};
 } // namespace UE::Audio::Insights
diff --git a/Engine/Plugins/AudioInsights/Source/AudioInsights/Private/Views/SoundDashboardViewFactory.cpp b/Engine/Plugins/AudioInsights/Source/AudioInsights/Private/Views/SoundDashboardViewFactory.cpp
--- a/Engine/Plugins/AudioInsights/Source/AudioInsights/Private/Views/SoundDashboardViewFactory.cpp
+++ b/Engine/Plugins/AudioInsights/Source/AudioInsights/Private/Views/SoundDashboardViewFactory.cpp
@@ -143,8 +143,6 @@ namespace UE::Audio::Insights
 					CacheInitExpandStateRecursive(CastEntry(*ChildEntry));
 				}
 			}
-
-			SoundEntry.bShouldForceExpandChildren = SoundEntry.bIsExpanded;
 		}
 		
 		bool EntryCanHaveChildren(const TSharedRef<IDashboardDataTreeViewEntry>& InEntry)
@@ -415,8 +413,6 @@ namespace UE::Audio::Insights
 		{
 			// Only copy data that has possibly changed from the other entry
 			PinnedEntry.TimeoutTimestamp = OriginalEntry.TimeoutTimestamp;
-
-			PinnedEntry.bShouldForceExpandChildren = OriginalEntry.bShouldForceExpandChildren;
 			PinnedEntry.bIsVisible = OriginalEntry.bIsVisible;
 
 			// Just copy the last entry in the buffer rather than the whole buffer
@@ -647,6 +643,24 @@ namespace UE::Audio::Insights
 			FExecuteAction::CreateLambda([this]() { ChangeTreeViewingMode(ETreeViewingOptions::FlatList); }),
 			FCanExecuteAction(),
 			FIsActionChecked::CreateLambda([this]() { return TreeViewingMode == ETreeViewingOptions::FlatList; }));
+
+		CommandList->MapAction(
+			Commands.GetAutoExpandCategoriesCommand(),
+			FExecuteAction::CreateLambda([this]() { ChangeAutoExpandMode(EAutoExpandOptions::Categories); }),
+			FCanExecuteAction(),
+			FIsActionChecked::CreateLambda([this]() { return AutoExpandMode == EAutoExpandOptions::Categories; }));
+
+		CommandList->MapAction(
+			Commands.GetAutoExpandEverythingCommand(),
+			FExecuteAction::CreateLambda([this]() { ChangeAutoExpandMode(EAutoExpandOptions::Everything); }),
+			FCanExecuteAction(),
+			FIsActionChecked::CreateLambda([this]() { return AutoExpandMode == EAutoExpandOptions::Everything; }));
+
+		CommandList->MapAction(
+			Commands.GetAutoExpandNothingCommand(),
+			FExecuteAction::CreateLambda([this]() { ChangeAutoExpandMode(EAutoExpandOptions::Nothing); }),
+			FCanExecuteAction(),
+			FIsActionChecked::CreateLambda([this]() { return AutoExpandMode == EAutoExpandOptions::Nothing; }));
 	}
 
 #if WITH_EDITOR
@@ -1237,6 +1251,48 @@ namespace UE::Audio::Insights
 		MenuBuilder.AddMenuEntry(Commands.GetViewFlatList());
 	}
 
+	void FSoundDashboardViewFactory::ChangeAutoExpandMode(const EAutoExpandOptions SelectedMode)
+	{
+		if (AutoExpandMode == SelectedMode)
+		{
+			return;
+		}
+
+		AutoExpandMode = SelectedMode;
+
+		for (const TSharedPtr<IDashboardDataTreeViewEntry>& Entry : DataViewEntries)
+		{
+			RefreshExpansionChecks(Entry);
+		}
+	}
+
+	void FSoundDashboardViewFactory::RefreshExpansionChecks(const TSharedPtr<IDashboardDataTreeViewEntry>& Entry)
+	{
+		using namespace FSoundDashboardViewFactoryPrivate;
+
+		if (!Entry.IsValid())
+		{
+			return;
+		}
+
+		FSoundDashboardEntry& SoundEntry = CastEntry(*Entry);
+		SoundEntry.bHasSetInitExpansion = false;
+
+		for (const TSharedPtr<IDashboardDataTreeViewEntry>& Child : Entry->Children)
+		{
+			RefreshExpansionChecks(Child);
+		}
+	}
+
+	void FSoundDashboardViewFactory::BuildAutoExpandMenuContent(FMenuBuilder& MenuBuilder)
+	{
+		const FSoundDashboardCommands& Commands = FSoundDashboardCommands::Get();
+
+		MenuBuilder.AddMenuEntry(Commands.GetAutoExpandCategoriesCommand());
+		MenuBuilder.AddMenuEntry(Commands.GetAutoExpandEverythingCommand());
+		MenuBuilder.AddMenuEntry(Commands.GetAutoExpandNothingCommand());
+	}
+
 	void FSoundDashboardViewFactory::FilterText()
 	{
 		using namespace FSoundDashboardViewFactoryPrivate;
@@ -1838,6 +1894,28 @@ namespace UE::Audio::Insights
 		return bDataReset;
 	}
 
+	bool FSoundDashboardViewFactory::ShouldAutoExpand(const TSharedPtr<IDashboardDataTreeViewEntry>& Item) const
+	{
+		using namespace FSoundDashboardViewFactoryPrivate;
+
+		if (!Item.IsValid())
+		{
+			return false;
+		}
+
+		switch (AutoExpandMode)
+		{
+			case EAutoExpandOptions::Everything:
+				return true;
+			case EAutoExpandOptions::Nothing:
+				return false;
+			case EAutoExpandOptions::Categories:
+			default:
+				FSoundDashboardEntry& Entry = CastEntry(*Item);
+				return Entry.bIsCategory;
+		}
+	}
+
 	void FSoundDashboardViewFactory::RecursiveSort(TArray<TSharedPtr<IDashboardDataTreeViewEntry>>& OutTree, TFunctionRef<bool(const FSoundDashboardEntry&, const FSoundDashboardEntry&)> Predicate)
 	{
 		using namespace FSoundDashboardViewFactoryPrivate;
@@ -1979,6 +2057,15 @@ namespace UE::Audio::Insights
 				FNewMenuDelegate::CreateRaw(this, &FSoundDashboardViewFactory::BuildViewOptionsMenuContent),
 				false
 			);
+
+			// Auto-Expand options
+			MenuBuilder.AddSubMenu
+			(
+				LOCTEXT("SoundsDashboard_AutoExpand", "Auto-Expand"),
+				LOCTEXT("SoundsDashboard_AutoExpandTooltip", "Choose whether new items in the Sounds Dashboard are automatically expanded or not."),
+				FNewMenuDelegate::CreateRaw(this, &FSoundDashboardViewFactory::BuildAutoExpandMenuContent),
+				false
+			);
 		}
 
 		MenuBuilder.EndSection();
@@ -2516,7 +2603,6 @@ namespace UE::Audio::Insights
 					{
 						TSharedPtr<FSoundDashboardEntry> NewPinnedCategoryEntry = MakeShared<FSoundDashboardEntry>(CastEntry(*DataCategoryEntry));
 						NewPinnedCategoryEntry->Children.Reset();
-						NewPinnedCategoryEntry->bShouldForceExpandChildren = true;
 
 						TSharedPtr<FPinnedSoundEntryWrapper> PinnedCategory = PinnedItemEntries->AddChildEntry(NewPinnedCategoryEntry);
 						PinnedCategory->AddChildEntry(Entry);
diff --git a/Engine/Plugins/AudioInsights/Source/AudioInsights/Private/Views/TreeDashboardViewFactory.cpp b/Engine/Plugins/AudioInsights/Source/AudioInsights/Private/Views/TreeDashboardViewFactory.cpp
--- a/Engine/Plugins/AudioInsights/Source/AudioInsights/Private/Views/TreeDashboardViewFactory.cpp
+++ b/Engine/Plugins/AudioInsights/Source/AudioInsights/Private/Views/TreeDashboardViewFactory.cpp
@@ -143,10 +143,10 @@ namespace UE::Audio::Insights
 						{
 							OutChildren = InParent->Children;
 
-							if (FilteredEntriesListView.IsValid() && InParent->ShouldInitExpandChildren())
+							if (FilteredEntriesListView.IsValid() && !InParent->HasSetInitExpansion())
 							{
-								FilteredEntriesListView->SetItemExpansion(InParent, true);
-								InParent->ResetShouldInitExpandChildren();
+								FilteredEntriesListView->SetItemExpansion(InParent, ShouldAutoExpand(InParent));
+								InParent->ResetHasSetInitExpansion();
 							}
 						}
 					})
diff --git a/Engine/Plugins/AudioInsights/Source/AudioInsights/Public/AudioInsightsDataSource.h b/Engine/Plugins/AudioInsights/Source/AudioInsights/Public/AudioInsightsDataSource.h
--- a/Engine/Plugins/AudioInsights/Source/AudioInsights/Public/AudioInsightsDataSource.h
+++ b/Engine/Plugins/AudioInsights/Source/AudioInsights/Public/AudioInsightsDataSource.h
@@ -20,8 +20,8 @@ namespace UE::Audio::Insights
 		virtual ~IDashboardDataTreeViewEntry() = default;
 
 		virtual bool IsValid() const = 0;
-		virtual bool ShouldInitExpandChildren() const = 0;
-		virtual void ResetShouldInitExpandChildren() = 0;
+		virtual bool HasSetInitExpansion() const = 0;
+		virtual void ResetHasSetInitExpansion() = 0;
 
 		TArray<TSharedPtr<IDashboardDataTreeViewEntry>> Children;
 		bool bIsExpanded = false;
diff --git a/Engine/Plugins/AudioInsights/Source/AudioInsights/Public/Messages/SoundTraceMessages.h b/Engine/Plugins/AudioInsights/Source/AudioInsights/Public/Messages/SoundTraceMessages.h
--- a/Engine/Plugins/AudioInsights/Source/AudioInsights/Public/Messages/SoundTraceMessages.h
+++ b/Engine/Plugins/AudioInsights/Source/AudioInsights/Public/Messages/SoundTraceMessages.h
@@ -179,8 +179,8 @@ namespace UE::Audio::Insights
 
 		virtual bool IsValid() const override { return PlayOrder != INDEX_NONE; }
 
-		virtual bool ShouldInitExpandChildren() const override { return bShouldForceExpandChildren; }
-		virtual void ResetShouldInitExpandChildren() override { bShouldForceExpandChildren = false; }
+		virtual bool HasSetInitExpansion() const override { return bHasSetInitExpansion; }
+		virtual void ResetHasSetInitExpansion() override { bHasSetInitExpansion = true; }
 
 		enum class EPinnedEntryType : uint8
 		{
@@ -197,7 +197,7 @@ namespace UE::Audio::Insights
 		ESoundDashboardEntryType EntryType = ESoundDashboardEntryType::None;
 		EPinnedEntryType PinnedEntryType = EPinnedEntryType::None;
 		bool bIsCategory = false;
-		bool bShouldForceExpandChildren = true;
+		bool bHasSetInitExpansion = false;
 		bool bIsVisible = false;
 
 		::Audio::TCircularAudioBuffer<FDataPoint> PriorityDataPoints;
diff --git a/Engine/Plugins/AudioInsights/Source/AudioInsights/Public/Views/SoundDashboardViewFactory.h b/Engine/Plugins/AudioInsights/Source/AudioInsights/Public/Views/SoundDashboardViewFactory.h
--- a/Engine/Plugins/AudioInsights/Source/AudioInsights/Public/Views/SoundDashboardViewFactory.h
+++ b/Engine/Plugins/AudioInsights/Source/AudioInsights/Public/Views/SoundDashboardViewFactory.h
@@ -130,6 +130,7 @@ namespace UE::Audio::Insights
 		UE_API virtual bool ResetTreeData() override;
 
 		virtual int32 GetDisplayedTreeDepth() const override { return static_cast<int32>(TreeViewingMode); }
+		UE_API virtual bool ShouldAutoExpand(const TSharedPtr<IDashboardDataTreeViewEntry>& Item) const override;
 
 		enum class ETreeViewingOptions : uint8
 		{
@@ -141,6 +142,17 @@ namespace UE::Audio::Insights
 		UE_API void ChangeTreeViewingMode(const ETreeViewingOptions SelectedMode);
 		UE_API void BuildViewOptionsMenuContent(FMenuBuilder& MenuBuilder);
 
+		enum class EAutoExpandOptions : uint8
+		{
+			Categories = 0,
+			Everything,
+			Nothing
+		};
+
+		UE_API void ChangeAutoExpandMode(const EAutoExpandOptions SelectedMode);
+		UE_API void RefreshExpansionChecks(const TSharedPtr<IDashboardDataTreeViewEntry>& Entry);
+		UE_API void BuildAutoExpandMenuContent(FMenuBuilder& MenuBuilder);
+
 		UE_API void FilterText();
 		UE_API void RecursiveSort(TArray<TSharedPtr<IDashboardDataTreeViewEntry>>& OutTree, TFunctionRef<bool(const FSoundDashboardEntry&, const FSoundDashboardEntry&)> Predicate);
 		UE_API void SortByPredicate(TFunctionRef<bool(const FSoundDashboardEntry&, const FSoundDashboardEntry&)> Predicate);
@@ -201,6 +213,7 @@ namespace UE::Audio::Insights
 		TSharedPtr<SWidget> SoundsFilterBarButton;
 
 		ETreeViewingOptions TreeViewingMode = ETreeViewingOptions::FullTree;
+		EAutoExpandOptions AutoExpandMode = EAutoExpandOptions::Categories;
 		ESoundDashboardFilterFlags SelectedFilterFlags = AllFilterFlags;
 		bool bIsPinnedCategoryFilterEnabled = true;
 		bool bShowRecentlyStoppedSounds = false;
diff --git a/Engine/Plugins/AudioInsights/Source/AudioInsights/Public/Views/TreeDashboardViewFactory.h b/Engine/Plugins/AudioInsights/Source/AudioInsights/Public/Views/TreeDashboardViewFactory.h
--- a/Engine/Plugins/AudioInsights/Source/AudioInsights/Public/Views/TreeDashboardViewFactory.h
+++ b/Engine/Plugins/AudioInsights/Source/AudioInsights/Public/Views/TreeDashboardViewFactory.h
@@ -111,6 +111,8 @@ namespace UE::Audio::Insights
 
 		virtual int32 GetDisplayedTreeDepth() const { return 0; }
 
+		virtual bool ShouldAutoExpand(const TSharedPtr<IDashboardDataTreeViewEntry>& Item) const { return true; }
+
 		virtual bool ResetTreeData()
 		{
 			if (!DataViewEntries.IsEmpty())
