[Audio Insights] Change maximum number of fractional digits for amplitude in Audio Insights to one. Make sure the plots window initializes to correct number format.

diff --git a/Engine/Plugins/AudioInsights/Source/AudioInsights/Private/Views/MixerSourceDashboardViewFactory.cpp b/Engine/Plugins/AudioInsights/Source/AudioInsights/Private/Views/MixerSourceDashboardViewFactory.cpp
--- a/Engine/Plugins/AudioInsights/Source/AudioInsights/Private/Views/MixerSourceDashboardViewFactory.cpp
+++ b/Engine/Plugins/AudioInsights/Source/AudioInsights/Private/Views/MixerSourceDashboardViewFactory.cpp
@@ -156,7 +156,7 @@ namespace UE::Audio::Insights
 
 							return bDisplayAmpPeakInDb 
 								? AudioInsightsUtils::ConvertToDecibelsText(Amplitude)
-								: FText::AsNumber(Amplitude, FSlateStyle::Get().GetAmpFloatFormat());
+								: FText::AsNumber(Amplitude, FSlateStyle::Get().GetLinearVolumeFloatFormat());
 						},
 						false /* bDefaultHidden */,
 						0.12f /* FillWidth */
@@ -169,7 +169,7 @@ namespace UE::Audio::Insights
 						[](const IDashboardDataViewEntry& InData)
 						{
 							const ::Audio::TCircularAudioBuffer<FDataPoint>& VolumeDataPoints = MixerSourcePrivate::CastEntry(InData).VolumeDataPoints;
-							return FText::AsNumber(MixerSourcePrivate::GetLastEntryArrayValue(VolumeDataPoints), FSlateStyle::Get().GetAmpFloatFormat());
+							return FText::AsNumber(MixerSourcePrivate::GetLastEntryArrayValue(VolumeDataPoints), FSlateStyle::Get().GetLinearVolumeFloatFormat());
 						},
 						false /* bDefaultHidden */,
 						0.07f /* FillWidth */
@@ -182,7 +182,7 @@ namespace UE::Audio::Insights
 						[](const IDashboardDataViewEntry& InData)
 						{
 							const ::Audio::TCircularAudioBuffer<FDataPoint>& DistanceAttenuationDataPoints = MixerSourcePrivate::CastEntry(InData).DistanceAttenuationDataPoints;
-							return FText::AsNumber(MixerSourcePrivate::GetLastEntryArrayValue(DistanceAttenuationDataPoints), FSlateStyle::Get().GetAmpFloatFormat());
+							return FText::AsNumber(MixerSourcePrivate::GetLastEntryArrayValue(DistanceAttenuationDataPoints), FSlateStyle::Get().GetLinearVolumeFloatFormat());
 						},
 						true  /* bDefaultHidden */,
 						0.15f /* FillWidth */
@@ -690,14 +690,14 @@ namespace UE::Audio::Insights
 					{
 
 						[](const IDashboardDataViewEntry& InData) -> const ::Audio::TCircularAudioBuffer<FDataPoint>& { return MixerSourcePrivate::CastEntry(InData).VolumeDataPoints; },
-						FSlateStyle::Get().GetAmpFloatFormat()
+						FSlateStyle::Get().GetLinearVolumeFloatFormat()
 					}
 				},
 				{
 					DistanceAttenuationColumnName,
 					{
 						[](const IDashboardDataViewEntry& InData) -> const ::Audio::TCircularAudioBuffer<FDataPoint>& { return MixerSourcePrivate::CastEntry(InData).DistanceAttenuationDataPoints; },
-						FSlateStyle::Get().GetAmpFloatFormat()
+						FSlateStyle::Get().GetLinearVolumeFloatFormat()
 					}
 				},
 				{
@@ -836,6 +836,9 @@ namespace UE::Audio::Insights
 				.ViewRange_Lambda(GetViewRange)
 				.PixelSnappingMethod(EWidgetPixelSnapping::Disabled);
 			}
+
+			// Set the initial number format for the Amplitude plot
+			PlotWidgets[0]->SetYValueFormattingOptions(*FSlateStyle::Get().GetAmpFloatFormat());
 		}
 
 		// Create plot column combo box widgets  
@@ -844,7 +847,7 @@ namespace UE::Audio::Insights
 			return SNew(SComboBox<FName>)
 				.ToolTipText(MixerSourcePrivate::PlotColumnSelectDescription)
 				.OptionsSource(&ColumnNames)
-				.OnGenerateWidget_Lambda([this](const FName& ColumnName)
+				.OnGenerateWidget_Lambda([this, PlotWidgetIndex](const FName& ColumnName)
 				{
 					return SNew(STextBlock)
 					.Font(FAppStyle::GetFontStyle(TEXT("PropertyWindow.NormalFont")))
diff --git a/Engine/Plugins/AudioInsights/Source/AudioInsights/Private/Views/SoundDashboardViewFactory.cpp b/Engine/Plugins/AudioInsights/Source/AudioInsights/Private/Views/SoundDashboardViewFactory.cpp
--- a/Engine/Plugins/AudioInsights/Source/AudioInsights/Private/Views/SoundDashboardViewFactory.cpp
+++ b/Engine/Plugins/AudioInsights/Source/AudioInsights/Private/Views/SoundDashboardViewFactory.cpp
@@ -1589,7 +1589,7 @@ namespace UE::Audio::Insights
 
 						return PriorityValue >= VolumeWeightedMaxPriority 
 							? LOCTEXT("AudioDashboard_Sounds_Max", "MAX") 
-							: FText::AsNumber(PriorityValue, FSlateStyle::Get().GetAmpFloatFormat());
+							: FText::AsNumber(PriorityValue, FSlateStyle::Get().GetLinearVolumeFloatFormat());
 					}
 				}
 			},
@@ -1623,7 +1623,7 @@ namespace UE::Audio::Insights
 
 						return bDisplayAmpPeakInDb 
 							? AudioInsightsUtils::ConvertToDecibelsText(AmplitudeValue) 
-							: FText::AsNumber(AmplitudeValue, FSlateStyle::Get().GetAmpFloatFormat());
+							: FText::AsNumber(AmplitudeValue, FSlateStyle::Get().GetLinearVolumeFloatFormat());
 					}
 				}
 			},
@@ -1638,7 +1638,7 @@ namespace UE::Audio::Insights
 							return FText::GetEmpty();
 						}
 
-						return FText::AsNumber(GetLastEntryArrayValue(SoundDashboardEntry.VolumeDataPoints), FSlateStyle::Get().GetAmpFloatFormat());
+						return FText::AsNumber(GetLastEntryArrayValue(SoundDashboardEntry.VolumeDataPoints), FSlateStyle::Get().GetLinearVolumeFloatFormat());
 					}
 				}
 			},
diff --git a/Engine/Plugins/AudioInsights/Source/AudioInsights/Public/AudioInsightsStyle.h b/Engine/Plugins/AudioInsights/Source/AudioInsights/Public/AudioInsightsStyle.h
--- a/Engine/Plugins/AudioInsights/Source/AudioInsights/Public/AudioInsightsStyle.h
+++ b/Engine/Plugins/AudioInsights/Source/AudioInsights/Public/AudioInsightsStyle.h
@@ -30,6 +30,16 @@ namespace UE::Audio::Insights
 		}
 
 		const FNumberFormattingOptions* GetAmpFloatFormat() const
+		{
+			static FNumberFormattingOptions FloatFormat;
+			FloatFormat.MinimumIntegralDigits = 1;
+			FloatFormat.MaximumIntegralDigits = 3;
+			FloatFormat.MinimumFractionalDigits = 1;
+			FloatFormat.MaximumFractionalDigits = 1;
+			return &FloatFormat;
+		};
+
+		const FNumberFormattingOptions* GetLinearVolumeFloatFormat() const
 		{
 			static FNumberFormattingOptions FloatFormat;
 			FloatFormat.MinimumIntegralDigits = 1;
diff --git a/Engine/Plugins/Runtime/AudioModulation/Source/AudioModulationEditor/Private/Insights/Views/ControlBusDashboardViewFactory.cpp b/Engine/Plugins/Runtime/AudioModulation/Source/AudioModulationEditor/Private/Insights/Views/ControlBusDashboardViewFactory.cpp
--- a/Engine/Plugins/Runtime/AudioModulation/Source/AudioModulationEditor/Private/Insights/Views/ControlBusDashboardViewFactory.cpp
+++ b/Engine/Plugins/Runtime/AudioModulation/Source/AudioModulationEditor/Private/Insights/Views/ControlBusDashboardViewFactory.cpp
@@ -148,7 +148,7 @@ namespace AudioModulationEditor
 					"Value",
 					{
 						LOCTEXT("ControlBus_ValueColumnDisplayName", "Value"),
-						[](const IDashboardDataViewEntry& InData) { return FText::AsNumber(ControlBusPrivate::CastEntry(InData).Value, FSlateStyle::Get().GetAmpFloatFormat()); },
+						[](const IDashboardDataViewEntry& InData) { return FText::AsNumber(ControlBusPrivate::CastEntry(InData).Value, FSlateStyle::Get().GetLinearVolumeFloatFormat()); },
 						false /* bDefaultHidden */,
 						0.07f /* FillWidth */
 					}
diff --git a/Engine/Plugins/Runtime/AudioModulation/Source/AudioModulationEditor/Private/Insights/Views/ModulationMatrixDashboardViewFactory.cpp b/Engine/Plugins/Runtime/AudioModulation/Source/AudioModulationEditor/Private/Insights/Views/ModulationMatrixDashboardViewFactory.cpp
--- a/Engine/Plugins/Runtime/AudioModulation/Source/AudioModulationEditor/Private/Insights/Views/ModulationMatrixDashboardViewFactory.cpp
+++ b/Engine/Plugins/Runtime/AudioModulation/Source/AudioModulationEditor/Private/Insights/Views/ModulationMatrixDashboardViewFactory.cpp
@@ -274,7 +274,7 @@ namespace AudioModulationEditor
 
 					if (const float* FoundValuePtr = ModulationMatrixEntry.BusIdToValueMap.Find(BusId))
 					{
-						return FText::AsNumber(*FoundValuePtr, FSlateStyle::Get().GetAmpFloatFormat());
+						return FText::AsNumber(*FoundValuePtr, FSlateStyle::Get().GetLinearVolumeFloatFormat());
 					}
 
 					return FText::GetEmpty();
