[Audio Insights] Add editor settings for the Sounds Dashboard so the settings are persistent. Added setting for timeout time for recently stopped sounds.

diff --git a/Engine/Plugins/AudioInsights/Source/AudioInsights/Private/Providers/SoundTraceProvider.cpp b/Engine/Plugins/AudioInsights/Source/AudioInsights/Private/Providers/SoundTraceProvider.cpp
--- a/Engine/Plugins/AudioInsights/Source/AudioInsights/Private/Providers/SoundTraceProvider.cpp
+++ b/Engine/Plugins/AudioInsights/Source/AudioInsights/Private/Providers/SoundTraceProvider.cpp
@@ -3,6 +3,7 @@
 
 #include "AudioInsightsModule.h"
 #include "Messages/MixerSourceTraceMessages.h"
+#include "Settings/SoundDashboardSettings.h"
 #include "Trace/Analyzer.h"
 #include "TraceServices/Model/AnalysisSession.h"
 
@@ -165,6 +166,19 @@ namespace UE::Audio::Insights
 	FSoundTraceProvider::FSoundTraceProvider()
 		: TDeviceDataMapTraceProvider<ESoundDashboardEntryType, TSharedPtr<FSoundDashboardEntry>>(GetName_Static())
 	{
+#if WITH_EDITOR
+		OnReadSettingsHandle = FSoundDashboardSettings::OnReadSettings.AddLambda([this](const FSoundDashboardSettings& InSettings) 
+		{
+			DashboardTimeoutTime = static_cast<double>(InSettings.StoppedSoundTimeoutTime);
+		});
+#endif
+	}
+
+	FSoundTraceProvider::~FSoundTraceProvider()
+	{
+#if WITH_EDITOR
+		FSoundDashboardSettings::OnReadSettings.Remove(OnReadSettingsHandle);
+#endif
 	}
 
 	FName FSoundTraceProvider::GetName_Static()
@@ -545,7 +559,6 @@ namespace UE::Audio::Insights
 			{
 				if ((*OutActiveSoundEntry)->Timestamp < Msg.Timestamp)
 				{
-					static constexpr double DashboardTimeoutTime = 3.0;
 					const double TimeoutTimestamp = FPlatformTime::Seconds() + DashboardTimeoutTime;
 
 					(*OutActiveSoundEntry)->TimeoutTimestamp = TimeoutTimestamp;
diff --git a/Engine/Plugins/AudioInsights/Source/AudioInsights/Private/Providers/SoundTraceProvider.h b/Engine/Plugins/AudioInsights/Source/AudioInsights/Private/Providers/SoundTraceProvider.h
--- a/Engine/Plugins/AudioInsights/Source/AudioInsights/Private/Providers/SoundTraceProvider.h
+++ b/Engine/Plugins/AudioInsights/Source/AudioInsights/Private/Providers/SoundTraceProvider.h
@@ -22,7 +22,7 @@ namespace UE::Audio::Insights
 	{
 	public:
 		FSoundTraceProvider();
-		virtual ~FSoundTraceProvider() = default;
+		virtual ~FSoundTraceProvider();
 
 		static FName GetName_Static();
 
@@ -35,6 +35,8 @@ namespace UE::Audio::Insights
 		virtual bool ProcessMessages() override;
 		virtual bool ProcessManuallyUpdatedEntries() override;
 
+		void SetDashboardTimeoutTime(const double InTimeOutTime) { DashboardTimeoutTime = InTimeOutTime; }
+
 	private:
 		void GetOrCreateActiveSoundEntry(const FSoundStartMessage& Msg, TSharedPtr<FSoundDashboardEntry>& OutReturnedSoundEntry);
 		void GetActiveSoundEntryFromIDs(const uint32 PlayOrderID, const ::Audio::FDeviceId DeviceId, TSharedPtr<FSoundDashboardEntry>& OutSoundEntry);
@@ -65,5 +67,11 @@ namespace UE::Audio::Insights
 		TMap<uint32, SoundEntryKeys> ActiveSoundToEntryKeysMap;
 		TArray<SoundMessageIDs> EntriesTimingOut;
 		TSet<uint32> SoundsStoppedBeforeStart;
+
+		double DashboardTimeoutTime = 3.0;
+
+#if WITH_EDITOR
+		FDelegateHandle OnReadSettingsHandle;
+#endif // WITH_EDITOR
 	};
 } // namespace UE::Audio::Insights
diff --git a/Engine/Plugins/AudioInsights/Source/AudioInsights/Private/Settings/SoundDashboardSettings.cpp b/Engine/Plugins/AudioInsights/Source/AudioInsights/Private/Settings/SoundDashboardSettings.cpp
new file mode 100644
--- /dev/null
+++ b/Engine/Plugins/AudioInsights/Source/AudioInsights/Private/Settings/SoundDashboardSettings.cpp
@@ -0,0 +1,11 @@
+// Copyright Epic Games, Inc. All Rights Reserved.
+#include "Settings/SoundDashboardSettings.h"
+
+#if WITH_EDITOR
+
+FSoundDashboardSettings::FOnReadSoundDashboardSettings FSoundDashboardSettings::OnReadSettings;
+FSoundDashboardSettings::FOnWriteSoundDashboardSettings FSoundDashboardSettings::OnWriteSettings;
+FSoundDashboardSettings::FOnRequestReadSoundDashboardSettings FSoundDashboardSettings::OnRequestReadSettings;
+FSoundDashboardSettings::FOnRequestWriteSoundDashboardSettings FSoundDashboardSettings::OnRequestWriteSettings;
+
+#endif // WITH_EDITOR
\ No newline at end of file
diff --git a/Engine/Plugins/AudioInsights/Source/AudioInsights/Private/Views/SoundDashboardViewFactory.cpp b/Engine/Plugins/AudioInsights/Source/AudioInsights/Private/Views/SoundDashboardViewFactory.cpp
--- a/Engine/Plugins/AudioInsights/Source/AudioInsights/Private/Views/SoundDashboardViewFactory.cpp
+++ b/Engine/Plugins/AudioInsights/Source/AudioInsights/Private/Views/SoundDashboardViewFactory.cpp
@@ -10,6 +10,7 @@
 #include "Providers/SoundTraceProvider.h"
 #include "SoundDashboardCommands.h"
 #include "Views/SAudioFilterBar.h"
+#include "Widgets/Input/SNumericEntryBox.h"
 #include "Widgets/Input/SSegmentedControl.h"
 
 #if WITH_EDITOR
@@ -228,6 +229,27 @@ namespace UE::Audio::Insights
 			return TotalNumChildren;
 		}
 
+		void ClampMaxTimeoutRecursive(const TSharedPtr<IDashboardDataTreeViewEntry>& Entry, const double NewTimeOutTimestamp)
+		{
+			if (!Entry.IsValid())
+			{
+				return;
+			}
+
+			FSoundDashboardEntry& SoundEntry = CastEntry(*Entry);
+
+			// Clamp max timeout to the current max timeout from the settings
+			if (SoundEntry.TimeoutTimestamp != INVALID_TIMEOUT && SoundEntry.TimeoutTimestamp > NewTimeOutTimestamp)
+			{
+				SoundEntry.TimeoutTimestamp = NewTimeOutTimestamp;
+			}
+
+			for (const TSharedPtr<IDashboardDataTreeViewEntry>& Child : Entry->Children)
+			{
+				ClampMaxTimeoutRecursive(Child, NewTimeOutTimestamp);
+			}
+		}
+
 #if WITH_EDITOR
 		void SetMuteSolo(const IDashboardDataTreeViewEntry& InEntry, const EMuteSoloMode InMuteSoloMode, const bool bInOnOff)
 		{
@@ -586,11 +608,21 @@ namespace UE::Audio::Insights
 		FSoundDashboardCommands::Register();
 
 		BindCommands();
+
+#if WITH_EDITOR
+		FSoundDashboardSettings::OnReadSettings.AddRaw(this, &FSoundDashboardViewFactory::OnReadEditorSettings);
+		FSoundDashboardSettings::OnWriteSettings.AddRaw(this, &FSoundDashboardViewFactory::OnWriteEditorSettings);
+#endif
 	}
 
 	FSoundDashboardViewFactory::~FSoundDashboardViewFactory()
 	{
 		FSoundDashboardCommands::Unregister();
+
+#if WITH_EDITOR
+		FSoundDashboardSettings::OnReadSettings.RemoveAll(this);
+		FSoundDashboardSettings::OnWriteSettings.RemoveAll(this);
+#endif
 	}
 
 	FName FSoundDashboardViewFactory::GetName() const
@@ -628,39 +660,39 @@ namespace UE::Audio::Insights
 
 		CommandList->MapAction(
 			Commands.GetViewFullTreeCommand(),
-			FExecuteAction::CreateLambda([this]() { ChangeTreeViewingMode(ETreeViewingOptions::FullTree); }),
+			FExecuteAction::CreateLambda([this]() { ChangeTreeViewingMode(ESoundDashboardTreeViewingOptions::FullTree); }),
 			FCanExecuteAction(),
-			FIsActionChecked::CreateLambda([this]() { return TreeViewingMode == ETreeViewingOptions::FullTree; }));
+			FIsActionChecked::CreateLambda([this]() { return TreeViewingMode == ESoundDashboardTreeViewingOptions::FullTree; }));
 
 		CommandList->MapAction(
 			Commands.GetViewActiveSoundsCommand(),
-			FExecuteAction::CreateLambda([this]() { ChangeTreeViewingMode(ETreeViewingOptions::ActiveSounds); }),
+			FExecuteAction::CreateLambda([this]() { ChangeTreeViewingMode(ESoundDashboardTreeViewingOptions::ActiveSounds); }),
 			FCanExecuteAction(),
-			FIsActionChecked::CreateLambda([this]() { return TreeViewingMode == ETreeViewingOptions::ActiveSounds; }));
+			FIsActionChecked::CreateLambda([this]() { return TreeViewingMode == ESoundDashboardTreeViewingOptions::ActiveSounds; }));
 
 		CommandList->MapAction(
 			Commands.GetViewFlatList(),
-			FExecuteAction::CreateLambda([this]() { ChangeTreeViewingMode(ETreeViewingOptions::FlatList); }),
+			FExecuteAction::CreateLambda([this]() { ChangeTreeViewingMode(ESoundDashboardTreeViewingOptions::FlatList); }),
 			FCanExecuteAction(),
-			FIsActionChecked::CreateLambda([this]() { return TreeViewingMode == ETreeViewingOptions::FlatList; }));
+			FIsActionChecked::CreateLambda([this]() { return TreeViewingMode == ESoundDashboardTreeViewingOptions::FlatList; }));
 
 		CommandList->MapAction(
 			Commands.GetAutoExpandCategoriesCommand(),
-			FExecuteAction::CreateLambda([this]() { ChangeAutoExpandMode(EAutoExpandOptions::Categories); }),
+			FExecuteAction::CreateLambda([this]() { ChangeAutoExpandMode(ESoundDashboardAutoExpandOptions::Categories); }),
 			FCanExecuteAction(),
-			FIsActionChecked::CreateLambda([this]() { return AutoExpandMode == EAutoExpandOptions::Categories; }));
+			FIsActionChecked::CreateLambda([this]() { return AutoExpandMode == ESoundDashboardAutoExpandOptions::Categories; }));
 
 		CommandList->MapAction(
 			Commands.GetAutoExpandEverythingCommand(),
-			FExecuteAction::CreateLambda([this]() { ChangeAutoExpandMode(EAutoExpandOptions::Everything); }),
+			FExecuteAction::CreateLambda([this]() { ChangeAutoExpandMode(ESoundDashboardAutoExpandOptions::Everything); }),
 			FCanExecuteAction(),
-			FIsActionChecked::CreateLambda([this]() { return AutoExpandMode == EAutoExpandOptions::Everything; }));
+			FIsActionChecked::CreateLambda([this]() { return AutoExpandMode == ESoundDashboardAutoExpandOptions::Everything; }));
 
 		CommandList->MapAction(
 			Commands.GetAutoExpandNothingCommand(),
-			FExecuteAction::CreateLambda([this]() { ChangeAutoExpandMode(EAutoExpandOptions::Nothing); }),
+			FExecuteAction::CreateLambda([this]() { ChangeAutoExpandMode(ESoundDashboardAutoExpandOptions::Nothing); }),
 			FCanExecuteAction(),
-			FIsActionChecked::CreateLambda([this]() { return AutoExpandMode == EAutoExpandOptions::Nothing; }));
+			FIsActionChecked::CreateLambda([this]() { return AutoExpandMode == ESoundDashboardAutoExpandOptions::Nothing; }));
 	}
 
 #if WITH_EDITOR
@@ -896,6 +928,10 @@ namespace UE::Audio::Insights
 
 	TSharedRef<SWidget> FSoundDashboardViewFactory::MakeWidget(TSharedRef<SDockTab> OwnerTab, const FSpawnTabArgs& SpawnTabArgs)
 	{
+#if WITH_EDITOR
+		FSoundDashboardSettings::OnRequestReadSettings.Broadcast();
+#endif // WITH_EDITOR
+
 		return SNew(SVerticalBox)
 #if WITH_EDITOR
 			+ SVerticalBox::Slot()
@@ -981,7 +1017,7 @@ namespace UE::Audio::Insights
 					{
 						const uint32 TotalNumChildren = CountNumChildren(*InRowData, bShowRecentlyStoppedSounds);
 
-						if (TreeViewingMode != ETreeViewingOptions::FullTree && TotalNumChildren == 0)
+						if (TreeViewingMode != ESoundDashboardTreeViewingOptions::FullTree && TotalNumChildren == 0)
 						{
 							return FText::GetEmpty();
 						}
@@ -1110,7 +1146,7 @@ namespace UE::Audio::Insights
 		}
 		else
 		{
-			if (IsRootItem(InRowData) && TreeViewingMode == ETreeViewingOptions::FullTree)
+			if (IsRootItem(InRowData) && TreeViewingMode == ESoundDashboardTreeViewingOptions::FullTree)
 			{
 				return SNullWidget::NullWidget;
 			}
@@ -1150,7 +1186,7 @@ namespace UE::Audio::Insights
 			{
 				const FSoundDashboardEntry& SoundEntry = CastEntry(*Item);
 
-				if (TreeViewingMode != ETreeViewingOptions::FullTree && SoundEntry.bIsCategory && SoundEntry.EntryType != ESoundDashboardEntryType::Pinned)
+				if (TreeViewingMode != ESoundDashboardTreeViewingOptions::FullTree && SoundEntry.bIsCategory && SoundEntry.EntryType != ESoundDashboardEntryType::Pinned)
 				{
 					return EVisibility::Collapsed;
 				}
@@ -1228,7 +1264,7 @@ namespace UE::Audio::Insights
 		});
 	}
 
-	void FSoundDashboardViewFactory::ChangeTreeViewingMode(const ETreeViewingOptions SelectedMode)
+	void FSoundDashboardViewFactory::ChangeTreeViewingMode(const ESoundDashboardTreeViewingOptions SelectedMode, const bool bUpdateEditorSettings /*= true*/)
 	{
 		if (TreeViewingMode == SelectedMode)
 		{
@@ -1240,6 +1276,13 @@ namespace UE::Audio::Insights
 		ClearSelection();
 
 		bTreeViewModeChanged = true;
+
+#if WITH_EDITOR
+		if (bUpdateEditorSettings)
+		{
+			FSoundDashboardSettings::OnRequestWriteSettings.Broadcast();
+		}
+#endif
 	}
 
 	void FSoundDashboardViewFactory::BuildViewOptionsMenuContent(FMenuBuilder& MenuBuilder)
@@ -1251,7 +1294,7 @@ namespace UE::Audio::Insights
 		MenuBuilder.AddMenuEntry(Commands.GetViewFlatList());
 	}
 
-	void FSoundDashboardViewFactory::ChangeAutoExpandMode(const EAutoExpandOptions SelectedMode)
+	void FSoundDashboardViewFactory::ChangeAutoExpandMode(const ESoundDashboardAutoExpandOptions SelectedMode, const bool bUpdateEditorSettings /*= true*/)
 	{
 		if (AutoExpandMode == SelectedMode)
 		{
@@ -1264,6 +1307,13 @@ namespace UE::Audio::Insights
 		{
 			RefreshExpansionChecks(Entry);
 		}
+
+#if WITH_EDITOR
+		if (bUpdateEditorSettings)
+		{
+			FSoundDashboardSettings::OnRequestWriteSettings.Broadcast();
+		}
+#endif
 	}
 
 	void FSoundDashboardViewFactory::RefreshExpansionChecks(const TSharedPtr<IDashboardDataTreeViewEntry>& Entry)
@@ -1310,7 +1360,7 @@ namespace UE::Audio::Insights
 
 			ResetVisibility(*CategoryEntry);
 
-			if (TreeViewingMode == ETreeViewingOptions::FullTree)
+			if (TreeViewingMode == ESoundDashboardTreeViewingOptions::FullTree)
 			{
 				for (const TSharedPtr<IDashboardDataTreeViewEntry>& Entry : CategoryEntry->Children)
 				{
@@ -1905,11 +1955,11 @@ namespace UE::Audio::Insights
 
 		switch (AutoExpandMode)
 		{
-			case EAutoExpandOptions::Everything:
+			case ESoundDashboardAutoExpandOptions::Everything:
 				return true;
-			case EAutoExpandOptions::Nothing:
+			case ESoundDashboardAutoExpandOptions::Nothing:
 				return false;
-			case EAutoExpandOptions::Categories:
+			case ESoundDashboardAutoExpandOptions::Categories:
 			default:
 				FSoundDashboardEntry& Entry = CastEntry(*Item);
 				return Entry.bIsCategory;
@@ -1980,6 +2030,9 @@ namespace UE::Audio::Insights
 				FUIAction(FExecuteAction::CreateLambda([this]()
 				{
 					bDisplayAmpPeakInDb = !bDisplayAmpPeakInDb;
+#if WITH_EDITOR
+					FSoundDashboardSettings::OnRequestWriteSettings.Broadcast();
+#endif // WITH_EDITOR
 				})),
 				SNew(SHorizontalBox)
 				+ SHorizontalBox::Slot()
@@ -1990,12 +2043,16 @@ namespace UE::Audio::Insights
 					.Text(LOCTEXT("SoundsDashboard_AmpPeakDisplayMode_Text", "Amp (Peak) Display Mode"))
 				]
 				+ SHorizontalBox::Slot()
+				.Padding(4.0f, 0.0f, 0.0f, 0.0f)
 				.AutoWidth()
 				[
 					SNew(SSegmentedControl<bool>)
 					.OnValueChanged_Lambda([this](bool bInValue)
 					{
 						bDisplayAmpPeakInDb = bInValue;
+#if WITH_EDITOR
+						FSoundDashboardSettings::OnRequestWriteSettings.Broadcast();
+#endif // WITH_EDITOR
 					})
 					.Value_Lambda([this]
 					{
@@ -2015,11 +2072,32 @@ namespace UE::Audio::Insights
 				EUserInterfaceActionType::CollapsedButton
 			);
 
+			// View options
+			MenuBuilder.AddSubMenu
+			(
+				LOCTEXT("SoundsDashboard_ViewOptions", "View"),
+				LOCTEXT("SoundsDashboard_ViewOptionsTooltip", "Choose how items in the Sounds Dashboard is displayed."),
+				FNewMenuDelegate::CreateRaw(this, &FSoundDashboardViewFactory::BuildViewOptionsMenuContent),
+				false
+			);
+
+			// Auto-Expand options
+			MenuBuilder.AddSubMenu
+			(
+				LOCTEXT("SoundsDashboard_AutoExpand", "Auto-Expand"),
+				LOCTEXT("SoundsDashboard_AutoExpandTooltip", "Choose whether new items in the Sounds Dashboard are automatically expanded or not."),
+				FNewMenuDelegate::CreateRaw(this, &FSoundDashboardViewFactory::BuildAutoExpandMenuContent),
+				false
+			);
+
 			// Show Stopped Sounds checkbox
 			MenuBuilder.AddMenuEntry(
 				FUIAction(FExecuteAction::CreateLambda([this]()
 				{
 					bShowRecentlyStoppedSounds = !bShowRecentlyStoppedSounds;
+#if WITH_EDITOR
+					FSoundDashboardSettings::OnRequestWriteSettings.Broadcast();
+#endif // WITH_EDITOR
 				})),
 				SNew(SHorizontalBox)
 				+ SHorizontalBox::Slot()
@@ -2049,22 +2127,56 @@ namespace UE::Audio::Insights
 				EUserInterfaceActionType::CollapsedButton
 			);
 
-			// View options
-			MenuBuilder.AddSubMenu
-			(
-				LOCTEXT("SoundsDashboard_ViewOptions", "View"),
-				LOCTEXT("SoundsDashboard_ViewOptionsTooltip", "Choose how items in the Sounds Dashboard is displayed."),
-				FNewMenuDelegate::CreateRaw(this, &FSoundDashboardViewFactory::BuildViewOptionsMenuContent),
-				false
-			);
+			// Recently Stopped Sound Timeout
+			MenuBuilder.AddMenuEntry(
+				FUIAction(),
+				SNew(SHorizontalBox)
+				+ SHorizontalBox::Slot()
+				.HAlign(HAlign_Right)
+				.VAlign(VAlign_Center)
+				.AutoWidth()
+				[
+					SNew(STextBlock)
+					.Margin(FMargin(0.0, 4.0, 5.0, 0.0))
+					.Text(LOCTEXT("SoundsDashboard_RecentlyStoppedSoundTimeout", "Recently Stopped Timeout"))
+				]
+				+ SHorizontalBox::Slot()
+				.HAlign(HAlign_Left)
+				.VAlign(VAlign_Center)
+				.AutoWidth()
+				.Padding(4.0f, 0.0f, 0.0f, 0.0f)
+				[
+					SNew(SNumericEntryBox<float>)
+					.MinDesiredValueWidth(30)
+					.MinValue(0.0f)
+					.MaxValue(TOptional<float>())
+					.MaxSliderValue(TOptional<float>())
+					.MaxFractionalDigits(2)
+					.AllowWheel(true)
+					.WheelStep(0.1f)
+					.AllowSpin(true)
+					.Value_Lambda([this]()
+					{
+						return RecentlyStoppedSoundTimeout;
+					})
+					.OnValueChanged_Lambda([this](float Value)
+					{
+						RecentlyStoppedSoundTimeout = Value;
+						UpdateMaxStoppedSoundTimeout();
 
-			// Auto-Expand options
-			MenuBuilder.AddSubMenu
-			(
-				LOCTEXT("SoundsDashboard_AutoExpand", "Auto-Expand"),
-				LOCTEXT("SoundsDashboard_AutoExpandTooltip", "Choose whether new items in the Sounds Dashboard are automatically expanded or not."),
-				FNewMenuDelegate::CreateRaw(this, &FSoundDashboardViewFactory::BuildAutoExpandMenuContent),
-				false
+#if WITH_EDITOR
+						FSoundDashboardSettings::OnRequestWriteSettings.Broadcast();
+#else
+						const TSharedPtr<FSoundTraceProvider> Provider = FindProvider<FSoundTraceProvider>();
+						if (Provider.IsValid())
+						{
+							Provider->SetDashboardTimeoutTime(RecentlyStoppedSoundTimeout);
+						}
+#endif // WITH_EDITOR
+					})
+				],
+				NAME_None /*InExtensionHook*/,
+				LOCTEXT("SoundsDashboard_RecentlyStoppedSoundTimeoutTooltip", "How long recently stopped sounds stay greyed out in the Sounds Dashboard before being removed (in seconds)")
 			);
 		}
 
@@ -2096,7 +2208,54 @@ namespace UE::Audio::Insights
 			];
 	}
 
+	void FSoundDashboardViewFactory::UpdateMaxStoppedSoundTimeout()
+	{
+		using namespace FSoundDashboardViewFactoryPrivate;
+
+		const double NewTimeoutTimestamp = FPlatformTime::Seconds() + static_cast<double>(RecentlyStoppedSoundTimeout);
+
+		for (const TSharedPtr<IDashboardDataTreeViewEntry>& Entry : DataViewEntries)
+		{
+			ClampMaxTimeoutRecursive(Entry, NewTimeoutTimestamp);
+		}
+	}
+
 #if WITH_EDITOR
+	void FSoundDashboardViewFactory::OnReadEditorSettings(const FSoundDashboardSettings& InSettings)
+	{
+		ChangeTreeViewingMode(InSettings.TreeViewMode, false /*bUpdateEditorSettings*/);
+		ChangeAutoExpandMode(InSettings.AutoExpandMode, false /*bUpdateEditorSettings*/);
+
+		bShowRecentlyStoppedSounds = InSettings.bShowStoppedSounds;
+		bDisplayAmpPeakInDb = InSettings.AmplitudeDisplayMode == EAudioAmplitudeDisplayMode::Decibels;
+
+		if (RecentlyStoppedSoundTimeout != InSettings.StoppedSoundTimeoutTime)
+		{
+			RecentlyStoppedSoundTimeout = InSettings.StoppedSoundTimeoutTime;
+
+			UpdateMaxStoppedSoundTimeout();
+		}
+	}
+
+	void FSoundDashboardViewFactory::OnWriteEditorSettings(FSoundDashboardSettings& OutSettings)
+	{
+		OutSettings.TreeViewMode = TreeViewingMode;
+		OutSettings.AutoExpandMode = AutoExpandMode;
+		OutSettings.bShowStoppedSounds = bShowRecentlyStoppedSounds;
+		OutSettings.AmplitudeDisplayMode = bDisplayAmpPeakInDb ? EAudioAmplitudeDisplayMode::Decibels : EAudioAmplitudeDisplayMode::Linear;
+
+		if (OutSettings.StoppedSoundTimeoutTime != RecentlyStoppedSoundTimeout)
+		{
+			OutSettings.StoppedSoundTimeoutTime = RecentlyStoppedSoundTimeout;
+			
+			const TSharedPtr<FSoundTraceProvider> Provider = FindProvider<FSoundTraceProvider>();
+			if (Provider.IsValid())
+			{
+				Provider->SetDashboardTimeoutTime(RecentlyStoppedSoundTimeout);
+			}
+		}
+	}
+
 	TSharedRef<SWidget> FSoundDashboardViewFactory::CreateMuteSoloButton(const TSharedRef<FTraceTreeDashboardViewFactory::SRowWidget>& InRowWidget, 
 		const TSharedRef<IDashboardDataTreeViewEntry>& InRowData,
 		const FName& InColumn, 
@@ -2301,6 +2460,11 @@ namespace UE::Audio::Insights
 
 	void FSoundDashboardViewFactory::ClearSelection()
 	{
+		if (FilteredEntriesListView == nullptr)
+		{
+			return;
+		}
+
 		// Make sure to clear the selection and the internal selector in SListView
 		// to ensure no shared references keep the entry alive past the point
 		// it has been removed from the dashboard
@@ -2366,7 +2530,7 @@ namespace UE::Audio::Insights
 	{
 		using namespace FSoundDashboardViewFactoryPrivate;
 
-		if (SelectedItems.Contains(Entry) && (!IsRootItem(Entry.ToSharedRef()) || TreeViewingMode != ETreeViewingOptions::FullTree))
+		if (SelectedItems.Contains(Entry) && (!IsRootItem(Entry.ToSharedRef()) || TreeViewingMode != ESoundDashboardTreeViewingOptions::FullTree))
 		{
 			const FSoundDashboardEntry& SoundEntry = CastEntry(*Entry);
 
@@ -2382,7 +2546,7 @@ namespace UE::Audio::Insights
 			{
 				if (SelectedItems.Contains(Child))
 				{
-					if (TreeViewingMode == ETreeViewingOptions::FullTree && IsRootItem(Entry.ToSharedRef()))
+					if (TreeViewingMode == ESoundDashboardTreeViewingOptions::FullTree && IsRootItem(Entry.ToSharedRef()))
 					{
 						MarkBranchAsPinned(Child, true /*bIsPinned*/);
 						CreatePinnedEntry(Child);
@@ -2484,7 +2648,7 @@ namespace UE::Audio::Insights
 			{
 				if (bSoundEntryIsPinned)
 				{
-					if (TreeViewingMode == ETreeViewingOptions::FullTree && SoundEntry.bIsCategory)
+					if (TreeViewingMode == ESoundDashboardTreeViewingOptions::FullTree && SoundEntry.bIsCategory)
 					{
 						MarkBranchAsPinned(Child, true /*bIsPinned*/);
 						CreatePinnedEntry(Child);
@@ -2501,7 +2665,7 @@ namespace UE::Audio::Insights
 					const FSoundDashboardEntry& ChildEntry = FSoundDashboardViewFactoryPrivate::CastEntry(*Child);
 					const bool bChildEntryIsPinned = ChildEntry.PinnedEntryType == FSoundDashboardEntry::EPinnedEntryType::HiddenOriginalEntry;
 
-					if (bChildEntryIsPinned && TreeViewingMode != ETreeViewingOptions::FlatList && !SoundEntry.bIsCategory)
+					if (bChildEntryIsPinned && TreeViewingMode != ESoundDashboardTreeViewingOptions::FlatList && !SoundEntry.bIsCategory)
 					{
 						MarkBranchAsPinned(Entry, true /*bIsPinned*/);
 						CreatePinnedEntry(Entry);
@@ -2551,7 +2715,7 @@ namespace UE::Audio::Insights
 		// The pinned area will delete itself once empty
 		InitPinnedItemEntries();
 
-		if (TreeViewingMode != ETreeViewingOptions::FullTree)
+		if (TreeViewingMode != ESoundDashboardTreeViewingOptions::FullTree)
 		{
 			PinnedItemEntries->AddChildEntry(Entry);
 		}
diff --git a/Engine/Plugins/AudioInsights/Source/AudioInsights/Public/Settings/SoundDashboardSettings.h b/Engine/Plugins/AudioInsights/Source/AudioInsights/Public/Settings/SoundDashboardSettings.h
new file mode 100644
--- /dev/null
+++ b/Engine/Plugins/AudioInsights/Source/AudioInsights/Public/Settings/SoundDashboardSettings.h
@@ -0,0 +1,95 @@
+// Copyright Epic Games, Inc. All Rights Reserved.
+
+#pragma once
+
+#include "Delegates/DelegateCombinations.h"
+#include "Engine/DeveloperSettings.h"
+
+#include "SoundDashboardSettings.generated.h"
+
+#define LOCTEXT_NAMESPACE "AudioInsights"
+
+/** EAudioAmplitudeDisplayMode
+ *
+ * The units that audio amplitude and volume values are displayed in
+ */
+UENUM(BlueprintType)
+enum class EAudioAmplitudeDisplayMode : uint8
+{
+	/** Displays amplitude values in decibels. */
+	Decibels		UMETA(DisplayName = "dB"),
+
+	/** Displays amplitude values in linear scale. */
+	Linear			UMETA(DisplayName = "Linear")
+};
+
+/** ESoundDashboardTreeViewingOptions
+*
+* Change how sounds are organized within the Sound Dashboard tab in Audio Insights.
+*/
+UENUM(BlueprintType)
+enum class ESoundDashboardTreeViewingOptions : uint8
+{
+	/** Organize sounds into categories. */
+	FullTree		UMETA(DisplayName = "Tree View"),
+
+	/** Organize sounds into Active Sounds. */
+	ActiveSounds	UMETA(DisplayName = "Active Sounds"),
+
+	/** Display sounds as individual waves in a flat list. */
+	FlatList		UMETA(DisplayName = "Flat List")
+};
+
+/** ESoundDashboardAutoExpandOptions
+ *
+ * Control whether new sounds entering the Sound Dashboard are auto-expanded/collapsed by default.
+ */
+UENUM(BlueprintType)
+enum class ESoundDashboardAutoExpandOptions : uint8
+{
+	/** Auto-expand new categories. */
+	Categories		UMETA(DisplayName = "Categories"),
+
+	/** Auto-expand all new categories and sounds. */
+	Everything		UMETA(DisplayName = "Everything"),
+
+	/** Don't auto-expand anything. */
+	Nothing			UMETA(DisplayName = "Nothing")
+};
+
+USTRUCT()
+struct FSoundDashboardSettings
+{
+	GENERATED_BODY()
+
+	UPROPERTY(EditAnywhere, config, Category = SoundDashboard, meta = (DisplayName = "Amp (Peak) Display Mode"))
+	EAudioAmplitudeDisplayMode AmplitudeDisplayMode = EAudioAmplitudeDisplayMode::Decibels;
+
+	UPROPERTY(EditAnywhere, config, Category = SoundDashboard, meta = (DisplayName = "View"))
+	ESoundDashboardTreeViewingOptions TreeViewMode = ESoundDashboardTreeViewingOptions::FullTree;
+
+	UPROPERTY(EditAnywhere, config, Category = SoundDashboard, meta = (DisplayName = "Auto-Expand"))
+	ESoundDashboardAutoExpandOptions AutoExpandMode = ESoundDashboardAutoExpandOptions::Categories;
+
+	UPROPERTY(EditAnywhere, config, Category = SoundDashboard, meta = (DisplayName = "Show Stopped Sounds"))
+	bool bShowStoppedSounds = false;
+
+	UPROPERTY(EditAnywhere, config, Category = SoundDashboard, meta = (ClampMin = "0", DisplayName = "Recently Stopped Sounds Timeout"))
+	float StoppedSoundTimeoutTime = 3.0f;
+
+#if WITH_EDITOR
+	DECLARE_MULTICAST_DELEGATE_OneParam(FOnReadSoundDashboardSettings, const FSoundDashboardSettings&);
+	static AUDIOINSIGHTS_API FOnReadSoundDashboardSettings OnReadSettings;
+
+	DECLARE_MULTICAST_DELEGATE_OneParam(FOnWriteSoundDashboardSettings, FSoundDashboardSettings&);
+	static AUDIOINSIGHTS_API FOnWriteSoundDashboardSettings OnWriteSettings;
+
+	DECLARE_MULTICAST_DELEGATE(FOnRequestReadSoundDashboardSettings);
+	static AUDIOINSIGHTS_API FOnRequestReadSoundDashboardSettings OnRequestReadSettings;
+
+	DECLARE_MULTICAST_DELEGATE(FOnRequestWriteSoundDashboardSettings);
+	static AUDIOINSIGHTS_API FOnRequestWriteSoundDashboardSettings OnRequestWriteSettings;
+#endif
+};
+
+#undef LOCTEXT_NAMESPACE
\ No newline at end of file
diff --git a/Engine/Plugins/AudioInsights/Source/AudioInsights/Public/Views/SoundDashboardViewFactory.h b/Engine/Plugins/AudioInsights/Source/AudioInsights/Public/Views/SoundDashboardViewFactory.h
--- a/Engine/Plugins/AudioInsights/Source/AudioInsights/Public/Views/SoundDashboardViewFactory.h
+++ b/Engine/Plugins/AudioInsights/Source/AudioInsights/Public/Views/SoundDashboardViewFactory.h
@@ -1,6 +1,7 @@
 // Copyright Epic Games, Inc. All Rights Reserved.
 #pragma once
 
+#include "Settings/SoundDashboardSettings.h"
 #include "Views/TreeDashboardViewFactory.h"
 
 #include "Filters/GenericFilter.h"
@@ -132,24 +133,10 @@ namespace UE::Audio::Insights
 		virtual int32 GetDisplayedTreeDepth() const override { return static_cast<int32>(TreeViewingMode); }
 		UE_API virtual bool ShouldAutoExpand(const TSharedPtr<IDashboardDataTreeViewEntry>& Item) const override;
 
-		enum class ETreeViewingOptions : uint8
-		{
-			FullTree = 0,
-			ActiveSounds,
-			FlatList
-		};
-
-		UE_API void ChangeTreeViewingMode(const ETreeViewingOptions SelectedMode);
+		UE_API void ChangeTreeViewingMode(const ESoundDashboardTreeViewingOptions SelectedMode, const bool bUpdateEditorSettings = true);
 		UE_API void BuildViewOptionsMenuContent(FMenuBuilder& MenuBuilder);
 
-		enum class EAutoExpandOptions : uint8
-		{
-			Categories = 0,
-			Everything,
-			Nothing
-		};
-
-		UE_API void ChangeAutoExpandMode(const EAutoExpandOptions SelectedMode);
+		UE_API void ChangeAutoExpandMode(const ESoundDashboardAutoExpandOptions SelectedMode, const bool bUpdateEditorSettings = true);
 		UE_API void RefreshExpansionChecks(const TSharedPtr<IDashboardDataTreeViewEntry>& Entry);
 		UE_API void BuildAutoExpandMenuContent(FMenuBuilder& MenuBuilder);
 
@@ -160,7 +147,12 @@ namespace UE::Audio::Insights
 		UE_API TSharedRef<SWidget> OnGetSettingsMenuContent();
 		UE_API TSharedRef<SWidget> MakeSettingsButtonWidget();
 
+		void UpdateMaxStoppedSoundTimeout();
+
 #if WITH_EDITOR
+		void OnReadEditorSettings(const FSoundDashboardSettings& Settings);
+		void OnWriteEditorSettings(FSoundDashboardSettings& Settings);
+
 		UE_API TSharedRef<SWidget> MakeMuteSoloWidget();
 
 		UE_API TSharedRef<SWidget> CreateMuteSoloButton(const TSharedRef<FTraceTreeDashboardViewFactory::SRowWidget>& InRowWidget, 
@@ -212,9 +204,12 @@ namespace UE::Audio::Insights
 		TSharedPtr<SWidget> SoundsFilterBar;
 		TSharedPtr<SWidget> SoundsFilterBarButton;
 
-		ETreeViewingOptions TreeViewingMode = ETreeViewingOptions::FullTree;
-		EAutoExpandOptions AutoExpandMode = EAutoExpandOptions::Categories;
+		ESoundDashboardTreeViewingOptions TreeViewingMode = ESoundDashboardTreeViewingOptions::FullTree;
+		ESoundDashboardAutoExpandOptions AutoExpandMode = ESoundDashboardAutoExpandOptions::Categories;
 		ESoundDashboardFilterFlags SelectedFilterFlags = AllFilterFlags;
+
+		float RecentlyStoppedSoundTimeout = 3.0f;
+
 		bool bIsPinnedCategoryFilterEnabled = true;
 		bool bShowRecentlyStoppedSounds = false;
 		bool bTreeViewModeChanged = false;
diff --git a/Engine/Plugins/AudioInsights/Source/AudioInsightsEditor/Private/AudioInsightsEditorSettings.cpp b/Engine/Plugins/AudioInsights/Source/AudioInsightsEditor/Private/AudioInsightsEditorSettings.cpp
--- a/Engine/Plugins/AudioInsights/Source/AudioInsightsEditor/Private/AudioInsightsEditorSettings.cpp
+++ b/Engine/Plugins/AudioInsights/Source/AudioInsightsEditor/Private/AudioInsightsEditorSettings.cpp
@@ -21,4 +21,42 @@ FText UAudioInsightsEditorSettings::GetSectionDescription() const
 	return LOCTEXT("AudioInsightsEditorSettings_SectionDesc", "Configure Audio Insights editor settings.");
 }
 
+void UAudioInsightsEditorSettings::PostInitProperties()
+{
+	Super::PostInitProperties();
+
+	OnRequestReadSettingsHandle = FSoundDashboardSettings::OnRequestReadSettings.AddLambda([this]() { FSoundDashboardSettings::OnReadSettings.Broadcast(SoundDashboardSettings); });
+	OnRequestWriteSettingsHandle = FSoundDashboardSettings::OnRequestWriteSettings.AddLambda([this]() { FSoundDashboardSettings::OnWriteSettings.Broadcast(SoundDashboardSettings); });
+}
+
+void UAudioInsightsEditorSettings::BeginDestroy()
+{
+	Super::BeginDestroy();
+
+	FSoundDashboardSettings::OnRequestReadSettings.Remove(OnRequestReadSettingsHandle);
+	FSoundDashboardSettings::OnRequestWriteSettings.Remove(OnRequestWriteSettingsHandle);
+}
+
+void UAudioInsightsEditorSettings::PostEditChangeChainProperty(FPropertyChangedChainEvent& PropertyChangedEvent)
+{
+	Super::PostEditChangeChainProperty(PropertyChangedEvent);
+
+	if (PropertyChangedEvent.Property == nullptr)
+	{
+		return;
+	}
+
+	const TDoubleLinkedList<FProperty*>::TDoubleLinkedListNode* ActiveMemberNode = PropertyChangedEvent.PropertyChain.GetActiveMemberNode();
+
+	if (ActiveMemberNode == nullptr || ActiveMemberNode->GetValue() == nullptr)
+	{
+		return;
+	}
+
+	if (ActiveMemberNode->GetValue()->GetFName() == GET_MEMBER_NAME_CHECKED(UAudioInsightsEditorSettings, SoundDashboardSettings))
+	{
+		FSoundDashboardSettings::OnReadSettings.Broadcast(SoundDashboardSettings);
+	}
+}
+
 #undef LOCTEXT_NAMESPACE
diff --git a/Engine/Plugins/AudioInsights/Source/AudioInsightsEditor/Public/AudioInsightsEditorSettings.h b/Engine/Plugins/AudioInsights/Source/AudioInsightsEditor/Public/AudioInsightsEditorSettings.h
--- a/Engine/Plugins/AudioInsights/Source/AudioInsightsEditor/Public/AudioInsightsEditorSettings.h
+++ b/Engine/Plugins/AudioInsights/Source/AudioInsightsEditor/Public/AudioInsightsEditorSettings.h
@@ -5,6 +5,7 @@
 #include "AudioSpectrogram.h"
 #include "AudioSpectrumAnalyzer.h"
 #include "Engine/DeveloperSettings.h"
+#include "Settings/SoundDashboardSettings.h"
 
 #include "AudioInsightsEditorSettings.generated.h"
 
@@ -20,6 +21,10 @@ public:
 	UE_API virtual FText GetSectionText() const override;
 	UE_API virtual FText GetSectionDescription() const override;
 
+	UE_API virtual void PostInitProperties() override;
+	UE_API virtual void BeginDestroy() override;
+	UE_API virtual void PostEditChangeChainProperty(struct FPropertyChangedChainEvent& PropertyChangedEvent) override;
+
 	/** Whether to automatically set the first PIE client in Audio Insights World Filter. */
 	UPROPERTY(Config, EditAnywhere, Category="World Filter")
 	bool bWorldFilterDefaultsToFirstClient = false;
@@ -31,6 +36,14 @@ public:
 	/** Settings for analyzer rack spectrum analyzer widget */
 	UPROPERTY(EditAnywhere, config, Category = SpectrumAnalyzer, meta = (ShowOnlyInnerProperties))
 	FSpectrumAnalyzerRackUnitSettings SpectrumAnalyzerSettings;
+
+	/** Settings for the Sound Dashboard */
+	UPROPERTY(EditAnywhere, config, Category = SoundDashboard, meta = (ShowOnlyInnerProperties))
+	FSoundDashboardSettings SoundDashboardSettings;
+
+private:
+	FDelegateHandle OnRequestReadSettingsHandle;
+	FDelegateHandle OnRequestWriteSettingsHandle;
 };
 
 #undef UE_API
