[Audio Insights] Add columns for Play Order, Distance Attenuation, LPF & HPF Frequency to the Sounds Dashboard. Update method of averaging accumulated aggregate data on active sound entries to only include sound wave entries where data has been set.

diff --git a/Engine/Plugins/AudioInsights/Source/AudioInsights/Private/Messages/SoundTraceMessages.cpp b/Engine/Plugins/AudioInsights/Source/AudioInsights/Private/Messages/SoundTraceMessages.cpp
--- a/Engine/Plugins/AudioInsights/Source/AudioInsights/Private/Messages/SoundTraceMessages.cpp
+++ b/Engine/Plugins/AudioInsights/Source/AudioInsights/Private/Messages/SoundTraceMessages.cpp
@@ -67,6 +67,9 @@ namespace UE::Audio::Insights
 
 		PriorityDataPoints.SetCapacity(DataPointsCapacity);
 		DistanceDataPoints.SetCapacity(DataPointsCapacity);
+		DistanceAttenuationDataPoints.SetCapacity(DataPointsCapacity);
+		LPFFreqDataPoints.SetCapacity(DataPointsCapacity);
+		HPFFreqDataPoints.SetCapacity(DataPointsCapacity);
 		AmplitudeDataPoints.SetCapacity(DataPointsCapacity);
 		VolumeDataPoints.SetCapacity(DataPointsCapacity);
 		PitchDataPoints.SetCapacity(DataPointsCapacity);
diff --git a/Engine/Plugins/AudioInsights/Source/AudioInsights/Private/Providers/SoundTraceProvider.cpp b/Engine/Plugins/AudioInsights/Source/AudioInsights/Private/Providers/SoundTraceProvider.cpp
--- a/Engine/Plugins/AudioInsights/Source/AudioInsights/Private/Providers/SoundTraceProvider.cpp
+++ b/Engine/Plugins/AudioInsights/Source/AudioInsights/Private/Providers/SoundTraceProvider.cpp
@@ -209,6 +209,8 @@ namespace UE::Audio::Insights
 				Builder.RouteEvent(RouteId_Stop,				"Audio", "SoundStop");
 				Builder.RouteEvent(RouteId_Priority,			"Audio", "SoundPriority");
 				Builder.RouteEvent(RouteId_Distance,			"Audio", "SoundDistance");
+				Builder.RouteEvent(RouteId_DistanceAttenuation,	"Audio", "MixerSourceDistanceAttenuation");
+				Builder.RouteEvent(RouteId_Filters,				"Audio", "MixerSourceFilters");
 				Builder.RouteEvent(RouteId_Amplitude,			"Audio", "MixerSourceEnvelope");
 				Builder.RouteEvent(RouteId_Volume,				"Audio", "MixerSourceVolume");
 				Builder.RouteEvent(RouteId_Pitch,				"Audio", "MixerSourcePitch");
@@ -253,6 +255,19 @@ namespace UE::Audio::Insights
 						break;
 					}
 
+					case RouteId_DistanceAttenuation:
+					{
+						Messages.DistanceAttenuationMessages.Enqueue(FMixerSourceDistanceAttenuationMessage{ Context });
+						break;
+					}
+
+					case RouteId_Filters:
+					{
+						Messages.LPFFreqMessages.Enqueue(FMixerSourceLPFFreqMessage{ Context });
+						Messages.HPFFreqMessages.Enqueue(FMixerSourceHPFFreqMessage{ Context });
+						break;
+					}
+
 					case RouteId_Amplitude:
 					{
 						Messages.AmplitudeMessages.Enqueue(FMixerSourceEnvelopeMessage{ Context });
@@ -301,6 +316,8 @@ namespace UE::Audio::Insights
 				RouteId_Stop,
 				RouteId_Priority,
 				RouteId_Distance,
+				RouteId_DistanceAttenuation,
+				RouteId_Filters,
 				RouteId_Amplitude,
 				RouteId_Volume,
 				RouteId_Pitch,
@@ -467,6 +484,66 @@ namespace UE::Audio::Insights
 		});
 		EntriesWithPoppedDataPoints.Reset();
 
+		ProcessMessageQueue<FMixerSourceDistanceAttenuationMessage>(TraceMessages.DistanceAttenuationMessages, GetActiveSoundEntryFromParameterMessage,
+		[this, &EntriesWithPoppedDataPoints](const FMixerSourceDistanceAttenuationMessage& Msg, const TSharedPtr<FSoundDashboardEntry>* OutActiveSoundEntry)
+		{
+			FSoundDashboardEntry* SoundEntry = FSoundTraceProviderPrivate::FindSoundEntryForActiveSound(Msg.PlayOrder, OutActiveSoundEntry);
+			if (SoundEntry == nullptr)
+			{
+				return;
+			}
+
+#if !WITH_EDITOR
+			if (SessionCachedMessages.IsValid())
+			{
+				SessionCachedMessages->DistanceAttenuationCachedMessages.EmplaceBack(Msg);
+			}
+#endif // !WITH_EDITOR
+
+			FSoundTraceProviderPrivate::UpdateParameterEntry(Msg, Msg.DistanceAttenuation, SoundEntry->DistanceAttenuationDataPoints, EntriesWithPoppedDataPoints);
+		});
+		EntriesWithPoppedDataPoints.Reset();
+
+		ProcessMessageQueue<FMixerSourceLPFFreqMessage>(TraceMessages.LPFFreqMessages, GetActiveSoundEntryFromParameterMessage,
+		[this, &EntriesWithPoppedDataPoints](const FMixerSourceLPFFreqMessage& Msg, const TSharedPtr<FSoundDashboardEntry>* OutActiveSoundEntry)
+		{
+			FSoundDashboardEntry* SoundEntry = FSoundTraceProviderPrivate::FindSoundEntryForActiveSound(Msg.PlayOrder, OutActiveSoundEntry);
+			if (SoundEntry == nullptr)
+			{
+				return;
+			}
+
+#if !WITH_EDITOR
+			if (SessionCachedMessages.IsValid())
+			{
+				SessionCachedMessages->LPFFreqCachedMessages.EmplaceBack(Msg);
+			}
+#endif // !WITH_EDITOR
+
+			FSoundTraceProviderPrivate::UpdateParameterEntry(Msg, Msg.LPFFrequency, SoundEntry->LPFFreqDataPoints, EntriesWithPoppedDataPoints);
+		});
+		EntriesWithPoppedDataPoints.Reset();
+
+		ProcessMessageQueue<FMixerSourceHPFFreqMessage>(TraceMessages.HPFFreqMessages, GetActiveSoundEntryFromParameterMessage,
+		[this, &EntriesWithPoppedDataPoints](const FMixerSourceHPFFreqMessage& Msg, const TSharedPtr<FSoundDashboardEntry>* OutActiveSoundEntry)
+		{
+			FSoundDashboardEntry* SoundEntry = FSoundTraceProviderPrivate::FindSoundEntryForActiveSound(Msg.PlayOrder, OutActiveSoundEntry);
+			if (SoundEntry == nullptr)
+			{
+				return;
+			}
+
+#if !WITH_EDITOR
+			if (SessionCachedMessages.IsValid())
+			{
+				SessionCachedMessages->HPFFreqCachedMessages.EmplaceBack(Msg);
+			}
+#endif // !WITH_EDITOR
+
+			FSoundTraceProviderPrivate::UpdateParameterEntry(Msg, Msg.HPFFrequency, SoundEntry->HPFFreqDataPoints, EntriesWithPoppedDataPoints);
+		});
+		EntriesWithPoppedDataPoints.Reset();
+
 		ProcessMessageQueue<FMixerSourceEnvelopeMessage>(TraceMessages.AmplitudeMessages, GetActiveSoundEntryFromParameterMessage,
 		[this, &EntriesWithPoppedDataPoints](const FMixerSourceEnvelopeMessage& Msg, const TSharedPtr<FSoundDashboardEntry>* OutActiveSoundEntry)
 		{
@@ -645,10 +722,10 @@ namespace UE::Audio::Insights
 			}
 
 			TSharedPtr<IDashboardDataTreeViewEntry>* SoundEntry = CategoryEntry->Children.FindByPredicate([&Msg](TSharedPtr<IDashboardDataTreeViewEntry> ChildEntry)
-				{
-					FSoundDashboardEntry& NewEntry = FSoundTraceProviderPrivate::CastEntry(*ChildEntry.Get());
-					return NewEntry.PlayOrder == Msg.PlayOrder;
-				});
+			{
+				FSoundDashboardEntry& NewEntry = FSoundTraceProviderPrivate::CastEntry(*ChildEntry.Get());
+				return NewEntry.PlayOrder == Msg.PlayOrder;
+			});
 
 			if (SoundEntry == nullptr || !SoundEntry->IsValid())
 			{
@@ -781,19 +858,36 @@ namespace UE::Audio::Insights
 			return First.Value < Second.Value ? First : Second;
 		};
 
-		auto AccumulateDataPoints = [](const FDataPoint& AccumulativeValue, const FDataPoint& AdditionalValue) -> FDataPoint
+		struct FCumulativeDataPoint
+		{
+			FDataPoint TotalAccumulationDataPoint = { 0.0, 0.0f };
+			int32 NumDataPoints = 0;
+		};
+
+		auto AccumulateDataPoints = [](const FCumulativeDataPoint& AccumulativeValue, const ::Audio::TCircularAudioBuffer<FDataPoint>& DataBuffer) -> FCumulativeDataPoint
 		{
-			return { FMath::Max(AccumulativeValue.Key, AdditionalValue.Key), AccumulativeValue.Value + AdditionalValue.Value };
+			if (DataBuffer.Num() == 0)
+			{
+				return AccumulativeValue;
+			}
+
+			const FDataPoint AdditionalValue = PeekLastValue(DataBuffer);
+
+			return
+			{
+				{ FMath::Max(AccumulativeValue.TotalAccumulationDataPoint.Key, AdditionalValue.Key), AccumulativeValue.TotalAccumulationDataPoint.Value + AdditionalValue.Value },
+				AccumulativeValue.NumDataPoints + 1
+			};
 		};
 
-		auto AverageDataPoints = [](const FDataPoint& DataPointsTotal, const int32 NumDataPoints) -> FDataPoint
+		auto AverageDataPoints = [](const FCumulativeDataPoint& DataPointsTotal) -> FDataPoint
 		{
-			if (NumDataPoints == 0)
+			if (DataPointsTotal.NumDataPoints == 0)
 			{
 				return { 0.0, 0.0f };
 			}
 
-			return { DataPointsTotal.Key,  DataPointsTotal.Value / static_cast<float>(NumDataPoints) };
+			return { DataPointsTotal.TotalAccumulationDataPoint.Key, DataPointsTotal.TotalAccumulationDataPoint.Value / static_cast<float>(DataPointsTotal.NumDataPoints) };
 		};
 
 		auto AddSingleValueToBuffer = [](::Audio::TCircularAudioBuffer<FDataPoint>& To, const FDataPoint& Value)
@@ -811,12 +905,21 @@ namespace UE::Audio::Insights
 		FDataPoint MaxPriority { 0.0, 0.0f };
 
 		FDataPoint AveragePeakAmp { 0.0, 0.0f };
-		FDataPoint CumulativePeakAmplitude { 0.0, 0.0f };
+		FCumulativeDataPoint CumulativePeakAmplitude;
 
 		FDataPoint AveragePitch { 0.0, 0.0f };
-		FDataPoint CumulativePitch { 0.0, 0.0f };
+		FCumulativeDataPoint CumulativePitch;
+
+		FDataPoint AverageDistanceAttenuation { 0.0, 0.0f };
+		FCumulativeDataPoint CumulativeDistanceAttenuation;
+
+		FDataPoint AverageLPFFreq { 0.0, 0.0f };
+		FCumulativeDataPoint CumulativeLPFFreq;
 
-		FDataPoint CumulativeRelativeRenderCost { 0.0, 0.0f };
+		FDataPoint AverageHPFFreq { 0.0, 0.0f };
+		FCumulativeDataPoint CumulativeHPFFreq;
+
+		FCumulativeDataPoint CumulativeRelativeRenderCost;
 
 		for (const TSharedPtr<IDashboardDataTreeViewEntry>& WaveInstanceEntry : ActiveSoundEntry.Children)
 		{
@@ -826,20 +929,29 @@ namespace UE::Audio::Insights
 			MinDistance = GetMinDatapoint(MinDistance, PeekLastValue(WaveSoundEntry.DistanceDataPoints, MinDistance.Value));
 			MaxPriority = GetMaxDatapoint(MaxPriority, PeekLastValue(WaveSoundEntry.PriorityDataPoints, MaxPriority.Value));
 
-			CumulativePeakAmplitude = AccumulateDataPoints(CumulativePeakAmplitude, PeekLastValue(WaveSoundEntry.AmplitudeDataPoints));
-			CumulativePitch = AccumulateDataPoints(CumulativePitch, PeekLastValue(WaveSoundEntry.PitchDataPoints));
-			CumulativeRelativeRenderCost = AccumulateDataPoints(CumulativeRelativeRenderCost, PeekLastValue(WaveSoundEntry.RelativeRenderCostDataPoints));
+			CumulativePeakAmplitude = AccumulateDataPoints(CumulativePeakAmplitude, WaveSoundEntry.AmplitudeDataPoints);
+			CumulativePitch = AccumulateDataPoints(CumulativePitch, WaveSoundEntry.PitchDataPoints);
+			CumulativeDistanceAttenuation = AccumulateDataPoints(CumulativeDistanceAttenuation, WaveSoundEntry.DistanceAttenuationDataPoints);
+			CumulativeLPFFreq = AccumulateDataPoints(CumulativeLPFFreq, WaveSoundEntry.LPFFreqDataPoints);
+			CumulativeHPFFreq = AccumulateDataPoints(CumulativeHPFFreq, WaveSoundEntry.HPFFreqDataPoints);
+			CumulativeRelativeRenderCost = AccumulateDataPoints(CumulativeRelativeRenderCost, WaveSoundEntry.RelativeRenderCostDataPoints);
 		}
 
-		AveragePeakAmp = AverageDataPoints(CumulativePeakAmplitude, NumWaveInstances);
-		AveragePitch = AverageDataPoints(CumulativePitch, NumWaveInstances);
+		AveragePeakAmp = AverageDataPoints(CumulativePeakAmplitude);
+		AveragePitch = AverageDataPoints(CumulativePitch);
+		AverageDistanceAttenuation = AverageDataPoints(CumulativeDistanceAttenuation);
+		AverageLPFFreq = AverageDataPoints(CumulativeLPFFreq);
+		AverageHPFFreq = AverageDataPoints(CumulativeHPFFreq);
 
 		AddSingleValueToBuffer(ActiveSoundEntry.VolumeDataPoints, MaxVolume);
 		AddSingleValueToBuffer(ActiveSoundEntry.DistanceDataPoints, MinDistance);
+		AddSingleValueToBuffer(ActiveSoundEntry.DistanceAttenuationDataPoints, AverageDistanceAttenuation);
+		AddSingleValueToBuffer(ActiveSoundEntry.LPFFreqDataPoints, AverageLPFFreq);
+		AddSingleValueToBuffer(ActiveSoundEntry.HPFFreqDataPoints, AverageHPFFreq);
 		AddSingleValueToBuffer(ActiveSoundEntry.PriorityDataPoints, MaxPriority);
 		AddSingleValueToBuffer(ActiveSoundEntry.AmplitudeDataPoints, AveragePeakAmp);
 		AddSingleValueToBuffer(ActiveSoundEntry.PitchDataPoints, AveragePitch);
-		AddSingleValueToBuffer(ActiveSoundEntry.RelativeRenderCostDataPoints, CumulativeRelativeRenderCost);
+		AddSingleValueToBuffer(ActiveSoundEntry.RelativeRenderCostDataPoints, CumulativeRelativeRenderCost.TotalAccumulationDataPoint);
 	}
 
 #if !WITH_EDITOR
@@ -980,6 +1092,9 @@ namespace UE::Audio::Insights
 			{
 				FSoundPriorityMessage PriorityMessage;
 				FSoundDistanceMessage DistanceMessage;
+				FMixerSourceDistanceAttenuationMessage DistanceAttenuationMessage;
+				FMixerSourceLPFFreqMessage LPFFreqMessage;
+				FMixerSourceHPFFreqMessage HPFFreqMessage;
 				FMixerSourceEnvelopeMessage AmplitudeMessage;
 				FMixerSourceVolumeMessage VolumeMessage;
 				FMixerSourcePitchMessage PitchMessage;
@@ -1012,6 +1127,27 @@ namespace UE::Audio::Insights
 					CachedEntryInfos[Index].DistanceMessage = *FoundDistanceCachedMessage;
 				}
 
+				// Distance Attenuation
+				const FMixerSourceDistanceAttenuationMessage* FoundDistanceAttenuationCachedMessage = FindClosestMessageToTimestamp(SessionCachedMessages->DistanceAttenuationCachedMessages, TimeMarker, PlayOrder);
+				if (FoundDistanceAttenuationCachedMessage)
+				{
+					CachedEntryInfos[Index].DistanceAttenuationMessage = *FoundDistanceAttenuationCachedMessage;
+				}
+
+				// LPF Freq
+				const FMixerSourceLPFFreqMessage* FoundLPFFreqCachedMessage = FindClosestMessageToTimestamp(SessionCachedMessages->LPFFreqCachedMessages, TimeMarker, PlayOrder);
+				if (FoundLPFFreqCachedMessage)
+				{
+					CachedEntryInfos[Index].LPFFreqMessage = *FoundLPFFreqCachedMessage;
+				}
+
+				// HPF Freq
+				const FMixerSourceHPFFreqMessage* FoundHPFFreqCachedMessage = FindClosestMessageToTimestamp(SessionCachedMessages->HPFFreqCachedMessages, TimeMarker, PlayOrder);
+				if (FoundHPFFreqCachedMessage)
+				{
+					CachedEntryInfos[Index].HPFFreqMessage = *FoundHPFFreqCachedMessage;
+				}
+
 				// Amplitude
 				const FMixerSourceEnvelopeMessage* FoundAmplitudeCachedMessage = FindClosestMessageToTimestamp(SessionCachedMessages->AmplitudeCachedMessages, TimeMarker, PlayOrder);
 				if (FoundAmplitudeCachedMessage)
@@ -1059,6 +1195,9 @@ namespace UE::Audio::Insights
 
 				SoundEntry->PriorityDataPoints.Push({ CachedEntryInfo.PriorityMessage.Timestamp, CachedEntryInfo.PriorityMessage.Priority });
 				SoundEntry->DistanceDataPoints.Push({ CachedEntryInfo.DistanceMessage.Timestamp, CachedEntryInfo.DistanceMessage.Distance });
+				SoundEntry->DistanceAttenuationDataPoints.Push({ CachedEntryInfo.DistanceAttenuationMessage.Timestamp, CachedEntryInfo.DistanceAttenuationMessage.DistanceAttenuation });
+				SoundEntry->LPFFreqDataPoints.Push({ CachedEntryInfo.LPFFreqMessage.Timestamp, CachedEntryInfo.LPFFreqMessage.LPFFrequency });
+				SoundEntry->HPFFreqDataPoints.Push({ CachedEntryInfo.HPFFreqMessage.Timestamp, CachedEntryInfo.HPFFreqMessage.HPFFrequency });
 				SoundEntry->AmplitudeDataPoints.Push({ CachedEntryInfo.AmplitudeMessage.Timestamp, CachedEntryInfo.AmplitudeMessage.Envelope });
 				SoundEntry->VolumeDataPoints.Push({ CachedEntryInfo.VolumeMessage.Timestamp, CachedEntryInfo.VolumeMessage.Volume });
 				SoundEntry->PitchDataPoints.Push({ CachedEntryInfo.PitchMessage.Timestamp, CachedEntryInfo.PitchMessage.Pitch });
diff --git a/Engine/Plugins/AudioInsights/Source/AudioInsights/Private/Views/SoundDashboardViewFactory.cpp b/Engine/Plugins/AudioInsights/Source/AudioInsights/Private/Views/SoundDashboardViewFactory.cpp
--- a/Engine/Plugins/AudioInsights/Source/AudioInsights/Private/Views/SoundDashboardViewFactory.cpp
+++ b/Engine/Plugins/AudioInsights/Source/AudioInsights/Private/Views/SoundDashboardViewFactory.cpp
@@ -456,6 +456,9 @@ namespace UE::Audio::Insights
 
 			AddLastValueInBuffer(PinnedEntry.PriorityDataPoints, OriginalEntry.PriorityDataPoints);
 			AddLastValueInBuffer(PinnedEntry.DistanceDataPoints, OriginalEntry.DistanceDataPoints);
+			AddLastValueInBuffer(PinnedEntry.DistanceAttenuationDataPoints, OriginalEntry.DistanceAttenuationDataPoints);
+			AddLastValueInBuffer(PinnedEntry.LPFFreqDataPoints, OriginalEntry.LPFFreqDataPoints);
+			AddLastValueInBuffer(PinnedEntry.HPFFreqDataPoints, OriginalEntry.HPFFreqDataPoints);
 			AddLastValueInBuffer(PinnedEntry.AmplitudeDataPoints, OriginalEntry.AmplitudeDataPoints);
 			AddLastValueInBuffer(PinnedEntry.VolumeDataPoints, OriginalEntry.VolumeDataPoints);
 			AddLastValueInBuffer(PinnedEntry.PitchDataPoints, OriginalEntry.PitchDataPoints);
@@ -1464,6 +1467,17 @@ namespace UE::Audio::Insights
 					.Alignment        = EHorizontalAlignment::HAlign_Left
 				}
 			},
+			{
+				"PlayOrder",
+				{
+					.DisplayName = LOCTEXT("SoundDashboard_PlayOrderColumnDisplayName", "Play Order"),
+					.IconName = NAME_None,
+					.bShowDisplayName = true,
+					.bDefaultHidden = true,
+					.FillWidth = 0.1f,
+					.Alignment = EHorizontalAlignment::HAlign_Left
+				}
+			},
 			{
 				"Priority",
 				{
@@ -1486,6 +1500,17 @@ namespace UE::Audio::Insights
 					.Alignment        = EHorizontalAlignment::HAlign_Left
 				}
 			},
+			{
+				"DistanceAttenuation",
+				{
+					.DisplayName = LOCTEXT("SoundDashboard_DistanceAttenuationColumnDisplayName", "Distance Attenuation"),
+					.IconName = NAME_None,
+					.bShowDisplayName = true,
+					.bDefaultHidden = true,
+					.FillWidth = 0.1f,
+					.Alignment = EHorizontalAlignment::HAlign_Left
+				}
+			},
 			{
 				"Amplitude",
 				{
@@ -1508,6 +1533,28 @@ namespace UE::Audio::Insights
 					.Alignment        = EHorizontalAlignment::HAlign_Left
 				}
 			},
+			{
+				"LPFFreq",
+				{
+					.DisplayName = LOCTEXT("SoundDashboard_LPFFreqColumnDisplayName", "LPF Freq (Hz)"),
+					.IconName = NAME_None,
+					.bShowDisplayName = true,
+					.bDefaultHidden = true,
+					.FillWidth = 0.1f,
+					.Alignment = EHorizontalAlignment::HAlign_Left
+				}
+			},
+			{
+				"HPFFreq",
+				{
+					.DisplayName = LOCTEXT("SoundDashboard_HPFFreqColumnDisplayName", "HPF Freq (Hz)"),
+					.IconName = NAME_None,
+					.bShowDisplayName = true,
+					.bDefaultHidden = true,
+					.FillWidth = 0.1f,
+					.Alignment = EHorizontalAlignment::HAlign_Left
+				}
+			},
 			{
 				"Pitch",
 				{
@@ -1621,6 +1668,21 @@ namespace UE::Audio::Insights
 					}
 				}
 			},
+			{
+				"PlayOrder",
+				{
+					.GetDisplayValue = [](const IDashboardDataTreeViewEntry& InData)
+					{
+						const FSoundDashboardEntry& SoundDashboardEntry = CastEntry(InData);
+						if (SoundDashboardEntry.bIsCategory)
+						{
+							return FText::GetEmpty();
+						}
+
+						return FText::AsNumber(SoundDashboardEntry.PlayOrder);
+					}
+				}
+			},
 			{
 				"Priority",
 				{
@@ -1658,6 +1720,21 @@ namespace UE::Audio::Insights
 					}
 				}
 			},
+			{
+				"DistanceAttenuation",
+				{
+					.GetDisplayValue = [](const IDashboardDataTreeViewEntry& InData)
+					{
+						const FSoundDashboardEntry& SoundDashboardEntry = CastEntry(InData);
+						if (SoundDashboardEntry.bIsCategory)
+						{
+							return FText::GetEmpty();
+						}
+
+						return FText::AsNumber(GetLastEntryArrayValue(SoundDashboardEntry.DistanceAttenuationDataPoints), FSlateStyle::Get().GetLinearVolumeFloatFormat());
+					}
+				}
+			},
 			{
 				"Amplitude",
 				{
@@ -1692,6 +1769,36 @@ namespace UE::Audio::Insights
 					}
 				}
 			},
+			{
+				"LPFFreq",
+				{
+					.GetDisplayValue = [](const IDashboardDataTreeViewEntry& InData)
+					{
+						const FSoundDashboardEntry& SoundDashboardEntry = CastEntry(InData);
+						if (SoundDashboardEntry.bIsCategory)
+						{
+							return FText::GetEmpty();
+						}
+
+						return FText::AsNumber(GetLastEntryArrayValue(SoundDashboardEntry.LPFFreqDataPoints), FSlateStyle::Get().GetFreqFloatFormat());
+					}
+				}
+			},
+			{
+				"HPFFreq",
+				{
+					.GetDisplayValue = [](const IDashboardDataTreeViewEntry& InData)
+					{
+						const FSoundDashboardEntry& SoundDashboardEntry = CastEntry(InData);
+						if (SoundDashboardEntry.bIsCategory)
+						{
+							return FText::GetEmpty();
+						}
+
+						return FText::AsNumber(GetLastEntryArrayValue(SoundDashboardEntry.HPFFreqDataPoints), FSlateStyle::Get().GetFreqFloatFormat());
+					}
+				}
+			},
 			{
 				"Pitch",
 				{
@@ -1797,6 +1904,39 @@ namespace UE::Audio::Insights
 			return ComparisonDiff < 0.0f;
 		};
 
+		auto SortByDistanceAttenuation = [&SortByPlayOrder](const FSoundDashboardEntry& First, const FSoundDashboardEntry& Second)
+		{
+			const float ComparisonDiff = GetLastEntryArrayValue(First.DistanceAttenuationDataPoints) - GetLastEntryArrayValue(Second.DistanceAttenuationDataPoints);
+			if (FMath::IsNearlyZero(ComparisonDiff, UE_KINDA_SMALL_NUMBER))
+			{
+				return SortByPlayOrder(First, Second);
+			}
+
+			return ComparisonDiff < 0.0f;
+		};
+
+		auto SortByLPFFreq = [&SortByPlayOrder](const FSoundDashboardEntry& First, const FSoundDashboardEntry& Second)
+		{
+			const float ComparisonDiff = GetLastEntryArrayValue(First.LPFFreqDataPoints) - GetLastEntryArrayValue(Second.LPFFreqDataPoints);
+			if (FMath::IsNearlyZero(ComparisonDiff, UE_KINDA_SMALL_NUMBER))
+			{
+				return SortByPlayOrder(First, Second);
+			}
+
+			return ComparisonDiff < 0.0f;
+		};
+
+		auto SortByHPFFreq = [&SortByPlayOrder](const FSoundDashboardEntry& First, const FSoundDashboardEntry& Second)
+		{
+			const float ComparisonDiff = GetLastEntryArrayValue(First.HPFFreqDataPoints) - GetLastEntryArrayValue(Second.HPFFreqDataPoints);
+			if (FMath::IsNearlyZero(ComparisonDiff, UE_KINDA_SMALL_NUMBER))
+			{
+				return SortByPlayOrder(First, Second);
+			}
+
+			return ComparisonDiff < 0.0f;
+		};
+
 		auto SortByAmplitude = [&SortByPlayOrder](const FSoundDashboardEntry& First, const FSoundDashboardEntry& Second)
 		{
 			const float ComparisonDiff = GetLastEntryArrayValue(First.AmplitudeDataPoints) - GetLastEntryArrayValue(Second.AmplitudeDataPoints);
@@ -1869,6 +2009,10 @@ namespace UE::Audio::Insights
 		{
 			SortByPredicate(SortByName);
 		}
+		else if (SortByColumn == "PlayOrder")
+		{
+			SortByPredicate(SortByPlayOrder);
+		}
 		else if (SortByColumn == "Priority")
 		{
 			SortByPredicate(SortByPriority);
@@ -1877,6 +2021,10 @@ namespace UE::Audio::Insights
 		{
 			SortByPredicate(SortByDistance);
 		}
+		else if (SortByColumn == "DistanceAttenuation")
+		{
+			SortByPredicate(SortByDistanceAttenuation);
+		}
 		else if (SortByColumn == "Amplitude")
 		{
 			SortByPredicate(SortByAmplitude);
@@ -1885,6 +2033,14 @@ namespace UE::Audio::Insights
 		{
 			SortByPredicate(SortByVolume);
 		}
+		else if (SortByColumn == "LPFFreq")
+		{
+			SortByPredicate(SortByLPFFreq);
+		}
+		else if (SortByColumn == "HPFFreq")
+		{
+			SortByPredicate(SortByHPFFreq);
+		}
 		else if (SortByColumn == "Pitch")
 		{
 			SortByPredicate(SortByPitch);
diff --git a/Engine/Plugins/AudioInsights/Source/AudioInsights/Public/Messages/MixerSourceTraceMessages.h b/Engine/Plugins/AudioInsights/Source/AudioInsights/Public/Messages/MixerSourceTraceMessages.h
--- a/Engine/Plugins/AudioInsights/Source/AudioInsights/Public/Messages/MixerSourceTraceMessages.h
+++ b/Engine/Plugins/AudioInsights/Source/AudioInsights/Public/Messages/MixerSourceTraceMessages.h
@@ -77,9 +77,9 @@ namespace UE::Audio::Insights
 	};
 
 
-	DEFINE_MIXERSOURCE_PARAM_MESSAGE(FMixerSourceDistanceAttenuationMessage, FMixerSourceMessageBase, DistanceAttenuation, float, 0.0f)
-	DEFINE_MIXERSOURCE_PARAM_MESSAGE(FMixerSourceHPFFreqMessage, FMixerSourceMessageBase, HPFFrequency, float, MIN_FILTER_FREQUENCY)
-	DEFINE_MIXERSOURCE_PARAM_MESSAGE(FMixerSourceLPFFreqMessage, FMixerSourceMessageBase, LPFFrequency, float, MAX_FILTER_FREQUENCY)
+	DEFINE_MIXERSOURCE_PARAM_MESSAGE(FMixerSourceDistanceAttenuationMessage, FMixerSourceParameterMessage, DistanceAttenuation, float, 1.0f)
+	DEFINE_MIXERSOURCE_PARAM_MESSAGE(FMixerSourceHPFFreqMessage, FMixerSourceParameterMessage, HPFFrequency, float, MIN_FILTER_FREQUENCY)
+	DEFINE_MIXERSOURCE_PARAM_MESSAGE(FMixerSourceLPFFreqMessage, FMixerSourceParameterMessage, LPFFrequency, float, MAX_FILTER_FREQUENCY)
 	DEFINE_MIXERSOURCE_PARAM_MESSAGE(FMixerSourceEnvelopeMessage, FMixerSourceParameterMessage, Envelope, float, 0.0f)
 	DEFINE_MIXERSOURCE_PARAM_MESSAGE(FMixerSourcePitchMessage, FMixerSourceParameterMessage, Pitch, float, 1.0f)
 	DEFINE_MIXERSOURCE_PARAM_MESSAGE(FMixerSourceVolumeMessage, FMixerSourceParameterMessage, Volume, float, 1.0f)
diff --git a/Engine/Plugins/AudioInsights/Source/AudioInsights/Public/Messages/SoundTraceMessages.h b/Engine/Plugins/AudioInsights/Source/AudioInsights/Public/Messages/SoundTraceMessages.h
--- a/Engine/Plugins/AudioInsights/Source/AudioInsights/Public/Messages/SoundTraceMessages.h
+++ b/Engine/Plugins/AudioInsights/Source/AudioInsights/Public/Messages/SoundTraceMessages.h
@@ -102,6 +102,9 @@ namespace UE::Audio::Insights
 		TAnalyzerMessageQueue<FSoundWaveStartMessage> SoundWaveStartMessages{ 0.1 };
 		TAnalyzerMessageQueue<FSoundPriorityMessage> PriorityMessages{ 0.1 };
 		TAnalyzerMessageQueue<FSoundDistanceMessage> DistanceMessages { 0.1 };
+		TAnalyzerMessageQueue<FMixerSourceDistanceAttenuationMessage> DistanceAttenuationMessages { 0.1 };
+		TAnalyzerMessageQueue<FMixerSourceHPFFreqMessage> HPFFreqMessages{ 0.1 };
+		TAnalyzerMessageQueue<FMixerSourceLPFFreqMessage> LPFFreqMessages{ 0.1 };
 		TAnalyzerMessageQueue<FMixerSourceEnvelopeMessage> AmplitudeMessages { 0.1 };
 		TAnalyzerMessageQueue<FMixerSourceVolumeMessage> VolumeMessages { 0.1 };
 		TAnalyzerMessageQueue<FMixerSourcePitchMessage> PitchMessages { 0.1 };
@@ -117,6 +120,9 @@ namespace UE::Audio::Insights
 			, SoundWaveStartCachedMessages(InSession.GetLinearAllocator(), 16384)
 			, PriorityCachedMessages(InSession.GetLinearAllocator(), 16384)
 			, DistanceCachedMessages(InSession.GetLinearAllocator(), 16384)
+			, DistanceAttenuationCachedMessages(InSession.GetLinearAllocator(), 16384)
+			, HPFFreqCachedMessages(InSession.GetLinearAllocator(), 16384)
+			, LPFFreqCachedMessages(InSession.GetLinearAllocator(), 16384)
 			, AmplitudeCachedMessages(InSession.GetLinearAllocator(), 16384)
 			, VolumeCachedMessages(InSession.GetLinearAllocator(), 16384)
 			, PitchCachedMessages(InSession.GetLinearAllocator(), 16384)
@@ -130,6 +136,9 @@ namespace UE::Audio::Insights
 		TraceServices::TPagedArray<FSoundWaveStartMessage> SoundWaveStartCachedMessages;
 		TraceServices::TPagedArray<FSoundPriorityMessage> PriorityCachedMessages;
 		TraceServices::TPagedArray<FSoundDistanceMessage> DistanceCachedMessages;
+		TraceServices::TPagedArray<FMixerSourceDistanceAttenuationMessage> DistanceAttenuationCachedMessages;
+		TraceServices::TPagedArray<FMixerSourceHPFFreqMessage> HPFFreqCachedMessages;
+		TraceServices::TPagedArray<FMixerSourceLPFFreqMessage> LPFFreqCachedMessages;
 		TraceServices::TPagedArray<FMixerSourceEnvelopeMessage> AmplitudeCachedMessages;
 		TraceServices::TPagedArray<FMixerSourceVolumeMessage> VolumeCachedMessages;
 		TraceServices::TPagedArray<FMixerSourcePitchMessage> PitchCachedMessages;
@@ -202,6 +211,9 @@ namespace UE::Audio::Insights
 
 		::Audio::TCircularAudioBuffer<FDataPoint> PriorityDataPoints;
 		::Audio::TCircularAudioBuffer<FDataPoint> DistanceDataPoints;
+		::Audio::TCircularAudioBuffer<FDataPoint> DistanceAttenuationDataPoints;
+		::Audio::TCircularAudioBuffer<FDataPoint> LPFFreqDataPoints;
+		::Audio::TCircularAudioBuffer<FDataPoint> HPFFreqDataPoints;
 		::Audio::TCircularAudioBuffer<FDataPoint> AmplitudeDataPoints;
 		::Audio::TCircularAudioBuffer<FDataPoint> VolumeDataPoints;
 		::Audio::TCircularAudioBuffer<FDataPoint> PitchDataPoints;
diff --git a/Engine/Source/Runtime/AudioMixer/Private/AudioMixerSourceManager.cpp b/Engine/Source/Runtime/AudioMixer/Private/AudioMixerSourceManager.cpp
--- a/Engine/Source/Runtime/AudioMixer/Private/AudioMixerSourceManager.cpp
+++ b/Engine/Source/Runtime/AudioMixer/Private/AudioMixerSourceManager.cpp
@@ -233,6 +233,7 @@ UE_TRACE_EVENT_BEGIN(Audio, MixerSourceDistanceAttenuation)
 	UE_TRACE_EVENT_FIELD(uint32, DeviceId)
 	UE_TRACE_EVENT_FIELD(uint64, Timestamp)
 	UE_TRACE_EVENT_FIELD(uint32, PlayOrder)
+	UE_TRACE_EVENT_FIELD(uint32, ActiveSoundPlayOrder)
 	UE_TRACE_EVENT_FIELD(float, DistanceAttenuation)
 UE_TRACE_EVENT_END()
 
@@ -248,6 +249,7 @@ UE_TRACE_EVENT_BEGIN(Audio, MixerSourceFilters)
 	UE_TRACE_EVENT_FIELD(uint32, DeviceId)
 	UE_TRACE_EVENT_FIELD(uint64, Timestamp)
 	UE_TRACE_EVENT_FIELD(uint32, PlayOrder)
+	UE_TRACE_EVENT_FIELD(uint32, ActiveSoundPlayOrder)
 	UE_TRACE_EVENT_FIELD(float, LPFFrequency)
 	UE_TRACE_EVENT_FIELD(float, HPFFrequency)
 UE_TRACE_EVENT_END()
@@ -2983,6 +2985,7 @@ namespace Audio
 					<< MixerSourceDistanceAttenuation.DeviceId(MixerDevice->DeviceID)
 					<< MixerSourceDistanceAttenuation.Timestamp(FPlatformTime::Cycles64())
 					<< MixerSourceDistanceAttenuation.PlayOrder(SourceInfo.PlayOrder)
+					<< MixerSourceDistanceAttenuation.ActiveSoundPlayOrder(SourceInfo.ActiveSoundPlayOrder)
 					<< MixerSourceDistanceAttenuation.DistanceAttenuation(SourceInfo.DistanceAttenuationSourceDestination);
 			}
 #endif // UE_AUDIO_PROFILERTRACE_ENABLED
@@ -3122,6 +3125,7 @@ namespace Audio
 					<< MixerSourceFilters.DeviceId(MixerDevice->DeviceID)
 					<< MixerSourceFilters.Timestamp(FPlatformTime::Cycles64())
 					<< MixerSourceFilters.PlayOrder(SourceInfo.PlayOrder)
+					<< MixerSourceFilters.ActiveSoundPlayOrder(SourceInfo.ActiveSoundPlayOrder)
 					<< MixerSourceFilters.HPFFrequency(HPFFrequency)
 					<< MixerSourceFilters.LPFFrequency(LPFFrequency);
 				UE_TRACE_LOG(Audio, MixerSourceEnvelope, AudioMixerChannel)
