[Audio Insights] Move the Plots widget over to Sounds Dashboard.

diff --git a/Engine/Plugins/AudioInsights/Source/AudioInsights/Private/AudioInsightsDashboardFactory.cpp b/Engine/Plugins/AudioInsights/Source/AudioInsights/Private/AudioInsightsDashboardFactory.cpp
--- a/Engine/Plugins/AudioInsights/Source/AudioInsights/Private/AudioInsightsDashboardFactory.cpp
+++ b/Engine/Plugins/AudioInsights/Source/AudioInsights/Private/AudioInsightsDashboardFactory.cpp
@@ -171,7 +171,9 @@ namespace UE::Audio::Insights
 					);
 
 					if (const EDefaultDashboardTabStack DefaultTabStack = DashboardViewFactories[FactoryName]->GetDefaultTabStack();
-						DefaultTabStack == EDefaultDashboardTabStack::Log || DefaultTabStack == EDefaultDashboardTabStack::AudioMeters)
+						DefaultTabStack == EDefaultDashboardTabStack::Log 
+						|| DefaultTabStack == EDefaultDashboardTabStack::AudioMeters
+						|| DefaultTabStack == EDefaultDashboardTabStack::Plots)
 					{
 						MenuBuilder.AddMenuSeparator();
 					}
@@ -189,6 +191,7 @@ namespace UE::Audio::Insights
 
 		TSharedRef<FTabManager::FStack> LogTabStack = FTabManager::NewStack();
 		TSharedRef<FTabManager::FStack> AnalysisTabStack = FTabManager::NewStack();
+		TSharedRef<FTabManager::FStack> PlotsTabStack = FTabManager::NewStack();
 
 		for (const auto& [FactoryName, Factory] : DashboardViewFactories)
 		{
@@ -207,6 +210,12 @@ namespace UE::Audio::Insights
 					AnalysisTabStack->AddTab(FactoryName, ETabState::OpenedTab);
 				}
 				break;
+
+				case EDefaultDashboardTabStack::Plots:
+				{
+					PlotsTabStack->AddTab(FactoryName, ETabState::OpenedTab);
+				}
+				break;
 				
 				default:
 					break;
@@ -231,8 +240,27 @@ namespace UE::Audio::Insights
 				)
 				->Split
 				(
-					AnalysisTabStack
-					->SetSizeCoefficient(0.75f)
+					// Top
+					FTabManager::NewSplitter()
+					->SetOrientation(Orient_Vertical)
+					->SetSizeCoefficient(0.6f) // Column width
+					->Split
+					(
+						FTabManager::NewSplitter()
+						->SetOrientation(Orient_Vertical)
+						->SetSizeCoefficient(0.6f) // Row width
+						->Split
+						(
+							AnalysisTabStack
+							->SetSizeCoefficient(0.58f)
+						)
+						// Bottom
+						->Split
+						(
+							PlotsTabStack
+							->SetSizeCoefficient(0.42f)
+						)
+					)
 				)
 			)
 		);
diff --git a/Engine/Plugins/AudioInsights/Source/AudioInsights/Private/AudioInsightsModule.cpp b/Engine/Plugins/AudioInsights/Source/AudioInsights/Private/AudioInsightsModule.cpp
--- a/Engine/Plugins/AudioInsights/Source/AudioInsights/Private/AudioInsightsModule.cpp
+++ b/Engine/Plugins/AudioInsights/Source/AudioInsights/Private/AudioInsightsModule.cpp
@@ -16,6 +16,7 @@
 #include "AudioInsightsComponent.h"
 #include "Views/MixerSourceDashboardViewFactory.h"
 #include "Views/SoundDashboardViewFactory.h"
+#include "Views/SoundPlotsDashboardViewFactory.h"
 #include "Views/VirtualLoopDashboardViewFactory.h"
 #endif // !WITH_EDITOR
 
@@ -43,9 +44,15 @@ namespace UE::Audio::Insights
 #if !WITH_EDITOR
 			IModularFeatures::Get().RegisterModularFeature(UE::Insights::Timing::TimingViewExtenderFeatureName, &AudioInsightsTimingViewExtender);
 
-			DashboardFactory->RegisterViewFactory(MakeShared<FSoundDashboardViewFactory>());
+			TSharedRef<FSoundDashboardViewFactory> SoundsDashboard = MakeShared<FSoundDashboardViewFactory>();
+			TSharedRef<FSoundPlotsDashboardViewFactory> PlotsDashboard = MakeShared<FSoundPlotsDashboardViewFactory>();
+
+			DashboardFactory->RegisterViewFactory(SoundsDashboard);
 			DashboardFactory->RegisterViewFactory(MakeShared<FMixerSourceDashboardViewFactory>());
 			DashboardFactory->RegisterViewFactory(MakeShared<FVirtualLoopDashboardViewFactory>());
+			DashboardFactory->RegisterViewFactory(PlotsDashboard);
+
+			PlotsDashboard->InitPlots(SoundsDashboard);
 			
 			AudioInsightsComponent = FAudioInsightsComponent::CreateInstance();
 
diff --git a/Engine/Plugins/AudioInsights/Source/AudioInsights/Private/Providers/SoundTraceProvider.cpp b/Engine/Plugins/AudioInsights/Source/AudioInsights/Private/Providers/SoundTraceProvider.cpp
--- a/Engine/Plugins/AudioInsights/Source/AudioInsights/Private/Providers/SoundTraceProvider.cpp
+++ b/Engine/Plugins/AudioInsights/Source/AudioInsights/Private/Providers/SoundTraceProvider.cpp
@@ -602,6 +602,7 @@ namespace UE::Audio::Insights
 
 			FSoundTraceProviderPrivate::UpdateParameterEntry(Msg, Msg.Pitch, SoundEntry->PitchDataPoints, EntriesWithPoppedDataPoints);
 		});
+		EntriesWithPoppedDataPoints.Reset();
 
 		ProcessMessageQueue<FMixerSourceRelativeRenderCostMessage>(TraceMessages.RelativeRenderCostMessages, GetActiveSoundEntryFromParameterMessage,
 		[this, &EntriesWithPoppedDataPoints](const FMixerSourceRelativeRenderCostMessage& Msg, const TSharedPtr<FSoundDashboardEntry>* OutActiveSoundEntry)
@@ -621,6 +622,7 @@ namespace UE::Audio::Insights
 
 			FSoundTraceProviderPrivate::UpdateParameterEntry(Msg, Msg.RelativeRenderCost, SoundEntry->RelativeRenderCostDataPoints, EntriesWithPoppedDataPoints);
 		});
+		EntriesWithPoppedDataPoints.Reset();
 
 		ProcessMessageQueue<FSoundStopMessage>(TraceMessages.StopMessages, GetSoundEntryFromStopMessage,
 		[this](const FSoundStopMessage& Msg, const TSharedPtr<FSoundDashboardEntry>* OutActiveSoundEntry)
diff --git a/Engine/Plugins/AudioInsights/Source/AudioInsights/Private/Views/SAudioCurveView.cpp b/Engine/Plugins/AudioInsights/Source/AudioInsights/Private/Views/SAudioCurveView.cpp
--- a/Engine/Plugins/AudioInsights/Source/AudioInsights/Private/Views/SAudioCurveView.cpp
+++ b/Engine/Plugins/AudioInsights/Source/AudioInsights/Private/Views/SAudioCurveView.cpp
@@ -166,7 +166,7 @@ void SAudioCurveView::UpdateYDataRange()
 	};
 
 	float MinValue = TNumericLimits<float>::Max();
-	float MaxValue = TNumericLimits<float>::Min();
+	float MaxValue = TNumericLimits<float>::Lowest();
 	for (auto Iter = PointDataPerCurve->CreateConstIterator(); Iter; ++Iter)
 	{
 		const TArray<FCurvePoint>& CurvePoints = Iter->Value;
@@ -179,8 +179,16 @@ void SAudioCurveView::UpdateYDataRange()
 		MinValue = MinPoint ? FMath::Min(MinValue, MinPoint->Value) : 0.0f;
 		MaxValue = MaxPoint ? FMath::Max(MaxValue, MaxPoint->Value) : 1.0f;
 	}
-	YDataRange.X = MinValue;
-	YDataRange.Y = MaxValue;
+
+	if (MinValue < TNumericLimits<float>::Max())
+	{
+		YDataRange.X = MinValue;
+	}
+	
+	if (MaxValue > TNumericLimits<float>::Lowest())
+	{
+		YDataRange.Y = MaxValue;
+	}
 }
 
 #if !WITH_EDITOR
@@ -192,7 +200,7 @@ void SAudioCurveView::UpdateYDataRangeFromTimestampRange(const double InLowerBou
 		};
 
 	float MinValue = TNumericLimits<float>::Max();
-	float MaxValue = TNumericLimits<float>::Min();
+	float MaxValue = TNumericLimits<float>::Lowest();
 
 	if (PointDataPerCurve.IsValid())
 	{
@@ -231,7 +239,7 @@ void SAudioCurveView::UpdateYDataRangeFromTimestampRange(const double InLowerBou
 		MinValue = 0.0f;
 	}
 
-	if (MaxValue == TNumericLimits<float>::Min() || FMath::Abs(MaxValue - MinValue) < Epsilon)
+	if (MaxValue == TNumericLimits<float>::Lowest() || FMath::Abs(MaxValue - MinValue) < Epsilon)
 	{
 		MaxValue = MinValue + Epsilon;
 	}
@@ -257,6 +265,11 @@ void SAudioCurveView::SetYValueFormattingOptions(const FNumberFormattingOptions
 	YValueFormattingOptions = InYValueFormattingOptions;
 }
 
+void SAudioCurveView::SetInitYAxisRange(const FVector2f& InRange)
+{
+	YDataRange = InRange;
+}
+
 float SAudioCurveView::ValueToLocalY(const FVector2f AllottedLocalSize, const float Value) const
 {
 	// Slate Y values increase going down the screen, so base < top but base is above top on the screen
diff --git a/Engine/Plugins/AudioInsights/Source/AudioInsights/Private/Views/SoundDashboardViewFactory.cpp b/Engine/Plugins/AudioInsights/Source/AudioInsights/Private/Views/SoundDashboardViewFactory.cpp
--- a/Engine/Plugins/AudioInsights/Source/AudioInsights/Private/Views/SoundDashboardViewFactory.cpp
+++ b/Engine/Plugins/AudioInsights/Source/AudioInsights/Private/Views/SoundDashboardViewFactory.cpp
@@ -135,6 +135,47 @@ namespace UE::Audio::Insights
 			}
 		};
 
+		TSharedRef<SBox> CreateButtonContentWidget(const FName& InIconName = FName(), const FText& InLabel = FText::GetEmpty(), const FName& InTextStyle = TEXT("ButtonText"))
+		{
+			TSharedRef<SHorizontalBox> ButtonContainerWidget = SNew(SHorizontalBox);
+
+			// Button icon (optional)
+			if (!InIconName.IsNone())
+			{
+				ButtonContainerWidget->AddSlot()
+					.AutoWidth()
+					.VAlign(VAlign_Center)
+					[
+						SNew(SImage)
+							.ColorAndOpacity(FSlateColor::UseForeground())
+							.Image(FSlateStyle::Get().GetBrush(InIconName))
+					];
+			}
+
+			// Button text (optional)
+			if (!InLabel.IsEmpty())
+			{
+				const float LeftPadding = InIconName.IsNone() ? 0.0f : 4.0f;
+
+				ButtonContainerWidget->AddSlot()
+					.VAlign(VAlign_Center)
+					.Padding(LeftPadding, 0.0f, 0.0f, 0.0f)
+					.AutoWidth()
+					[
+						SNew(STextBlock)
+							.TextStyle(&FSlateStyle::Get().GetWidgetStyle<FTextBlockStyle>(InTextStyle))
+							.Justification(ETextJustify::Center)
+							.Text(InLabel)
+					];
+			}
+
+			return SNew(SBox)
+				.HeightOverride(16.0f)
+				[
+					ButtonContainerWidget
+				];
+		};
+
 		void CacheInitExpandStateRecursive(FSoundDashboardEntry& SoundEntry)
 		{
 			for (TSharedPtr<IDashboardDataTreeViewEntry>& ChildEntry : SoundEntry.Children)
@@ -418,6 +459,54 @@ namespace UE::Audio::Insights
 		const FText SoundWaveCategoryName = LOCTEXT("SoundDashboard_Filter_SoundWaveNameText", "Sound Wave");
 		const FText SoundCueTemplateCategoryName = LOCTEXT("SoundDashboard_Filter_SoundCueTemplateNameText", "Sound Cue Template");
 		const FText PinnedCategoryName = LOCTEXT("SoundDashboard_Filter_PinnedNameText", "Pinned");
+
+		const FName MuteColumnName = "Mute";
+		const FText MuteColumnDisplayName = LOCTEXT("SoundDashboard_MuteColumnDisplayName", "Mute");
+
+		const FName SoloColumnName = "Solo";
+		const FText SoloColumnDisplayName = LOCTEXT("SoundDashboard_SoloColumnDisplayName", "Solo");
+
+		const FName PlotColumnName = "Plot";
+		const FText PlotColumnDisplayName = LOCTEXT("SoundDashboard_PlotColumnDisplayName", "Plot");
+
+		const FName NameColumnName = "Name";
+		const FText NameColumnDisplayName = LOCTEXT("SoundDashboard_NameColumnDisplayName", "Name");
+
+		const FName PlayOrderColumnName = "PlayOrder";
+		const FText PlayOrderColumnDisplayName = LOCTEXT("SoundDashboard_PlayOrderColumnDisplayName", "Play Order");
+
+		const FName PriorityColumnName = "Priority";
+		const FText PriorityColumnDisplayName = LOCTEXT("SoundDashboard_PriorityColumnDisplayName", "Priority");
+
+		const FName DistanceColumnName = "Distance";
+		const FText DistanceColumnDisplayName = LOCTEXT("SoundDashboard_DistanceColumnDisplayName", "Distance");
+
+		const FName DistanceAttenuationColumnName = "DistanceAttenuation";
+		const FText DistanceAttenuationColumnDisplayName = LOCTEXT("SoundDashboard_DistanceAttenuationColumnDisplayName", "Distance Attenuation");
+
+		const FName AmplitudeColumnName = "Amplitude";
+		const FText AmplitudeColumnDisplayName = LOCTEXT("SoundDashboard_AmplitudeColumnDisplayName", "Amp (Peak)");
+
+		const FName VolumeColumnName = "Volume";
+		const FText VolumeColumnDisplayName = LOCTEXT("SoundDashboard_VolumeColumnDisplayName", "Volume");
+
+		const FName LPFFreqColumnName = "LPFFreq";
+		const FText LPFFreqColumnDisplayName = LOCTEXT("SoundDashboard_LPFFreqColumnDisplayName", "LPF Freq (Hz)");
+
+		const FName HPFFreqColumnName = "HPFFreq";
+		const FText HPFFreqColumnDisplayName = LOCTEXT("SoundDashboard_HPFFreqColumnDisplayName", "HPF Freq (Hz)");
+
+		const FName PitchColumnName = "Pitch";
+		const FText PitchColumnDisplayName = LOCTEXT("SoundDashboard_PitchColumnDisplayName", "Pitch");
+
+		const FName RelativeRenderCostColumnName = "RelativeRenderCost";
+		const FText RelativeRenderCostColumnDisplayName = LOCTEXT("SoundDashboard_RelativeRenderCostColumnDisplayName", "Rel. Render Cost");
+
+		const FName ActorColumnName = "Actor";
+		const FText ActorColumnDisplayName = LOCTEXT("SoundDashboard_ActorColumnDisplayName", "Actor");
+
+		const FName CategoryColumnName = "Category";
+		const FText CategoryColumnDisplayName = LOCTEXT("SoundDashboard_CategoryColumnDisplayName", "Category");
 	} // namespace FSoundDashboardViewFactoryPrivate
 
 	/////////////////////////////////////////////////////////////////////////////////////////
@@ -593,6 +682,29 @@ namespace UE::Audio::Insights
 		return PinnedSectionEntry.IsValid() && (OriginalDataEntry.IsValid() || FSoundDashboardViewFactoryPrivate::IsCategoryItem(*PinnedSectionEntry));
 	}
 
+	TSharedPtr<IDashboardDataTreeViewEntry> FPinnedSoundEntryWrapper::FindOriginalEntryInChildren(const TSharedPtr<IDashboardDataTreeViewEntry>& PinnedEntry)
+	{
+		if (!PinnedEntry.IsValid())
+		{
+			return nullptr;
+		}
+
+		if (GetPinnedSectionEntry() == PinnedEntry)
+		{
+			return GetOriginalDataEntry();
+		}
+
+		for (const TSharedPtr<FPinnedSoundEntryWrapper>& ChildPinnedEntryWrapper : PinnedWrapperChildren)
+		{
+			TSharedPtr<IDashboardDataTreeViewEntry> FoundEntry = ChildPinnedEntryWrapper->FindOriginalEntryInChildren(PinnedEntry);
+			if (FoundEntry.IsValid())
+			{
+				return FoundEntry;
+			}
+		}
+		return nullptr;
+	}
+
 	/////////////////////////////////////////////////////////////////////////////////////////
 	// FSoundDashboardViewFactory
 	FSoundDashboardViewFactory::FSoundDashboardViewFactory()
@@ -703,47 +815,6 @@ namespace UE::Audio::Insights
 	{
 		using namespace FSoundDashboardViewFactoryPrivate;
 
-		const auto CreateButtonContentWidget = [](const FName& InIconName = FName(), const FText& InLabel = FText::GetEmpty(), const FName& InTextStyle = TEXT("ButtonText"))
-        {
-        	TSharedRef<SHorizontalBox> ButtonContainerWidget = SNew(SHorizontalBox);
-
-			// Button icon (optional)
-        	if (!InIconName.IsNone())
-        	{
-				ButtonContainerWidget->AddSlot()
-	        	.AutoWidth()
-	            .VAlign(VAlign_Center)
-	            [
-	                SNew(SImage)
-	                .ColorAndOpacity(FSlateColor::UseForeground())
-	                .Image(FSlateStyle::Get().GetBrush(InIconName))
-	            ];
-        	}
-
-			// Button text (optional)
-        	if (!InLabel.IsEmpty())
-        	{
-				const float LeftPadding = InIconName.IsNone() ? 0.0f : 4.0f;
-
-				ButtonContainerWidget->AddSlot()
-        		.VAlign(VAlign_Center)
-	            .Padding(LeftPadding, 0.0f, 0.0f, 0.0f)
-	            .AutoWidth()
-	            [
-	                SNew(STextBlock)
-	                .TextStyle(&FSlateStyle::Get().GetWidgetStyle<FTextBlockStyle>(InTextStyle))
-	                .Justification(ETextJustify::Center)
-	                .Text(InLabel)
-	            ];
-        	}
-
-        	return SNew(SBox)
-				.HeightOverride(16.0f)
-				[
-					ButtonContainerWidget
-				];
-        };
-
 		return SNew(SHorizontalBox)
 			// Mute button
 			+ SHorizontalBox::Slot()
@@ -804,16 +875,6 @@ namespace UE::Audio::Insights
 				[
 					CreateButtonContentWidget("AudioInsights.Icon.SoundDashboard.Reset", LOCTEXT("SoundsDashboard_ClearMutesAndSolosButtonText", "Clear All Mutes/Solos"), "SmallButtonText")
 				]
-			]
-			// Empty Spacing
-			+ SHorizontalBox::Slot()
-			.FillWidth(1.0f)
-			// Settings button
-			+ SHorizontalBox::Slot()
-			.HAlign(HAlign_Right)
-			.AutoWidth()
-			[
-				MakeSettingsButtonWidget()
 			];
 	}
 #endif // WITH_EDITOR
@@ -942,15 +1003,53 @@ namespace UE::Audio::Insights
 			.HAlign(HAlign_Fill)
 			.Padding(0.0f, 2.0f, 0.0f, 0.0f)
 			[
-				MakeMuteSoloWidget()
+				SNew(SHorizontalBox)
+				+ SHorizontalBox::Slot()
+				.AutoWidth()
+				[
+					MakeMuteSoloWidget()
+				]
+
+				+SHorizontalBox::Slot()
+				.AutoWidth()
+				[
+					MakeShowPlotWidget()
+				]
+				
+				// Empty Spacing
+				+ SHorizontalBox::Slot()
+				.FillWidth(1.0f)
+				// Settings button
+				+ SHorizontalBox::Slot()
+				.HAlign(HAlign_Right)
+				.AutoWidth()
+				[
+					MakeSettingsButtonWidget()
+				]
 			]
 #else
 			+ SVerticalBox::Slot()
 			.AutoHeight()
-			.HAlign(HAlign_Right)
+			.HAlign(HAlign_Fill)
 			.Padding(0.0f, 2.0f, 0.0f, 0.0f)
 			[
-				MakeSettingsButtonWidget()
+				SNew(SHorizontalBox)
+				+ SHorizontalBox::Slot()
+				.HAlign(HAlign_Left)
+				.AutoWidth()
+				[
+					MakeShowPlotWidget()
+				]
+				// Empty Spacing
+				+ SHorizontalBox::Slot()
+				.FillWidth(1.0f)
+				// Settings button
+				+ SHorizontalBox::Slot()
+				.HAlign(HAlign_Right)
+				.AutoWidth()
+				[
+					MakeSettingsButtonWidget()
+				]
 			]
 #endif // WITH_EDITOR
 			+ SVerticalBox::Slot()
@@ -966,7 +1065,7 @@ namespace UE::Audio::Insights
 
 		const FColumnData& ColumnData = GetColumns()[InColumn];
 
-		if (InColumn == "Name")
+		if (InColumn == NameColumnName)
 		{
 			const FName IconName = ColumnData.GetIconName.IsSet() ? ColumnData.GetIconName(InRowData.Get()) : NAME_None;
 
@@ -1053,7 +1152,7 @@ namespace UE::Audio::Insights
 		}
 
 #if WITH_EDITOR
-		if (InColumn == "Mute")
+		if (InColumn == MuteColumnName)
 		{
 			return CreateMuteSoloButton(InRowWidget, InRowData, InColumn,
 				[this](const TArray<TSharedPtr<IDashboardDataTreeViewEntry>>& InEntries)
@@ -1070,7 +1169,7 @@ namespace UE::Audio::Insights
 #endif // ENABLE_AUDIO_DEBUG
 				});
 		}
-		else if (InColumn == "Solo")
+		else if (InColumn == SoloColumnName)
 		{
 			return CreateMuteSoloButton(InRowWidget, InRowData, InColumn,
 				[this](const TArray<TSharedPtr<IDashboardDataTreeViewEntry>>& InEntries)
@@ -1089,7 +1188,19 @@ namespace UE::Audio::Insights
 		}
 		else
 #endif // WITH_EDITOR
-		if (InColumn == "Name")
+		if (InColumn == PlotColumnName)
+		{
+			return CreateShowPlotColumnButton(InRowWidget, InRowData, InColumn,
+				[this](const TArray<TSharedPtr<IDashboardDataTreeViewEntry>>& InEntries, const bool bActivatePlot)
+				{
+					ToggleShowPlotEntries(InEntries, bActivatePlot);
+				},
+				[](const IDashboardDataTreeViewEntry& InEntry)
+				{
+					return CastEntry(InEntry).bIsPlotActive;
+				});
+		}
+		else if (InColumn == NameColumnName)
 		{
 			if (IsRootItem(InRowData))
 			{
@@ -1265,6 +1376,11 @@ namespace UE::Audio::Insights
 
 			return bEntryTypePassesFilter;
 		});
+
+		if (InReason == FTraceTreeDashboardViewFactory::EProcessReason::EntriesUpdated)
+		{
+			OnProcessPlotData.Broadcast(DataViewEntries);
+		}
 	}
 
 	void FSoundDashboardViewFactory::ChangeTreeViewingMode(const ESoundDashboardTreeViewingOptions SelectedMode, const bool bUpdateEditorSettings /*= true*/)
@@ -1430,13 +1546,15 @@ namespace UE::Audio::Insights
 
 	const TMap<FName, FTraceTreeDashboardViewFactory::FHeaderRowColumnData>& FSoundDashboardViewFactory::GetHeaderRowColumns() const
 	{
+		using namespace FSoundDashboardViewFactoryPrivate;
+
 		static const TMap<FName, FTraceTreeDashboardViewFactory::FHeaderRowColumnData> HeaderRowColumnData =
 		{
 #if WITH_EDITOR
 			{
-				"Mute",
+				MuteColumnName,
 				{
-					.DisplayName      = LOCTEXT("SoundDashboard_MuteColumnDisplayName", "Mute"),
+					.DisplayName      = MuteColumnDisplayName,
 					.IconName         = "AudioInsights.Icon.SoundDashboard.Mute",
 					.bShowDisplayName = false,
 					.bDefaultHidden   = false,
@@ -1445,9 +1563,9 @@ namespace UE::Audio::Insights
 				}
 			},
 			{
-				"Solo",
+				SoloColumnName,
 				{
-					.DisplayName      = LOCTEXT("SoundDashboard_SoloColumnDisplayName", "Solo"),
+					.DisplayName      = SoloColumnDisplayName,
 					.IconName         = "AudioInsights.Icon.SoundDashboard.Solo",
 					.bShowDisplayName = false,
 					.bDefaultHidden   = false,
@@ -1457,9 +1575,20 @@ namespace UE::Audio::Insights
 			},
 #endif
 			{
-				"Name",
+				PlotColumnName,
+				{
+					.DisplayName = PlotColumnDisplayName,
+					.IconName = "AudioInsights.Icon.SoundDashboard.Plots",
+					.bShowDisplayName = false,
+					.bDefaultHidden = false,
+					.FillWidth = 0.05f,
+					.Alignment = EHorizontalAlignment::HAlign_Center
+				}
+			},
+			{
+				NameColumnName,
 				{
-					.DisplayName      = LOCTEXT("SoundDashboard_NameColumnDisplayName", "Name"),
+					.DisplayName      = NameColumnDisplayName,
 					.IconName         = NAME_None,
 					.bShowDisplayName = true,
 					.bDefaultHidden   = false,
@@ -1468,20 +1597,20 @@ namespace UE::Audio::Insights
 				}
 			},
 			{
-				"PlayOrder",
+				PlayOrderColumnName,
 				{
-					.DisplayName = LOCTEXT("SoundDashboard_PlayOrderColumnDisplayName", "Play Order"),
-					.IconName = NAME_None,
+					.DisplayName	  = PlayOrderColumnDisplayName,
+					.IconName		  = NAME_None,
 					.bShowDisplayName = true,
-					.bDefaultHidden = true,
-					.FillWidth = 0.1f,
-					.Alignment = EHorizontalAlignment::HAlign_Left
+					.bDefaultHidden   = true,
+					.FillWidth		  = 0.1f,
+					.Alignment		  = EHorizontalAlignment::HAlign_Left
 				}
 			},
 			{
-				"Priority",
+				PriorityColumnName,
 				{
-					.DisplayName	  = LOCTEXT("SoundDashboard_PriorityColumnDisplayName", "Priority"),
+					.DisplayName	  = PriorityColumnDisplayName,
 					.IconName		  = NAME_None,
 					.bShowDisplayName = true,
 					.bDefaultHidden	  = false,
@@ -1490,9 +1619,9 @@ namespace UE::Audio::Insights
 				}
 			},
 			{
-				"Distance",
+				DistanceColumnName,
 				{
-					.DisplayName      = LOCTEXT("SoundDashboard_DistanceColumnDisplayName", "Distance"),
+					.DisplayName      = DistanceColumnDisplayName,
 					.IconName         = NAME_None,
 					.bShowDisplayName = true,
 					.bDefaultHidden   = false,
@@ -1501,20 +1630,20 @@ namespace UE::Audio::Insights
 				}
 			},
 			{
-				"DistanceAttenuation",
+				DistanceAttenuationColumnName,
 				{
-					.DisplayName = LOCTEXT("SoundDashboard_DistanceAttenuationColumnDisplayName", "Distance Attenuation"),
-					.IconName = NAME_None,
+					.DisplayName	  = DistanceAttenuationColumnDisplayName,
+					.IconName		  = NAME_None,
 					.bShowDisplayName = true,
-					.bDefaultHidden = true,
-					.FillWidth = 0.1f,
-					.Alignment = EHorizontalAlignment::HAlign_Left
+					.bDefaultHidden   = true,
+					.FillWidth		  = 0.1f,
+					.Alignment		  = EHorizontalAlignment::HAlign_Left
 				}
 			},
 			{
-				"Amplitude",
+				AmplitudeColumnName,
 				{
-					.DisplayName      = LOCTEXT("SoundDashboard_AmplitudeColumnDisplayName", "Amp (Peak)"),
+					.DisplayName      = AmplitudeColumnDisplayName,
 					.IconName         = NAME_None,
 					.bShowDisplayName = true,
 					.bDefaultHidden   = false,
@@ -1523,9 +1652,9 @@ namespace UE::Audio::Insights
 				}
 			},
 			{
-				"Volume",
+				VolumeColumnName,
 				{
-					.DisplayName      = LOCTEXT("SoundDashboard_VolumeColumnDisplayName", "Volume"),
+					.DisplayName      = VolumeColumnDisplayName,
 					.IconName         = NAME_None,
 					.bShowDisplayName = true,
 					.bDefaultHidden   = false,
@@ -1534,31 +1663,31 @@ namespace UE::Audio::Insights
 				}
 			},
 			{
-				"LPFFreq",
+				LPFFreqColumnName,
 				{
-					.DisplayName = LOCTEXT("SoundDashboard_LPFFreqColumnDisplayName", "LPF Freq (Hz)"),
-					.IconName = NAME_None,
+					.DisplayName	  = LPFFreqColumnDisplayName,
+					.IconName		  = NAME_None,
 					.bShowDisplayName = true,
-					.bDefaultHidden = true,
-					.FillWidth = 0.1f,
-					.Alignment = EHorizontalAlignment::HAlign_Left
+					.bDefaultHidden   = true,
+					.FillWidth		  = 0.1f,
+					.Alignment		  = EHorizontalAlignment::HAlign_Left
 				}
 			},
 			{
-				"HPFFreq",
+				HPFFreqColumnName,
 				{
-					.DisplayName = LOCTEXT("SoundDashboard_HPFFreqColumnDisplayName", "HPF Freq (Hz)"),
-					.IconName = NAME_None,
+					.DisplayName	  = HPFFreqColumnDisplayName,
+					.IconName		  = NAME_None,
 					.bShowDisplayName = true,
-					.bDefaultHidden = true,
-					.FillWidth = 0.1f,
-					.Alignment = EHorizontalAlignment::HAlign_Left
+					.bDefaultHidden   = true,
+					.FillWidth		  = 0.1f,
+					.Alignment		  = EHorizontalAlignment::HAlign_Left
 				}
 			},
 			{
-				"Pitch",
+				PitchColumnName,
 				{
-					.DisplayName      = LOCTEXT("SoundDashboard_PitchColumnDisplayName", "Pitch"),
+					.DisplayName      = PitchColumnDisplayName,
 					.IconName         = NAME_None,
 					.bShowDisplayName = true,
 					.bDefaultHidden   = false,
@@ -1567,9 +1696,9 @@ namespace UE::Audio::Insights
 				}
 			},
 			{
-				"RelativeRenderCost",
+				RelativeRenderCostColumnName,
 				{
-					.DisplayName = LOCTEXT("SoundDashboard_RelativeRenderCostColumnDisplayName", "Rel. Render Cost"),
+					.DisplayName = RelativeRenderCostColumnDisplayName,
 					.IconName = NAME_None,
 					.bShowDisplayName = true,
 					.bDefaultHidden = true,
@@ -1579,9 +1708,9 @@ namespace UE::Audio::Insights
 			},
 #if WITH_EDITOR
 			{
-				"Actor",
+				ActorColumnName,
 				{
-					.DisplayName = LOCTEXT("SoundDashboard_ActorColumnDisplayName", "Actor"),
+					.DisplayName = ActorColumnDisplayName,
 					.IconName = NAME_None,
 					.bShowDisplayName = true,
 					.bDefaultHidden = true,
@@ -1591,7 +1720,7 @@ namespace UE::Audio::Insights
 			},
 #endif // WITH_EDITOR
 			{
-				"Category",
+				CategoryColumnName,
 				{
 					.DisplayName = LOCTEXT("SoundDashboard_CategoryColumnDisplayName", "Category"),
 					.IconName = NAME_None,
@@ -1614,7 +1743,7 @@ namespace UE::Audio::Insights
 		{
 #if WITH_EDITOR
 			{
-				"Mute",
+				MuteColumnName,
 				{
 					.GetIconName = [this](const IDashboardDataTreeViewEntry& InData)
 					{
@@ -1623,7 +1752,7 @@ namespace UE::Audio::Insights
 				}
 			},
 			{
-				"Solo",
+				SoloColumnName,
 				{
 					.GetIconName = [this](const IDashboardDataTreeViewEntry& InData)
 					{
@@ -1633,7 +1762,16 @@ namespace UE::Audio::Insights
 			},
 #endif
 			{
-				"Name",
+				PlotColumnName,
+				{
+					.GetIconName = [this](const IDashboardDataTreeViewEntry& InData)
+					{
+						return "AudioInsights.Icon.SoundDashboard.Plots";
+					}
+				}
+			},
+			{
+				NameColumnName,
 				{
 					.GetDisplayValue = [](const IDashboardDataTreeViewEntry& InData)
 					{
@@ -1669,7 +1807,7 @@ namespace UE::Audio::Insights
 				}
 			},
 			{
-				"PlayOrder",
+				PlayOrderColumnName,
 				{
 					.GetDisplayValue = [](const IDashboardDataTreeViewEntry& InData)
 					{
@@ -1684,7 +1822,7 @@ namespace UE::Audio::Insights
 				}
 			},
 			{
-				"Priority",
+				PriorityColumnName,
 				{
 					.GetDisplayValue = [](const IDashboardDataTreeViewEntry& InData)
 					{
@@ -1706,7 +1844,7 @@ namespace UE::Audio::Insights
 				}
 			},
 			{
-				"Distance",
+				DistanceColumnName,
 				{
 					.GetDisplayValue = [](const IDashboardDataTreeViewEntry& InData)
 					{
@@ -1721,7 +1859,7 @@ namespace UE::Audio::Insights
 				}
 			},
 			{
-				"DistanceAttenuation",
+				DistanceAttenuationColumnName,
 				{
 					.GetDisplayValue = [](const IDashboardDataTreeViewEntry& InData)
 					{
@@ -1736,7 +1874,7 @@ namespace UE::Audio::Insights
 				}
 			},
 			{
-				"Amplitude",
+				AmplitudeColumnName,
 				{
 					.GetDisplayValue = [this](const IDashboardDataTreeViewEntry& InData)
 					{
@@ -1755,7 +1893,7 @@ namespace UE::Audio::Insights
 				}
 			},
 			{
-				"Volume",
+				VolumeColumnName,
 				{
 					.GetDisplayValue = [](const IDashboardDataTreeViewEntry& InData)
 					{
@@ -1770,7 +1908,7 @@ namespace UE::Audio::Insights
 				}
 			},
 			{
-				"LPFFreq",
+				LPFFreqColumnName,
 				{
 					.GetDisplayValue = [](const IDashboardDataTreeViewEntry& InData)
 					{
@@ -1785,7 +1923,7 @@ namespace UE::Audio::Insights
 				}
 			},
 			{
-				"HPFFreq",
+				HPFFreqColumnName,
 				{
 					.GetDisplayValue = [](const IDashboardDataTreeViewEntry& InData)
 					{
@@ -1800,7 +1938,7 @@ namespace UE::Audio::Insights
 				}
 			},
 			{
-				"Pitch",
+				PitchColumnName,
 				{
 					.GetDisplayValue = [](const IDashboardDataTreeViewEntry& InData)
 					{
@@ -1815,7 +1953,7 @@ namespace UE::Audio::Insights
 				}
 			},
 			{
-				"RelativeRenderCost",
+				RelativeRenderCostColumnName,
 				{
 					.GetDisplayValue = [](const IDashboardDataTreeViewEntry& InData)
 					{
@@ -1831,7 +1969,7 @@ namespace UE::Audio::Insights
 			},
 #if WITH_EDITOR
 			{
-				"Actor",
+				ActorColumnName,
 				{
 					.GetDisplayValue = [](const IDashboardDataTreeViewEntry& InData)
 					{
@@ -1843,7 +1981,7 @@ namespace UE::Audio::Insights
 			},
 #endif // WITH_EDITOR
 			{
-				"Category",
+				CategoryColumnName,
 				{
 					.GetDisplayValue = [](const IDashboardDataTreeViewEntry& InData)
 					{
@@ -2005,57 +2143,57 @@ namespace UE::Audio::Insights
 			return Comparison < 0;
 		};
 
-		if (SortByColumn == "Name")
+		if (SortByColumn == NameColumnName)
 		{
 			SortByPredicate(SortByName);
 		}
-		else if (SortByColumn == "PlayOrder")
+		else if (SortByColumn == PlayOrderColumnName)
 		{
 			SortByPredicate(SortByPlayOrder);
 		}
-		else if (SortByColumn == "Priority")
+		else if (SortByColumn == PriorityColumnName)
 		{
 			SortByPredicate(SortByPriority);
 		}
-		else if (SortByColumn == "Distance")
+		else if (SortByColumn == DistanceColumnName)
 		{
 			SortByPredicate(SortByDistance);
 		}
-		else if (SortByColumn == "DistanceAttenuation")
+		else if (SortByColumn == DistanceAttenuationColumnName)
 		{
 			SortByPredicate(SortByDistanceAttenuation);
 		}
-		else if (SortByColumn == "Amplitude")
+		else if (SortByColumn == AmplitudeColumnName)
 		{
 			SortByPredicate(SortByAmplitude);
 		}
-		else if (SortByColumn == "Volume")
+		else if (SortByColumn == VolumeColumnName)
 		{
 			SortByPredicate(SortByVolume);
 		}
-		else if (SortByColumn == "LPFFreq")
+		else if (SortByColumn == LPFFreqColumnName)
 		{
 			SortByPredicate(SortByLPFFreq);
 		}
-		else if (SortByColumn == "HPFFreq")
+		else if (SortByColumn == HPFFreqColumnName)
 		{
 			SortByPredicate(SortByHPFFreq);
 		}
-		else if (SortByColumn == "Pitch")
+		else if (SortByColumn == PitchColumnName)
 		{
 			SortByPredicate(SortByPitch);
 		}
-		else if (SortByColumn == "RelativeRenderCost")
+		else if (SortByColumn == RelativeRenderCostColumnName)
 		{
 			SortByPredicate(SortByRelativeRenderCost);
 		}
 #if WITH_EDITOR
-		else if (SortByColumn == "Actor")
+		else if (SortByColumn == ActorColumnName)
 		{
 			SortByPredicate(SortByActor);
 		}
 #endif // WITH_EDITOR
-		else if (SortByColumn == "Category")
+		else if (SortByColumn == CategoryColumnName)
 		{
 			SortByPredicate(SortByCategory);
 		}
@@ -2076,6 +2214,15 @@ namespace UE::Audio::Insights
 		}
 	}
 
+	bool FSoundDashboardViewFactory::IsColumnSortable(const FName& ColumnId) const
+	{
+		using namespace FSoundDashboardViewFactoryPrivate;
+
+		return ColumnId != SoloColumnName
+			&& ColumnId != MuteColumnName
+			&& ColumnId != PlotColumnName;
+	}
+
 	bool FSoundDashboardViewFactory::ResetTreeData()
 	{
 		bool bDataReset = false;
@@ -2568,6 +2715,282 @@ namespace UE::Audio::Insights
 	}
 #endif // WITH_EDITOR
 
+	TSharedRef<SWidget> FSoundDashboardViewFactory::MakeShowPlotWidget()
+	{
+		using namespace FSoundDashboardViewFactoryPrivate;
+
+		return SNew(SHorizontalBox)
+			// Toggle Selected Plots button
+			+ SHorizontalBox::Slot()
+			.AutoWidth()
+			.Padding(8.0f, 6.0f)
+			[
+				SNew(SButton)
+				.ButtonStyle(&FAppStyle::Get().GetWidgetStyle<FButtonStyle>("SimpleButton"))
+				.ToolTipText(LOCTEXT("SoundDashboard_SelectedPlotsToggleText", "Toggles whether plots are enabled on the selected items."))
+				.OnClicked_Lambda([this]()
+				{
+					if (FilteredEntriesListView.IsValid())
+					{
+						bool bToggleOn = false;
+						for (const TSharedPtr<IDashboardDataTreeViewEntry>& SelectedEntry : FilteredEntriesListView->GetSelectedItems())
+						{
+							if (!SelectedEntry.IsValid())
+							{
+								continue;
+							}
+
+							const FSoundDashboardEntry& SoundEntry = CastEntry(*SelectedEntry);
+
+							if (!SoundEntry.bIsPlotActive)
+							{
+								bToggleOn = true;
+								break;
+							}
+						}
+
+						const bool bIgnoreCategories = SelectedItemsIncludesAnAsset();
+						ToggleShowPlotEntries(FilteredEntriesListView->GetSelectedItems(), bToggleOn, bIgnoreCategories);
+					}
+
+					return FReply::Handled();
+				})
+				[
+					CreateButtonContentWidget("AudioInsights.Icon.SoundDashboard.Plots", LOCTEXT("SoundDashboard_PlotsButtonText", "Plot Selected"), "SmallButtonText")
+				]
+			]
+			// Clear All Plots button
+			+ SHorizontalBox::Slot()
+			.AutoWidth()
+			.Padding(8.0f, 6.0f)
+			[
+				SNew(SButton)
+				.ButtonStyle(&FAppStyle::Get().GetWidgetStyle<FButtonStyle>("SimpleButton"))
+				.ToolTipText(LOCTEXT("SoundsDashboard_ClearPlotsTooltipText", "Clears all enabled plots."))
+				.OnClicked_Lambda([this]()
+				{
+					for (const TSharedPtr<IDashboardDataTreeViewEntry>& Entry : DataViewEntries)
+					{
+						SetShowPlotRecursive(Entry, false);
+					}
+
+					return FReply::Handled();
+				})
+				[
+					CreateButtonContentWidget("AudioInsights.Icon.SoundDashboard.Reset", LOCTEXT("SoundsDashboard_ClearMutesAndSolosButtonText", "Clear All Plots"), "SmallButtonText")
+				]
+			];
+	}
+
+	TSharedRef<SWidget> FSoundDashboardViewFactory::CreateShowPlotColumnButton(const TSharedRef<FTraceTreeDashboardViewFactory::SRowWidget>& InRowWidget, const TSharedRef<IDashboardDataTreeViewEntry>& InRowData, const FName& InColumn, TFunction<void(const TArray<TSharedPtr<IDashboardDataTreeViewEntry>>&, const bool)> ShowPlotToggleFunc, TFunctionRef<bool(const IDashboardDataTreeViewEntry&)> IsPlotActiveFunc)
+	{
+		using namespace FSoundDashboardViewFactoryPrivate;
+
+		return SNew(SHorizontalBox)
+			.Visibility_Lambda([this, InRowWidget, InRowData]()
+			{
+				return InRowWidget->GetVisibility();
+			})
+			+ SHorizontalBox::Slot()
+			.AutoWidth()
+			.Padding(2.0f, 0.0f, 0.0f, 0.0f)
+			[
+				SNew(SCheckBox)
+				.HAlign(EHorizontalAlignment::HAlign_Center)
+				.Style(&FSlateStyle::Get().GetWidgetStyle<FCheckBoxStyle>("SoundDashboard.MuteSoloButton"))
+				.IsChecked(ECheckBoxState::Unchecked)
+				.IsChecked_Lambda([this, InRowData]()
+				{
+					return CastEntry(*InRowData).bIsPlotActive ? ECheckBoxState::Checked : ECheckBoxState::Unchecked;
+				})
+				.OnCheckStateChanged_Lambda([this, InRowData, ShowPlotToggleFunc](ECheckBoxState NewState)
+				{
+					ShowPlotToggleFunc({ InRowData }, NewState == ECheckBoxState::Checked);
+				})
+				[
+					SNew(SImage)
+					.Image_Lambda([this, InRowData, InColumn, IsPlotActiveFunc]()
+					{
+						const FColumnData& ColumnData = GetColumns()[InColumn];
+						const FName IconName = IsPlotActiveFunc(*InRowData) ? ColumnData.GetIconName(InRowData.Get()) : "AudioInsights.Icon.SoundDashboard.Transparent";
+
+						return FSlateStyle::Get().GetBrush(IconName);
+					})
+				]
+			];
+	}
+
+	void FSoundDashboardViewFactory::ToggleShowPlotEntries(const TArray<TSharedPtr<IDashboardDataTreeViewEntry>>& InEntries, const bool bShowPlot, const bool bIgnoreCategories /*= false*/)
+	{
+		using namespace FSoundDashboardViewFactoryPrivate;
+
+		for (const TSharedPtr<IDashboardDataTreeViewEntry>& Entry : InEntries)
+		{
+			if (!Entry.IsValid())
+			{
+				continue;
+			}
+
+			FSoundDashboardEntry& SoundEntry = CastEntry(*Entry);
+			if (SoundEntry.bIsCategory)
+			{
+				if (bIgnoreCategories)
+				{
+					continue;
+				}
+
+				SetShowPlotRecursive(Entry, bShowPlot);
+			}
+			else
+			{
+				SetShowPlot(Entry, bShowPlot);
+			}
+		}
+	}
+
+	void FSoundDashboardViewFactory::SetShowPlotRecursive(const TSharedPtr<IDashboardDataTreeViewEntry>& Entry, const bool bShowPlot)
+	{
+		using namespace FSoundDashboardViewFactoryPrivate;
+
+		if (!Entry.IsValid())
+		{
+			return;
+		}
+
+		SetShowPlot(Entry, bShowPlot);
+
+		for (const TSharedPtr<IDashboardDataTreeViewEntry>& ChildEntry : Entry->Children)
+		{
+			SetShowPlotRecursive(ChildEntry, bShowPlot);
+		}
+	}
+
+	void FSoundDashboardViewFactory::SetShowPlot(const TSharedPtr<IDashboardDataTreeViewEntry>& Entry, const bool bShowPlot)
+	{
+		using namespace FSoundDashboardViewFactoryPrivate;
+
+		if (!Entry.IsValid())
+		{
+			return;
+		}
+
+		FSoundDashboardEntry& SoundEntry = CastEntry(*Entry);
+
+		// Special case: do not enable toggling plots on the Pinned category header
+		if (SoundEntry.EntryType == ESoundDashboardEntryType::Pinned)
+		{
+			return;
+		}
+
+		SoundEntry.bIsPlotActive = bShowPlot;
+
+		// If this is a Pinned entry, make sure to replicate the value in the original copy
+		if (PinnedItemEntries.IsValid() && SoundEntry.PinnedEntryType == FSoundDashboardEntry::EPinnedEntryType::PinnedCopy)
+		{
+			const TSharedPtr<IDashboardDataTreeViewEntry>& OriginalEntry = PinnedItemEntries->FindOriginalEntryInChildren(Entry);
+			if (OriginalEntry.IsValid())
+			{
+				FSoundDashboardEntry& OriginalSoundEntry = CastEntry(*OriginalEntry);
+
+				OriginalSoundEntry.bIsPlotActive = bShowPlot;
+			}
+		}
+	}
+
+	TMap<FName, FSoundPlotsWidgetView::FPlotColumnInfo> FSoundDashboardViewFactory::GetPlotColumnInfo()
+	{
+		using namespace FSoundDashboardViewFactoryPrivate;
+
+		auto CreatePlotColumnInfo = [this]()
+		{
+			return TMap<FName, FSoundPlotsWidgetView::FPlotColumnInfo>
+			{
+				{
+					AmplitudeColumnName,
+					{
+						AmplitudeColumnDisplayName,
+						[](const IDashboardDataTreeViewEntry& InEntry) -> const ::Audio::TCircularAudioBuffer<FDataPoint>& { return CastEntry(InEntry).AmplitudeDataPoints; },
+						[this](const float InValue) { return bDisplayAmpPeakInDb ? ::Audio::ConvertToDecibels(InValue) : InValue; },
+						[this]() { return bDisplayAmpPeakInDb ? FSlateStyle::Get().GetAmpFloatFormat() : FSlateStyle::Get().GetLinearVolumeFloatFormat(); },
+						[this]() { return bDisplayAmpPeakInDb ? FVector2f(-160.0f, 0.0f) : FVector2f(0.0f, 1.0f); }
+					}
+				},
+				{
+					VolumeColumnName,
+					{
+						VolumeColumnDisplayName,
+						[](const IDashboardDataTreeViewEntry& InEntry) -> const ::Audio::TCircularAudioBuffer<FDataPoint>& { return CastEntry(InEntry).VolumeDataPoints; },
+						nullptr,
+						[this]() { return FSlateStyle::Get().GetLinearVolumeFloatFormat(); },
+						[this]() { return FVector2f(0.0f, 1.0f); }
+					}
+				},
+				{
+					DistanceColumnName,
+					{
+						DistanceColumnDisplayName,
+						[](const IDashboardDataTreeViewEntry& InEntry) -> const ::Audio::TCircularAudioBuffer<FDataPoint>& { return CastEntry(InEntry).DistanceDataPoints; },
+						nullptr,
+						[this]() { return FSlateStyle::Get().GetFreqFloatFormat(); },
+						[this]() { return FVector2f(10000.0f, 0.0f); }
+					}
+				},
+				{
+					DistanceAttenuationColumnName,
+					{
+						DistanceAttenuationColumnDisplayName,
+						[](const IDashboardDataTreeViewEntry& InEntry) -> const ::Audio::TCircularAudioBuffer<FDataPoint>& { return CastEntry(InEntry).DistanceAttenuationDataPoints; },
+						nullptr,
+						[this]() { return FSlateStyle::Get().GetLinearVolumeFloatFormat(); },
+						[this]() { return FVector2f(0.0f, 1.0f); }
+					}
+				},
+				{
+					PitchColumnName,
+					{
+						PitchColumnDisplayName,
+						[](const IDashboardDataTreeViewEntry& InEntry) -> const ::Audio::TCircularAudioBuffer<FDataPoint>& { return CastEntry(InEntry).PitchDataPoints; },
+						nullptr,
+						[this]() { return FSlateStyle::Get().GetPitchFloatFormat(); },
+						[this]() { return FVector2f(0.0f, 1.0f); }
+					}
+				},
+				{
+					PriorityColumnName,
+					{
+						PriorityColumnDisplayName,
+						[](const IDashboardDataTreeViewEntry& InEntry) -> const ::Audio::TCircularAudioBuffer<FDataPoint>& { return CastEntry(InEntry).PriorityDataPoints; },
+						nullptr,
+						[this]() { return FSlateStyle::Get().GetLinearVolumeFloatFormat(); },
+						[this]() { return FVector2f(0.0f, 1.0f); }
+					}
+				},
+				{
+					LPFFreqColumnName,
+					{
+						LPFFreqColumnDisplayName,
+						[](const IDashboardDataTreeViewEntry& InEntry) -> const ::Audio::TCircularAudioBuffer<FDataPoint>& { return CastEntry(InEntry).LPFFreqDataPoints; },
+						nullptr,
+						[this]() { return FSlateStyle::Get().GetFreqFloatFormat(); },
+						[this]() { return FVector2f(20.0f, 20000.0f); }
+					}
+				},
+				{
+					HPFFreqColumnName,
+					{
+						HPFFreqColumnDisplayName,
+						[](const IDashboardDataTreeViewEntry& InEntry) -> const ::Audio::TCircularAudioBuffer<FDataPoint>& { return CastEntry(InEntry).HPFFreqDataPoints; },
+						nullptr,
+						[this]() { return FSlateStyle::Get().GetFreqFloatFormat(); },
+						[this]() { return FVector2f(20.0f, 20000.0f); }
+					}
+				}
+			};
+		};
+
+		return CreatePlotColumnInfo();
+	}
+
 	bool FSoundDashboardViewFactory::SelectedItemsIncludesAnAsset() const
 	{
 		using namespace FSoundDashboardViewFactoryPrivate;
diff --git a/Engine/Plugins/AudioInsights/Source/AudioInsights/Private/Views/SoundPlotsDashboardViewFactory.cpp b/Engine/Plugins/AudioInsights/Source/AudioInsights/Private/Views/SoundPlotsDashboardViewFactory.cpp
new file mode 100644
--- /dev/null
+++ b/Engine/Plugins/AudioInsights/Source/AudioInsights/Private/Views/SoundPlotsDashboardViewFactory.cpp
@@ -0,0 +1,67 @@
+// Copyright Epic Games, Inc. All Rights Reserved.
+#include "Views/SoundPlotsDashboardViewFactory.h"
+
+#include "AudioInsightsStyle.h"
+#include "Messages/SoundTraceMessages.h"
+#include "Views/SoundDashboardViewFactory.h"
+
+#define LOCTEXT_NAMESPACE "AudioInsights"
+
+namespace UE::Audio::Insights
+{
+	/////////////////////////////////////////////////////////////////////////////////////////
+	// FSoundPlotsDashboardViewFactoryPrivate
+	namespace FSoundPlotsDashboardViewFactoryPrivate
+	{
+		const FSoundDashboardEntry& CastEntry(const IDashboardDataTreeViewEntry& InData)
+		{
+			return static_cast<const FSoundDashboardEntry&>(InData);
+		};
+	}
+
+	/////////////////////////////////////////////////////////////////////////////////////////
+	// FSoundPlotsDashboardViewFactory
+	FName FSoundPlotsDashboardViewFactory::GetName() const
+	{
+		return "Plots";
+	}
+
+	FText FSoundPlotsDashboardViewFactory::GetDisplayName() const
+	{
+		return LOCTEXT("AudioDashboard_PlotsTab_Name", "Plots");
+	}
+
+	EDefaultDashboardTabStack FSoundPlotsDashboardViewFactory::GetDefaultTabStack() const
+	{
+		return EDefaultDashboardTabStack::Plots;
+	}
+
+	FSlateIcon FSoundPlotsDashboardViewFactory::GetIcon() const
+	{
+		return FSlateStyle::Get().CreateIcon("AudioInsights.Icon.SoundDashboard.Plots");
+	}
+
+	TSharedRef<SWidget> FSoundPlotsDashboardViewFactory::MakeWidget(TSharedRef<SDockTab> OwnerTab, const FSpawnTabArgs& SpawnTabArgs)
+	{
+		return PlotsWidget ? PlotsWidget->MakeWidget() : SNullWidget::NullWidget;
+	}
+
+	void FSoundPlotsDashboardViewFactory::InitPlots(const TSharedRef<FSoundDashboardViewFactory> SoundsDashboard)
+	{
+		using namespace FSoundPlotsDashboardViewFactoryPrivate;
+
+		TMap<FName, FSoundPlotsWidgetView::FPlotColumnInfo> ColumnInfo = SoundsDashboard->GetPlotColumnInfo();
+		PlotsWidget = MakeShared<FSoundPlotsWidgetView>(MoveTemp(ColumnInfo), [](const IDashboardDataTreeViewEntry& InEntry) -> bool { return CastEntry(InEntry).bIsPlotActive; });
+
+		OnProcessPlotData = SoundsDashboard->OnProcessPlotData.AddSP(this, &FSoundPlotsDashboardViewFactory::ProcessPlotData);
+	}
+
+	void FSoundPlotsDashboardViewFactory::ProcessPlotData(const TArray<TSharedPtr<IDashboardDataTreeViewEntry>>& DataViewEntries)
+	{
+		if (PlotsWidget.IsValid())
+		{
+			PlotsWidget->ProcessPlotData(DataViewEntries);
+		}
+	}
+} // namespace UE::Audio::Insights
+#undef LOCTEXT_NAMESPACE
diff --git a/Engine/Plugins/AudioInsights/Source/AudioInsights/Private/Views/SoundPlotsWidgetView.cpp b/Engine/Plugins/AudioInsights/Source/AudioInsights/Private/Views/SoundPlotsWidgetView.cpp
new file mode 100644
--- /dev/null
+++ b/Engine/Plugins/AudioInsights/Source/AudioInsights/Private/Views/SoundPlotsWidgetView.cpp
@@ -0,0 +1,466 @@
+// Copyright Epic Games, Inc. All Rights Reserved.
+#include "Views/SoundPlotsWidgetView.h"
+
+#include "AudioInsightsModule.h"
+#include "AudioInsightsStyle.h"
+#include "Features/IModularFeatures.h"
+#include "Widgets/Input/SComboBox.h"
+
+#if WITH_EDITOR
+#include "Editor.h"
+#else
+#include "AudioInsightsComponent.h"
+#include "AudioInsightsTimingViewExtender.h"
+#endif // WITH_EDITOR
+
+#define LOCTEXT_NAMESPACE "AudioInsights"
+
+namespace UE::Audio::Insights
+{
+	FSoundPlotsWidgetView::FSoundPlotsWidgetView(TMap<FName, FSoundPlotsWidgetView::FPlotColumnInfo>&& InColumnInfo, const TFunctionRef<bool(const IDashboardDataTreeViewEntry& InEntry)> InIsPlotEnabledForEntryFunc)
+		: ColumnInfo(InColumnInfo)
+		, IsPlotEnabledForEntryFunc(InIsPlotEnabledForEntryFunc)
+	{
+		FTraceModule& AudioInsightsTraceModule = static_cast<FTraceModule&>(FAudioInsightsModule::GetChecked().GetTraceModule());
+		AudioInsightsTraceModule.OnAnalysisStarting.AddRaw(this, &FSoundPlotsWidgetView::OnAnalysisStarting);
+	}
+
+	FSoundPlotsWidgetView::~FSoundPlotsWidgetView()
+	{
+		if (FModuleManager::Get().IsModuleLoaded("AudioInsights") && IModularFeatures::Get().IsModularFeatureAvailable(TraceServices::ModuleFeatureName))
+		{
+			FTraceModule& TraceModule = static_cast<FTraceModule&>(FAudioInsightsModule::GetChecked().GetTraceModule());
+			TraceModule.OnAnalysisStarting.RemoveAll(this);
+		}
+	}
+
+	TSharedRef<SWidget> FSoundPlotsWidgetView::MakeWidget()
+	{
+#if WITH_EDITOR
+		FEditorDelegates::PostPIEStarted.AddSP(this, &FSoundPlotsWidgetView::OnPIEStarted);
+		FEditorDelegates::EndPIE.AddSP(this, &FSoundPlotsWidgetView::OnPIEStopped);
+		FEditorDelegates::PausePIE.AddSP(this, &FSoundPlotsWidgetView::OnPIEPaused);
+		FEditorDelegates::ResumePIE.AddSP(this, &FSoundPlotsWidgetView::OnPIEResumed);
+#else
+		FAudioInsightsModule& AudioInsightsModule = FAudioInsightsModule::GetChecked();
+
+		AudioInsightsModule.GetTimingViewExtender().OnTimingViewTimeMarkerChanged.AddSP(this, &FSoundPlotsWidgetView::OnTimingViewTimeMarkerChanged);
+
+		const TSharedPtr<FAudioInsightsComponent> AudioInsightsComponent = AudioInsightsModule.GetAudioInsightsComponent();
+		if (AudioInsightsComponent.IsValid())
+		{
+			AudioInsightsComponent->OnTabSpawn.AddSP(this, &FSoundPlotsWidgetView::OnAudioInsightsComponentTabSpawn);
+			AudioInsightsComponent->OnSessionAnalysisCompleted.AddSP(this, &FSoundPlotsWidgetView::OnSessionAnalysisCompleted);
+		}
+#endif // WITH_EDITOR
+
+		return MakePlotsWidget();
+	}
+
+	void FSoundPlotsWidgetView::ProcessPlotData(const TArray<TSharedPtr<IDashboardDataTreeViewEntry>>& DataViewEntries)
+	{
+		if (!PlotWidgetMetadataPerCurve.IsValid() || DataViewEntries.IsEmpty())
+		{
+			return;
+		}
+
+		// Process new data
+		bool bHasNewMetadata = false;
+		for (const TSharedPtr<IDashboardDataTreeViewEntry>& DataEntry : DataViewEntries)
+		{
+			CollectDataRecursive(DataEntry, bHasNewMetadata);
+		}
+
+		// Set metadata for each widget if updated
+		if (bHasNewMetadata)
+		{
+			PlotWidget->SetCurvesMetadata(PlotWidgetMetadataPerCurve);
+		}
+
+		// Remove old points and set curve data for each widget
+#if WITH_EDITOR
+		const double PlotDrawLimitTimestamp = CurrentTimestamp - BeginTimestamp - (MaxPlotHistorySeconds + 0.2 /* extra grace time to avoid curve cuts being displayed */);
+#endif // WITH_EDITOR
+
+		const TSharedPtr<FPointDataPerCurveMap> CurveDataMapPtr = *PlotWidgetCurveIdToPointDataMapPerColumn.Find(SelectedPlotColumnName);
+		if (CurveDataMapPtr.IsValid())
+		{
+#if WITH_EDITOR
+			// Remove points that are older than max history limit from the most recent timestamp
+			for (auto& [CurveId, CurvePoints] : *CurveDataMapPtr)
+			{
+				const int32 FoundIndex = CurvePoints.IndexOfByPredicate([&PlotDrawLimitTimestamp](const FDataPoint& InDataPoint)
+				{
+					return InDataPoint.Key >= PlotDrawLimitTimestamp;
+				});
+
+				if (FoundIndex > 0)
+				{
+					CurvePoints.RemoveAt(0, FoundIndex, EAllowShrinking::No);
+				}
+			}
+#endif // WITH_EDITOR
+
+			if (!CurveDataMapPtr->IsEmpty())
+			{
+				PlotWidget->SetCurvesPointData(CurveDataMapPtr);
+				PlotWidget->SetYValueFormattingOptions(*GetPlotColumnNumberFormat(SelectedPlotColumnName));
+			}
+		}
+	}
+
+	void FSoundPlotsWidgetView::OnAnalysisStarting(const double Timestamp)
+	{
+#if WITH_EDITOR
+		BeginTimestamp = Timestamp - GStartTime;
+#else
+		BeginTimestamp = 0.0;
+#endif // WITH_EDITOR
+	}
+
+#if WITH_EDITOR
+	void FSoundPlotsWidgetView::OnPIEStarted(bool bSimulating)
+	{
+		GameState = EGameState::Running;
+	}
+
+	void FSoundPlotsWidgetView::OnPIEStopped(bool bSimulating)
+	{
+		ResetPlots();
+
+		GameState = EGameState::Stopped;
+	}
+
+	void FSoundPlotsWidgetView::OnPIEPaused(bool bSimulating)
+	{
+		GameState = EGameState::Paused;
+	}
+
+	void FSoundPlotsWidgetView::OnPIEResumed(bool bSimulating)
+	{
+		GameState = EGameState::Running;
+	}
+#else
+	void FSoundPlotsWidgetView::OnAudioInsightsComponentTabSpawn()
+	{
+		const TSharedPtr<const FAudioInsightsComponent> AudioInsightsComponent = FAudioInsightsModule::GetChecked().GetAudioInsightsComponent();
+		if (AudioInsightsComponent.IsValid())
+		{
+			GameState = AudioInsightsComponent->GetIsLiveSession() ? EGameState::Running : EGameState::Stopped;
+		}
+	}
+
+	void FSoundPlotsWidgetView::OnSessionAnalysisCompleted()
+	{
+		GameState = EGameState::Stopped;
+	}
+
+	void FSoundPlotsWidgetView::OnTimingViewTimeMarkerChanged(double InTimeMarker)
+	{
+		CurrentTimestamp = InTimeMarker;
+	}
+#endif // WITH_EDITOR
+
+	TSharedRef<SWidget> FSoundPlotsWidgetView::MakePlotsWidget()
+	{
+		// Initialize column options and initially selected column
+		ColumnInfo.GenerateKeyArray(ColumnNames);
+
+		if (ColumnNames.IsEmpty())
+		{
+			return SNullWidget::NullWidget;
+		}
+
+		SelectedPlotColumnName = ColumnNames[0];
+
+		// Initialize curve data and metadata
+		if (!PlotWidgetMetadataPerCurve.IsValid())
+		{
+			PlotWidgetMetadataPerCurve = MakeShared<TMap<int32, SAudioCurveView::FCurveMetadata>>();
+			for (const FName& ColumnName : ColumnNames)
+			{
+				TSharedPtr<FPointDataPerCurveMap> PointDataPerCurveMap = MakeShared<FPointDataPerCurveMap>();
+				PlotWidgetCurveIdToPointDataMapPerColumn.Emplace(ColumnName, MoveTemp(PointDataPerCurveMap));
+				ColumnNameToLatestTimestamp.Emplace(ColumnName, MakeShared<FLatestTimestampPerCurveMap>());
+			}
+		}
+
+		SAssignNew(PlotWidget, SAudioCurveView)
+			.ViewRange_Lambda([this]() { return GetViewRange(); })
+			.PixelSnappingMethod(EWidgetPixelSnapping::Disabled);
+
+		PlotWidget->SetYValueFormattingOptions(*GetPlotColumnNumberFormat(SelectedPlotColumnName));
+		PlotWidget->SetInitYAxisRange(GetPlotColumnDefaultYRange(SelectedPlotColumnName));
+
+		// Create plot column combo box widgets  
+		auto CreatePlotColumnComboBoxWidget = [this]()
+		{
+			return SNew(SComboBox<FName>)
+				.ToolTipText(LOCTEXT("SoundPlotsWidgetView_SelectPlotColumnDescription", "Select a column from the table to plot."))
+				.OptionsSource(&ColumnNames)
+				.OnGenerateWidget_Lambda([this](const FName& ColumnName)
+				{
+					return SNew(STextBlock)
+					.Font(FAppStyle::GetFontStyle(TEXT("PropertyWindow.NormalFont")))
+					.Text(GetPlotColumnDisplayName(ColumnName));
+				})
+				.OnSelectionChanged_Lambda([this](FName NewColumnName, ESelectInfo::Type)
+				{
+					SelectedPlotColumnName = NewColumnName;
+					if (TSharedPtr<FPointDataPerCurveMap>* DataMap = PlotWidgetCurveIdToPointDataMapPerColumn.Find(NewColumnName))
+					{
+						PlotWidget->SetCurvesPointData(*DataMap);
+						PlotWidget->SetYValueFormattingOptions(*GetPlotColumnNumberFormat(NewColumnName));
+						PlotWidget->SetInitYAxisRange(GetPlotColumnDefaultYRange(NewColumnName));
+					}
+				})
+				[
+					SNew(STextBlock)
+					.Font(FAppStyle::GetFontStyle(TEXT("PropertyWindow.NormalFont")))
+					.Text_Lambda([this]()
+					{
+						return GetPlotColumnDisplayName(SelectedPlotColumnName);
+					})
+				];
+		};
+
+		return SNew(SVerticalBox)
+			.Clipping(EWidgetClipping::ClipToBounds)
+			+ SVerticalBox::Slot()
+			.AutoHeight()
+			.HAlign(HAlign_Fill)
+			[
+				SNew(SSimpleTimeSlider)
+				.ViewRange_Lambda([this]() { return GetViewRange(); })
+				.ClampRangeHighlightSize(0.0f) // Hide clamp range
+				.ScrubPosition_Lambda([]() { return TNumericLimits<double>::Lowest(); }) // Hide scrub
+				.PixelSnappingMethod(EWidgetPixelSnapping::Disabled)
+			]
+			+ SVerticalBox::Slot()
+			.AutoHeight()
+			.HAlign(HAlign_Fill)
+			[
+				CreatePlotColumnComboBoxWidget()
+			]
+			+ SVerticalBox::Slot()
+			.HAlign(HAlign_Fill)
+			[
+				PlotWidget.ToSharedRef()
+			];
+	}
+
+	void FSoundPlotsWidgetView::CollectDataRecursive(const TSharedPtr<IDashboardDataTreeViewEntry>& Entry, bool& bOutHasNewMetadata)
+	{
+		if (!Entry.IsValid())
+		{
+			return;
+		}
+
+		const uint32 SourceId = Entry->GetEntryID();
+		const bool bPlotIsEnabled = IsPlotEnabledForEntryFunc(*Entry);
+
+		if (SourceId == INDEX_NONE || !bPlotIsEnabled)
+		{
+			for (const TSharedPtr<IDashboardDataTreeViewEntry>& ChildEntry : Entry->Children)
+			{
+				CollectDataRecursive(ChildEntry, bOutHasNewMetadata);
+			}
+
+			return;
+		}
+
+		// For each column, get the array for this data point's source id and add the value to that data array
+		for (const auto& [ColumnName, DataMap] : PlotWidgetCurveIdToPointDataMapPerColumn)
+		{
+			// Get the data point array for this source id, or add new point
+			TArray<FPlotCurvePoint>* DataPoints = DataMap->Find(SourceId);
+
+			if (DataPoints == nullptr)
+			{
+				DataPoints = &DataMap->Add(SourceId);
+
+				TSharedPtr<FLatestTimestampPerCurveMap>* LatestTimestamp = ColumnNameToLatestTimestamp.Find(ColumnName);
+				if (LatestTimestamp && LatestTimestamp->IsValid() && !(*LatestTimestamp)->Contains(SourceId))
+				{
+					(*LatestTimestamp)->Add(SourceId);
+				}
+			}
+
+			auto DataFunc = GetPlotColumnDataFunc(ColumnName);
+			if (!DataFunc.IsSet())
+			{
+				continue;
+			}
+
+			const ::Audio::TCircularAudioBuffer<FDataPoint>& TimeStampedValues = (DataFunc)(*Entry);
+
+			const ::Audio::DisjointedArrayView<const FDataPoint> TimeStampedValuesDisjointedArrayView = TimeStampedValues.PeekInPlace(TimeStampedValues.Num());
+
+			const TFunction<float(const float)> UnitConversionFunc = GetPlotColumnUnitConversionFunc(ColumnName);
+
+			for (const auto& [Timestamp, Value] : TimeStampedValuesDisjointedArrayView.FirstBuffer)
+			{
+				TSharedPtr<FLatestTimestampPerCurveMap>* LatestTimestampMap = ColumnNameToLatestTimestamp.Find(ColumnName);
+				if (LatestTimestampMap && LatestTimestampMap->IsValid())
+				{
+					double* LatestColumnTimestep = (*LatestTimestampMap)->Find(SourceId);
+					if (LatestColumnTimestep)
+					{
+						if (Timestamp < *LatestColumnTimestep)
+						{
+							continue;
+						}
+						else
+						{
+							*LatestColumnTimestep = Timestamp;
+						}
+					}
+				}
+
+				CurrentTimestamp = FMath::Max(CurrentTimestamp, Timestamp);
+
+#if WITH_EDITOR
+				const double DataPointTime = Timestamp - BeginTimestamp;
+#else
+				const double DataPointTime = Timestamp;
+#endif // WITH_EDITOR
+				DataPoints->Emplace(DataPointTime, UnitConversionFunc.IsSet() ? UnitConversionFunc(Value) : Value);
+			}
+		}
+
+		// Create metadata for this curve if necessary
+		if (!PlotWidgetMetadataPerCurve->Contains(SourceId))
+		{
+			FPlotCurveMetadata& NewMetadata = PlotWidgetMetadataPerCurve->Add(SourceId);
+			NewMetadata.CurveColor = FLinearColor(FColor::MakeRandomColor());
+			NewMetadata.DisplayName = Entry->GetDisplayName();
+			bOutHasNewMetadata = true;
+		}
+
+		for (const TSharedPtr<IDashboardDataTreeViewEntry>& ChildEntry : Entry->Children)
+		{
+			CollectDataRecursive(ChildEntry, bOutHasNewMetadata);
+		}
+	}
+
+	void FSoundPlotsWidgetView::ResetPlots()
+	{
+		for (const auto& [ColumnName, PointDataPerCurveMap] : PlotWidgetCurveIdToPointDataMapPerColumn)
+		{
+			PointDataPerCurveMap->Empty();
+		}
+
+		for (auto& [ColumnName, LatestTimestampsPerPlot] : ColumnNameToLatestTimestamp)
+		{
+			LatestTimestampsPerPlot->Empty();
+		}
+
+		if (PlotWidgetMetadataPerCurve.IsValid())
+		{
+			PlotWidgetMetadataPerCurve->Empty();
+		}
+
+		BeginTimestamp = TNumericLimits<double>::Max();
+		CurrentTimestamp = TNumericLimits<double>::Lowest();
+	}
+
+	const TFunction<const::Audio::TCircularAudioBuffer<FDataPoint>& (const IDashboardDataTreeViewEntry& InData)> FSoundPlotsWidgetView::GetPlotColumnDataFunc(const FName& ColumnName)
+	{
+		return ColumnInfo.Find(ColumnName)->DataFunc;
+	}
+
+	const TFunction<float(const float)> FSoundPlotsWidgetView::GetPlotColumnUnitConversionFunc(const FName& ColumnName)
+	{
+		return ColumnInfo.Find(ColumnName)->UnitConversionFunc;
+	}
+
+	const FNumberFormattingOptions* FSoundPlotsWidgetView::GetPlotColumnNumberFormat(const FName& ColumnName)
+	{
+		if (const FSoundPlotsWidgetView::FPlotColumnInfo* PlotColumnInfo = ColumnInfo.Find(ColumnName))
+		{
+			if (PlotColumnInfo->GetFormatOptionsFunc.IsSet())
+			{
+				return PlotColumnInfo->GetFormatOptionsFunc();
+			}
+		}
+		return nullptr;
+	}
+
+	const FVector2f FSoundPlotsWidgetView::GetPlotColumnDefaultYRange(const FName& ColumnName)
+	{
+		if (const FSoundPlotsWidgetView::FPlotColumnInfo* PlotColumnInfo = ColumnInfo.Find(ColumnName))
+		{
+			if (PlotColumnInfo->GetDefaultYDataRangeFunc.IsSet())
+			{
+				return PlotColumnInfo->GetDefaultYDataRangeFunc();
+			}
+		}
+		return FVector2f(0.0f, 1.0f);
+	}
+
+	const FText FSoundPlotsWidgetView::GetPlotColumnDisplayName(const FName& ColumnName)
+	{
+		if (const FSoundPlotsWidgetView::FPlotColumnInfo* PlotColumnInfo = ColumnInfo.Find(ColumnName))
+		{
+			return PlotColumnInfo->ColumnDisplayName;
+		}
+		return FText::GetEmpty();
+	}
+
+	TRange<double> FSoundPlotsWidgetView::GetViewRange()
+	{
+#if WITH_EDITOR
+		if (GameState == EGameState::Stopped || BeginTimestamp == TNumericLimits<double>::Max())
+#else
+		if (BeginTimestamp == TNumericLimits<double>::Max())
+#endif // WITH_EDITOR
+		{
+			return TRange<double>(0, MaxPlotHistorySeconds);
+		}
+
+		double RangeUpperBound = 0.0;
+
+#if WITH_EDITOR
+		const double CurrentTime = FPlatformTime::Seconds() - GStartTime;
+
+		double TimestampsDiff = 0.0;
+
+		const bool bAnyMessageReceived = CurrentTimestamp != TNumericLimits<double>::Lowest();
+
+		if (GameState == EGameState::Running && bAnyMessageReceived)
+		{
+			const double RelativeCurrentTime = CurrentTime - BeginTimestamp;
+			TimestampsDiff = RelativeCurrentTime - (CurrentTimestamp - BeginTimestamp);
+		}
+
+		const double FinalCurrentTime = bAnyMessageReceived ? CurrentTimestamp : CurrentTime;
+
+		constexpr double RangeAlignmentOffset = 0.2;
+		RangeUpperBound = FinalCurrentTime - BeginTimestamp + TimestampsDiff - RangeAlignmentOffset;
+#else
+		const TSharedPtr<const FAudioInsightsComponent> AudioInsightsComponent = FAudioInsightsModule::GetChecked().GetAudioInsightsComponent();
+		if (AudioInsightsComponent.IsValid() && !AudioInsightsComponent->GetIsLiveSession())
+		{
+			return TRange<double>(CurrentTimestamp - MaxPlotHistorySeconds, CurrentTimestamp);
+		}
+		else
+		{
+			const double CurrentTime = FPlatformTime::Seconds();
+			const double DeltaTime = CurrentTime - PreviousTime;
+
+			PreviousTime = CurrentTime;
+
+			const double TraceCurrentDurationSeconds = FAudioInsightsModule::GetChecked().GetTimingViewExtender().GetCurrentDurationSeconds();
+
+			CurrentRangeUpperBound = FMath::FInterpTo(CurrentRangeUpperBound, TraceCurrentDurationSeconds, DeltaTime, 1.0);
+
+			constexpr double RangeAlignmentOffset = 0.9;
+			RangeUpperBound = CurrentRangeUpperBound + RangeAlignmentOffset;
+		}
+#endif // WITH_EDITOR
+
+		return TRange<double>(RangeUpperBound - MaxPlotHistorySeconds, RangeUpperBound);
+	}
+} // namespace UE::Audio::Insights
+
+#undef LOCTEXT_NAMESPACE
diff --git a/Engine/Plugins/AudioInsights/Source/AudioInsights/Private/Views/TreeDashboardViewFactory.cpp b/Engine/Plugins/AudioInsights/Source/AudioInsights/Private/Views/TreeDashboardViewFactory.cpp
--- a/Engine/Plugins/AudioInsights/Source/AudioInsights/Private/Views/TreeDashboardViewFactory.cpp
+++ b/Engine/Plugins/AudioInsights/Source/AudioInsights/Private/Views/TreeDashboardViewFactory.cpp
@@ -274,8 +274,8 @@ namespace UE::Audio::Insights
 
 	void FTraceTreeDashboardViewFactory::OnColumnSortModeChanged(const EColumnSortPriority::Type InSortPriority, const FName& InColumnId, const EColumnSortMode::Type InSortMode)
 	{
-		// Disable sorting by the solo/mute columns
-		if (InColumnId == "Solo" || InColumnId == "Mute")
+		// Sorting can be disabled on specified columns
+		if (!IsColumnSortable(InColumnId))
 		{
 			return;
 		}
diff --git a/Engine/Plugins/AudioInsights/Source/AudioInsights/Public/AudioInsightsDataSource.h b/Engine/Plugins/AudioInsights/Source/AudioInsights/Public/AudioInsightsDataSource.h
--- a/Engine/Plugins/AudioInsights/Source/AudioInsights/Public/AudioInsightsDataSource.h
+++ b/Engine/Plugins/AudioInsights/Source/AudioInsights/Public/AudioInsightsDataSource.h
@@ -20,6 +20,8 @@ namespace UE::Audio::Insights
 		virtual ~IDashboardDataTreeViewEntry() = default;
 
 		virtual bool IsValid() const = 0;
+		virtual int32 GetEntryID() const = 0;
+		virtual const FText& GetDisplayName() const = 0;
 		virtual bool HasSetInitExpansion() const = 0;
 		virtual void ResetHasSetInitExpansion() = 0;
 
diff --git a/Engine/Plugins/AudioInsights/Source/AudioInsights/Public/AudioInsightsStyle.h b/Engine/Plugins/AudioInsights/Source/AudioInsights/Public/AudioInsightsStyle.h
--- a/Engine/Plugins/AudioInsights/Source/AudioInsights/Public/AudioInsightsStyle.h
+++ b/Engine/Plugins/AudioInsights/Source/AudioInsights/Public/AudioInsightsStyle.h
@@ -189,6 +189,7 @@ namespace UE::Audio::Insights
 			Set("AudioInsights.Icon.SoundDashboard.MetaSound",        new IMAGE_BRUSH_SVG(TEXT("Icons/SoundDashboard/metasound"),        Icon16));
 			Set("AudioInsights.Icon.SoundDashboard.Mute",             new IMAGE_BRUSH_SVG(TEXT("Icons/SoundDashboard/mute"),             Icon16));
 			Set("AudioInsights.Icon.SoundDashboard.Pin",              new IMAGE_BRUSH_SVG(TEXT("Icons/SoundDashboard/pin"),              Icon16));
+			Set("AudioInsights.Icon.SoundDashboard.Plots",			  new IMAGE_BRUSH_SVG(TEXT("Icons/SoundDashboard/plots"),			 Icon16));
 			Set("AudioInsights.Icon.SoundDashboard.ProceduralSource", new IMAGE_BRUSH_SVG(TEXT("Icons/SoundDashboard/proceduralsource"), Icon16));
 			Set("AudioInsights.Icon.SoundDashboard.Reset",            new IMAGE_BRUSH_SVG(TEXT("Icons/SoundDashboard/reset"),            Icon16));
 			Set("AudioInsights.Icon.SoundDashboard.Solo",             new IMAGE_BRUSH_SVG(TEXT("Icons/SoundDashboard/solo"),             Icon16));
diff --git a/Engine/Plugins/AudioInsights/Source/AudioInsights/Public/Messages/SoundTraceMessages.h b/Engine/Plugins/AudioInsights/Source/AudioInsights/Public/Messages/SoundTraceMessages.h
--- a/Engine/Plugins/AudioInsights/Source/AudioInsights/Public/Messages/SoundTraceMessages.h
+++ b/Engine/Plugins/AudioInsights/Source/AudioInsights/Public/Messages/SoundTraceMessages.h
@@ -187,6 +187,7 @@ namespace UE::Audio::Insights
 		}
 
 		virtual bool IsValid() const override { return PlayOrder != INDEX_NONE; }
+		virtual int32 GetEntryID() const override { return PlayOrder; }
 
 		virtual bool HasSetInitExpansion() const override { return bHasSetInitExpansion; }
 		virtual void ResetHasSetInitExpansion() override { bHasSetInitExpansion = true; }
@@ -208,6 +209,7 @@ namespace UE::Audio::Insights
 		bool bIsCategory = false;
 		bool bHasSetInitExpansion = false;
 		bool bIsVisible = false;
+		bool bIsPlotActive = false;
 
 		::Audio::TCircularAudioBuffer<FDataPoint> PriorityDataPoints;
 		::Audio::TCircularAudioBuffer<FDataPoint> DistanceDataPoints;
diff --git a/Engine/Plugins/AudioInsights/Source/AudioInsights/Public/Views/DashboardViewFactory.h b/Engine/Plugins/AudioInsights/Source/AudioInsights/Public/Views/DashboardViewFactory.h
--- a/Engine/Plugins/AudioInsights/Source/AudioInsights/Public/Views/DashboardViewFactory.h
+++ b/Engine/Plugins/AudioInsights/Source/AudioInsights/Public/Views/DashboardViewFactory.h
@@ -21,6 +21,7 @@ namespace UE::Audio::Insights
 		Analysis,
 		AudioMeters,
 		AudioAnalyzerRack,
+		Plots
 	};
 
 	class IDashboardViewFactory
diff --git a/Engine/Plugins/AudioInsights/Source/AudioInsights/Public/Views/SAudioCurveView.h b/Engine/Plugins/AudioInsights/Source/AudioInsights/Public/Views/SAudioCurveView.h
--- a/Engine/Plugins/AudioInsights/Source/AudioInsights/Public/Views/SAudioCurveView.h
+++ b/Engine/Plugins/AudioInsights/Source/AudioInsights/Public/Views/SAudioCurveView.h
@@ -58,6 +58,7 @@ public:
 	void Construct( const FArguments& InArgs );
 
 	void SetYValueFormattingOptions(const FNumberFormattingOptions InValueFormattingOptions);
+	void SetInitYAxisRange(const FVector2f& InRange);
 	void SetCurvesPointData(TSharedPtr<TMap<int32, TArray<FCurvePoint>>> InCurvesPointData);
 	void SetCurvesMetadata(TSharedPtr<TMap<int32, FCurveMetadata>> InMetadataPerCurve);
 
diff --git a/Engine/Plugins/AudioInsights/Source/AudioInsights/Public/Views/SoundDashboardViewFactory.h b/Engine/Plugins/AudioInsights/Source/AudioInsights/Public/Views/SoundDashboardViewFactory.h
--- a/Engine/Plugins/AudioInsights/Source/AudioInsights/Public/Views/SoundDashboardViewFactory.h
+++ b/Engine/Plugins/AudioInsights/Source/AudioInsights/Public/Views/SoundDashboardViewFactory.h
@@ -1,7 +1,9 @@
 // Copyright Epic Games, Inc. All Rights Reserved.
 #pragma once
 
+#include "Delegates/Delegate.h"
 #include "Settings/SoundDashboardSettings.h"
+#include "Views/SoundPlotsWidgetView.h"
 #include "Views/TreeDashboardViewFactory.h"
 
 #include "Filters/GenericFilter.h"
@@ -87,6 +89,8 @@ namespace UE::Audio::Insights
 
 		UE_API bool EntryIsValid() const;
 
+		UE_API TSharedPtr<IDashboardDataTreeViewEntry> FindOriginalEntryInChildren(const TSharedPtr<IDashboardDataTreeViewEntry>& RootPinnedEntry);
+
 		TArray<TSharedPtr<FPinnedSoundEntryWrapper>> PinnedWrapperChildren;
 
 	private:
@@ -107,6 +111,11 @@ namespace UE::Audio::Insights
 		UE_API virtual TSharedRef<SWidget> MakeWidget(TSharedRef<SDockTab> OwnerTab, const FSpawnTabArgs& SpawnTabArgs) override;
 		UE_API virtual void ProcessEntries(FTraceTreeDashboardViewFactory::EProcessReason Reason) override;
 
+		UE_API TMap<FName, FSoundPlotsWidgetView::FPlotColumnInfo> GetPlotColumnInfo();
+
+		DECLARE_MULTICAST_DELEGATE_OneParam(FOnProcessPlotData, const TArray<TSharedPtr<IDashboardDataTreeViewEntry>>& /*DataViewEntries*/);
+		FOnProcessPlotData OnProcessPlotData;
+
 	private:
 		UE_API void BindCommands();
 
@@ -127,6 +136,7 @@ namespace UE::Audio::Insights
 		UE_API virtual const TMap<FName, FTraceTreeDashboardViewFactory::FColumnData>& GetColumns() const override;
 
 		UE_API virtual void SortTable() override;
+		UE_API virtual bool IsColumnSortable(const FName& ColumnId) const override;
 		virtual TArray<TSharedPtr<IDashboardDataTreeViewEntry>>& GetTreeItemsSource() override { return FullTree; }
 		UE_API virtual bool ResetTreeData() override;
 
@@ -166,6 +176,18 @@ namespace UE::Audio::Insights
 		UE_API TArray<TObjectPtr<UObject>> GetSelectedEditableAssets() const;
 #endif // WITH_EDITOR
 
+		UE_API TSharedRef<SWidget> MakeShowPlotWidget();
+
+		UE_API TSharedRef<SWidget> CreateShowPlotColumnButton(const TSharedRef<FTraceTreeDashboardViewFactory::SRowWidget>& InRowWidget,
+			const TSharedRef<IDashboardDataTreeViewEntry>& InRowData,
+			const FName& InColumn,
+			TFunction<void(const TArray<TSharedPtr<IDashboardDataTreeViewEntry>>& /*InEntries*/, const bool /*bActivatePlot*/)> ShowPlotToggleFunc,
+			TFunctionRef<bool(const IDashboardDataTreeViewEntry& /*InEntry*/)> IsPlotActiveFunc);
+
+		UE_API void ToggleShowPlotEntries(const TArray<TSharedPtr<IDashboardDataTreeViewEntry>>& InEntries, const bool bShowPlot, const bool bIgnoreCategories = false);
+		UE_API void SetShowPlotRecursive(const TSharedPtr<IDashboardDataTreeViewEntry>& Entry, const bool bShowPlot);
+		UE_API void SetShowPlot(const TSharedPtr<IDashboardDataTreeViewEntry>& Entry, const bool bShowPlot);
+
 		UE_API bool SelectedItemsIncludesAnAsset() const;
 		UE_API bool SelectionIncludesUnpinnedItem() const;
 		UE_API void ClearSelection();
diff --git a/Engine/Plugins/AudioInsights/Source/AudioInsights/Public/Views/SoundPlotsDashboardViewFactory.h b/Engine/Plugins/AudioInsights/Source/AudioInsights/Public/Views/SoundPlotsDashboardViewFactory.h
new file mode 100644
--- /dev/null
+++ b/Engine/Plugins/AudioInsights/Source/AudioInsights/Public/Views/SoundPlotsDashboardViewFactory.h
@@ -0,0 +1,37 @@
+// Copyright Epic Games, Inc. All Rights Reserved.
+#pragma once
+
+#include "Templates/SharedPointer.h"
+#include "Views/DashboardViewFactory.h"
+#include "Views/SoundPlotsWidgetView.h"
+
+#define UE_API AUDIOINSIGHTS_API
+
+namespace UE::Audio::Insights
+{
+	class FSoundDashboardViewFactory;
+	class IDashboardDataTreeViewEntry;
+
+	class FSoundPlotsDashboardViewFactory : public IDashboardViewFactory, public TSharedFromThis<FSoundPlotsDashboardViewFactory>
+	{
+	public:
+		UE_API virtual ~FSoundPlotsDashboardViewFactory() = default;
+
+		UE_API virtual FName GetName() const override;
+		UE_API virtual FText GetDisplayName() const override;
+		UE_API virtual EDefaultDashboardTabStack GetDefaultTabStack() const override;
+		UE_API virtual FSlateIcon GetIcon() const override;
+		UE_API virtual TSharedRef<SWidget> MakeWidget(TSharedRef<SDockTab> OwnerTab, const FSpawnTabArgs& SpawnTabArgs) override;
+
+		UE_API void InitPlots(const TSharedRef<FSoundDashboardViewFactory> SoundsDashboard);
+
+	private:
+		TSharedPtr<FSoundPlotsWidgetView> PlotsWidget;
+
+		void ProcessPlotData(const TArray<TSharedPtr<IDashboardDataTreeViewEntry>>& DataViewEntries);
+
+		FDelegateHandle OnProcessPlotData;
+	};
+} // namespace UE::Audio::Insights
+
+#undef UE_API
\ No newline at end of file
diff --git a/Engine/Plugins/AudioInsights/Source/AudioInsights/Public/Views/SoundPlotsWidgetView.h b/Engine/Plugins/AudioInsights/Source/AudioInsights/Public/Views/SoundPlotsWidgetView.h
new file mode 100644
--- /dev/null
+++ b/Engine/Plugins/AudioInsights/Source/AudioInsights/Public/Views/SoundPlotsWidgetView.h
@@ -0,0 +1,109 @@
+// Copyright Epic Games, Inc. All Rights Reserved.
+#pragma once
+
+#include "AudioInsightsDataSource.h"
+#include "Messages/MixerSourceTraceMessages.h"
+#include "Templates/Function.h"
+#include "Views/SAudioCurveView.h"
+
+#define UE_API AUDIOINSIGHTS_API
+
+namespace UE::Audio::Insights
+{
+	class FSoundPlotsWidgetView : public TSharedFromThis<FSoundPlotsWidgetView>
+	{
+	public:
+		// Main information about what data is being plotted in this widget, along with functions to gather information.
+		// This data is passed into the constructor of this widget in a TMap of Column name -> Plot Column Info
+		struct FPlotColumnInfo
+		{
+			const FText ColumnDisplayName;
+			const TFunction<const ::Audio::TCircularAudioBuffer<FDataPoint>& (const IDashboardDataTreeViewEntry& InData)> DataFunc;
+			const TFunction<float(const float)> UnitConversionFunc;
+			const TFunction<const FNumberFormattingOptions*()> GetFormatOptionsFunc;
+			const TFunction<FVector2f()> GetDefaultYDataRangeFunc;
+		};
+
+		UE_API FSoundPlotsWidgetView(TMap<FName, FSoundPlotsWidgetView::FPlotColumnInfo>&& InColumnInfo, const TFunctionRef<bool(const IDashboardDataTreeViewEntry& InEntry)> InIsPlotEnabledForEntryFunc);
+		UE_API virtual ~FSoundPlotsWidgetView();
+
+		UE_API TSharedRef<SWidget> MakeWidget();
+		UE_API void ProcessPlotData(const TArray<TSharedPtr<IDashboardDataTreeViewEntry>>& DataViewEntries);
+
+	private:
+		using FPlotCurvePoint = SAudioCurveView::FCurvePoint;
+		using FPointDataPerCurveMap = TMap<int32, TArray<FPlotCurvePoint>>; // Map of source id to data point array 
+		using FLatestTimestampPerCurveMap = TMap<int32, double>; // Map of source id to most recent timestamp associated with any data
+		using FPlotCurveMetadata = SAudioCurveView::FCurveMetadata;
+
+		void OnAnalysisStarting(const double Timestamp);
+#if WITH_EDITOR
+		void OnPIEStarted(bool bSimulating);
+		void OnPIEStopped(bool bSimulating);
+		void OnPIEPaused(bool bSimulating);
+		void OnPIEResumed(bool bSimulating);
+#else
+		void OnAudioInsightsComponentTabSpawn();
+		void OnSessionAnalysisCompleted();
+		void OnTimingViewTimeMarkerChanged(double InTimeMarker);
+#endif // WITH_EDITOR
+
+		TSharedRef<SWidget> MakePlotsWidget();
+
+		void CollectDataRecursive(const TSharedPtr<IDashboardDataTreeViewEntry>& Entry, bool& bOutHasNewMetadata);
+		void ResetPlots();
+
+		// Column information used by plot widgets, keyed by column name. These keys should be a subset of the keys in GetColumns(). 
+		const TFunction<const ::Audio::TCircularAudioBuffer<FDataPoint>& (const IDashboardDataTreeViewEntry& InEntry)> GetPlotColumnDataFunc(const FName& ColumnName);
+		const TFunction<float (const float)> GetPlotColumnUnitConversionFunc(const FName& ColumnName);
+		const FNumberFormattingOptions* GetPlotColumnNumberFormat(const FName& ColumnName);
+		const FVector2f GetPlotColumnDefaultYRange(const FName& ColumnName);
+		const FText GetPlotColumnDisplayName(const FName& ColumnName);
+
+		TRange<double> GetViewRange();
+
+		// Maximum amount of data history kept for plots (in seconds)
+		inline static constexpr double MaxPlotHistorySeconds = 5.0;
+
+		// Curve points per timestamp per source id per column name 
+		TMap<FName, TSharedPtr<FPointDataPerCurveMap>> PlotWidgetCurveIdToPointDataMapPerColumn;
+
+		// Map to track the latest timestamp for each column's data points.
+		// Used to avoid adding duplicate data to the plot if this widget view is updated at a higher rate than the data is received
+		TMap<FName, TSharedPtr<FLatestTimestampPerCurveMap>> ColumnNameToLatestTimestamp;
+
+		// SourceId to metadata for the corresponding curve
+		TSharedPtr<TMap<int32, FPlotCurveMetadata>> PlotWidgetMetadataPerCurve;
+
+		// Column names for plot selector widget 
+		TArray<FName> ColumnNames;
+
+		// Func to access column data from the dashboard using this Plots view
+		const TMap<FName, FSoundPlotsWidgetView::FPlotColumnInfo> ColumnInfo;
+
+		// Func to enquire whether plotting is enabled for an entry
+		const TFunctionRef<bool(const IDashboardDataTreeViewEntry& InEntry)> IsPlotEnabledForEntryFunc;
+
+		double BeginTimestamp = TNumericLimits<double>::Max();
+		double CurrentTimestamp = TNumericLimits<double>::Lowest();
+
+		enum class EGameState : uint8
+		{
+			Running,
+			Stopped,
+			Paused
+		};
+
+		EGameState GameState = EGameState::Stopped;
+
+		TSharedPtr<SAudioCurveView> PlotWidget;
+		FName SelectedPlotColumnName;
+
+#if !WITH_EDITOR
+		double PreviousTime = 0.0;
+		double CurrentRangeUpperBound = 0.0;
+#endif // !WITH_EDITOR
+	};
+} // namespace UE::Audio::Insights
+
+#undef UE_API
diff --git a/Engine/Plugins/AudioInsights/Source/AudioInsights/Public/Views/TreeDashboardViewFactory.h b/Engine/Plugins/AudioInsights/Source/AudioInsights/Public/Views/TreeDashboardViewFactory.h
--- a/Engine/Plugins/AudioInsights/Source/AudioInsights/Public/Views/TreeDashboardViewFactory.h
+++ b/Engine/Plugins/AudioInsights/Source/AudioInsights/Public/Views/TreeDashboardViewFactory.h
@@ -26,8 +26,6 @@ namespace UE::Audio::Insights
 
 		virtual TObjectPtr<UObject> GetObject() = 0;
 		virtual const TObjectPtr<UObject> GetObject() const = 0;
-
-		virtual const FText& GetDisplayName() const = 0;
 	};
 
 	/** 
@@ -106,6 +104,7 @@ namespace UE::Audio::Insights
 
 		virtual void ProcessEntries(EProcessReason Reason) = 0;
 		virtual void SortTable() = 0;
+		virtual bool IsColumnSortable(const FName& ColumnId) const { return true; }
 
 		virtual TArray<TSharedPtr<IDashboardDataTreeViewEntry>>& GetTreeItemsSource() { return DataViewEntries; }
 
diff --git a/Engine/Plugins/AudioInsights/Source/AudioInsightsEditor/Private/AudioInsightsEditorDashboardFactory.cpp b/Engine/Plugins/AudioInsights/Source/AudioInsightsEditor/Private/AudioInsightsEditorDashboardFactory.cpp
--- a/Engine/Plugins/AudioInsights/Source/AudioInsightsEditor/Private/AudioInsightsEditorDashboardFactory.cpp
+++ b/Engine/Plugins/AudioInsights/Source/AudioInsightsEditor/Private/AudioInsightsEditorDashboardFactory.cpp
@@ -348,7 +348,9 @@ namespace UE::Audio::Insights
 					);
 
 					if (const EDefaultDashboardTabStack DefaultTabStack = DashboardViewFactories[FactoryName]->GetDefaultTabStack();
-						DefaultTabStack == EDefaultDashboardTabStack::Log || DefaultTabStack == EDefaultDashboardTabStack::AudioMeters)
+						DefaultTabStack == EDefaultDashboardTabStack::Log 
+						|| DefaultTabStack == EDefaultDashboardTabStack::AudioMeters
+						|| DefaultTabStack == EDefaultDashboardTabStack::Plots)
 					{
 						MenuBuilder.AddMenuSeparator();
 					}
@@ -571,6 +573,7 @@ namespace UE::Audio::Insights
 		TSharedRef<FTabManager::FStack> LogTabStack = FTabManager::NewStack();
 		TSharedRef<FTabManager::FStack> AnalysisTabStack = FTabManager::NewStack();
 		TSharedRef<FTabManager::FStack> AudioMetersTabStack = FTabManager::NewStack();
+		TSharedRef<FTabManager::FStack> PlotsTabStack = FTabManager::NewStack();
 		TSharedRef<FTabManager::FStack> AudioAnalyzerRackTabStack = FTabManager::NewStack()->SetHideTabWell(true)->SetSizeCoefficient(0.15f);
 
 		for (const auto& KVP : DashboardViewFactories)
@@ -605,6 +608,12 @@ namespace UE::Audio::Insights
 				}
 				break;
 
+				case EDefaultDashboardTabStack::Plots:
+				{
+					PlotsTabStack->AddTab(FactoryName, ETabState::OpenedTab);
+				}
+				break;
+
 				case EDefaultDashboardTabStack::AudioAnalyzerRack:
 				{
 					AudioAnalyzerRackTabStack->AddTab(FactoryName, ETabState::OpenedTab);
@@ -667,8 +676,21 @@ namespace UE::Audio::Insights
 					// Bottom
 					->Split
 					(
-						AudioMetersTabStack
-						->SetSizeCoefficient(0.42f)
+						// Top
+						FTabManager::NewSplitter()
+						->SetOrientation(Orient_Horizontal)
+						->SetSizeCoefficient(0.42f) // Row width
+						->Split
+						(
+							PlotsTabStack
+							->SetSizeCoefficient(0.5f)
+						)
+						// Bottom
+						->Split
+						(
+							AudioMetersTabStack
+							->SetSizeCoefficient(0.5f)
+						)
 					)
 				)
 				// Right column
diff --git a/Engine/Plugins/AudioInsights/Source/AudioInsightsEditor/Private/AudioInsightsEditorModule.cpp b/Engine/Plugins/AudioInsights/Source/AudioInsightsEditor/Private/AudioInsightsEditorModule.cpp
--- a/Engine/Plugins/AudioInsights/Source/AudioInsightsEditor/Private/AudioInsightsEditorModule.cpp
+++ b/Engine/Plugins/AudioInsights/Source/AudioInsightsEditor/Private/AudioInsightsEditorModule.cpp
@@ -14,6 +14,7 @@
 #include "Views/LogDashboardViewFactory.h"
 #include "Views/MixerSourceDashboardViewFactory.h"
 #include "Views/SoundDashboardViewFactory.h"
+#include "Views/SoundPlotsDashboardViewFactory.h"
 #include "Views/SubmixesDashboardViewFactory.h"
 #include "Views/ViewportDashboardViewFactory.h"
 #include "Views/VirtualLoopDashboardViewFactory.h"
@@ -35,16 +36,22 @@ namespace UE::Audio::Insights
 
 			DashboardFactory = MakeShared<FEditorDashboardFactory>();
 			
+			TSharedRef<FSoundDashboardViewFactory> SoundsDashboard = MakeShared<FSoundDashboardViewFactory>();
+			TSharedRef<FSoundPlotsDashboardViewFactory> PlotsDashboard = MakeShared<FSoundPlotsDashboardViewFactory>();
+
 			// @TODO UE-274216: Decide what to do with the Viewport dashboard
 			//DashboardFactory->RegisterViewFactory(MakeShared<FViewportDashboardViewFactory>());
 			DashboardFactory->RegisterViewFactory(MakeShared<FLogDashboardViewFactory>());
-			DashboardFactory->RegisterViewFactory(MakeShared<FSoundDashboardViewFactory>());
+			DashboardFactory->RegisterViewFactory(SoundsDashboard);
 			DashboardFactory->RegisterViewFactory(MakeShared<FMixerSourceDashboardViewFactory>());
 			DashboardFactory->RegisterViewFactory(MakeShared<FVirtualLoopDashboardViewFactory>());
 			DashboardFactory->RegisterViewFactory(MakeShared<FSubmixesDashboardViewFactory>());
 			DashboardFactory->RegisterViewFactory(MakeShared<FAudioBusesDashboardViewFactory>());
 			DashboardFactory->RegisterViewFactory(MakeShared<FAudioMetersDashboardViewFactory>());
+			DashboardFactory->RegisterViewFactory(PlotsDashboard);
 			DashboardFactory->RegisterViewFactory(MakeShared<FAudioAnalyzerRackDashboardViewFactory>());
+
+			PlotsDashboard->InitPlots(SoundsDashboard);
 		}
 	}
 
