[Audio Insights] Remove the Mixer Source dashboard from the Audio Insights tabs

diff --git a/Engine/Plugins/AudioInsights/Source/AudioInsights/Private/AudioInsightsModule.cpp b/Engine/Plugins/AudioInsights/Source/AudioInsights/Private/AudioInsightsModule.cpp
--- a/Engine/Plugins/AudioInsights/Source/AudioInsights/Private/AudioInsightsModule.cpp
+++ b/Engine/Plugins/AudioInsights/Source/AudioInsights/Private/AudioInsightsModule.cpp
@@ -14,7 +14,6 @@
 
 #if !WITH_EDITOR
 #include "AudioInsightsComponent.h"
-#include "Views/MixerSourceDashboardViewFactory.h"
 #include "Views/SoundDashboardViewFactory.h"
 #include "Views/SoundPlotsDashboardViewFactory.h"
 #include "Views/VirtualLoopDashboardViewFactory.h"
@@ -48,7 +47,6 @@ namespace UE::Audio::Insights
 			TSharedRef<FSoundPlotsDashboardViewFactory> PlotsDashboard = MakeShared<FSoundPlotsDashboardViewFactory>();
 
 			DashboardFactory->RegisterViewFactory(SoundsDashboard);
-			DashboardFactory->RegisterViewFactory(MakeShared<FMixerSourceDashboardViewFactory>());
 			DashboardFactory->RegisterViewFactory(MakeShared<FVirtualLoopDashboardViewFactory>());
 			DashboardFactory->RegisterViewFactory(PlotsDashboard);
 
diff --git a/Engine/Plugins/AudioInsights/Source/AudioInsightsEditor/Private/AudioInsightsEditorModule.cpp b/Engine/Plugins/AudioInsights/Source/AudioInsightsEditor/Private/AudioInsightsEditorModule.cpp
--- a/Engine/Plugins/AudioInsights/Source/AudioInsightsEditor/Private/AudioInsightsEditorModule.cpp
+++ b/Engine/Plugins/AudioInsights/Source/AudioInsightsEditor/Private/AudioInsightsEditorModule.cpp
@@ -12,7 +12,6 @@
 #include "Views/AudioBusesDashboardViewFactory.h"
 #include "Views/AudioMetersDashboardViewFactory.h"
 #include "Views/LogDashboardViewFactory.h"
-#include "Views/MixerSourceDashboardViewFactory.h"
 #include "Views/SoundDashboardViewFactory.h"
 #include "Views/SoundPlotsDashboardViewFactory.h"
 #include "Views/SubmixesDashboardViewFactory.h"
@@ -43,7 +42,6 @@ namespace UE::Audio::Insights
 			//DashboardFactory->RegisterViewFactory(MakeShared<FViewportDashboardViewFactory>());
 			DashboardFactory->RegisterViewFactory(MakeShared<FLogDashboardViewFactory>());
 			DashboardFactory->RegisterViewFactory(SoundsDashboard);
-			DashboardFactory->RegisterViewFactory(MakeShared<FMixerSourceDashboardViewFactory>());
 			DashboardFactory->RegisterViewFactory(MakeShared<FVirtualLoopDashboardViewFactory>());
 			DashboardFactory->RegisterViewFactory(MakeShared<FSubmixesDashboardViewFactory>());
 			DashboardFactory->RegisterViewFactory(MakeShared<FAudioBusesDashboardViewFactory>());
