[Audio Insights] Enable Audio Insights when Editor traces are opened in Unreal Insights

diff --git a/Engine/Plugins/AudioInsights/Source/AudioInsights/Private/AudioInsightsComponent.cpp b/Engine/Plugins/AudioInsights/Source/AudioInsights/Private/AudioInsightsComponent.cpp
--- a/Engine/Plugins/AudioInsights/Source/AudioInsights/Private/AudioInsightsComponent.cpp
+++ b/Engine/Plugins/AudioInsights/Source/AudioInsights/Private/AudioInsightsComponent.cpp
@@ -1,4 +1,5 @@
 // Copyright Epic Games, Inc. All Rights Reserved.
+#if !WITH_EDITOR
 #include "AudioInsightsComponent.h"
 
 #include "AudioInsightsModule.h"
@@ -72,10 +73,8 @@ namespace UE::Audio::Insights
 
 		FTSTicker::GetCoreTicker().RemoveTicker(OnTickHandle);
 
-#if !WITH_EDITOR
 		IAudioInsightsTraceModule& TraceModule = FAudioInsightsModule::GetChecked().GetTraceModule();
 		TraceModule.ResetTicker();
-#endif // !WITH_EDITOR
 
 		Instance.Reset();
 	}
@@ -93,7 +92,7 @@ namespace UE::Audio::Insights
 				FOnSpawnTab::CreateRaw(this,  &FAudioInsightsComponent::SpawnTab), 
 				FCanSpawnTab::CreateRaw(this, &FAudioInsightsComponent::CanSpawnTab))
 				.SetDisplayName(Config.TabLabel.IsSet()   ? Config.TabLabel.GetValue()   : LOCTEXT("AudioInsights_TabTitle", "Audio Insights"))
-				.SetTooltipText(Config.TabTooltip.IsSet() ? Config.TabTooltip.GetValue() : LOCTEXT("AudioInsights_TooltipText", "Open the Audio Insights tab (Only available for standalone game traces)."))
+				.SetTooltipText(Config.TabTooltip.IsSet() ? Config.TabTooltip.GetValue() : LOCTEXT("AudioInsights_TooltipText", "Open the Audio Insights tab."))
 				.SetIcon(Config.TabIcon.IsSet() ? Config.TabIcon.GetValue() : FSlateStyle::Get().CreateIcon("AudioInsights.Icon.Submix"));
 
 			const TSharedRef<FWorkspaceItem>* FoundWorkspace = FGlobalTabmanager::Get()->GetLocalWorkspaceMenuRoot()->GetChildItems().FindByPredicate(
@@ -121,10 +120,9 @@ namespace UE::Audio::Insights
 
 	TSharedRef<SDockTab> FAudioInsightsComponent::SpawnTab(const FSpawnTabArgs& Args)
 	{
-#if !WITH_EDITOR
 		IAudioInsightsTraceModule& TraceModule = FAudioInsightsModule::GetChecked().GetTraceModule();
 		TraceModule.RequestChannelUpdate();
-#endif // !WITH_EDITOR
+
 		const TSharedRef<SDockTab> DockTab = FAudioInsightsModule::GetChecked().CreateDashboardTabWidget(Args);
 
 		OnTabSpawn.Broadcast();
@@ -187,17 +185,11 @@ namespace UE::Audio::Insights
 
 					bIsEditorTrace = TraceServicesSessionInfo.TargetType == EBuildTargetType::Editor;
 
-#if !WITH_EDITOR
 					IAudioInsightsTraceModule& TraceModule = FAudioInsightsModule::GetChecked().GetTraceModule();
 					TraceModule.InitializeSessionInfo(TraceServicesSessionInfo);
-#endif // !WITH_EDITOR
 				}
 
-				// Allow to spawn tab if non-editor trace
-				if (!bIsEditorTrace)
-				{
-					bCanSpawnTab = true;
-				}
+				bCanSpawnTab = true;
 			}
 		}
 
@@ -206,3 +198,5 @@ namespace UE::Audio::Insights
 } // namespace UE::Audio::Insights
 
 #undef LOCTEXT_NAMESPACE
+
+#endif // !WITH_EDITOR
diff --git a/Engine/Plugins/AudioInsights/Source/AudioInsights/Private/AudioInsightsComponent.h b/Engine/Plugins/AudioInsights/Source/AudioInsights/Private/AudioInsightsComponent.h
--- a/Engine/Plugins/AudioInsights/Source/AudioInsights/Private/AudioInsightsComponent.h
+++ b/Engine/Plugins/AudioInsights/Source/AudioInsights/Private/AudioInsightsComponent.h
@@ -2,6 +2,7 @@
 
 #pragma once
 
+#if !WITH_EDITOR
 #include "Containers/Ticker.h"
 #include "Insights/IUnrealInsightsModule.h"
 
@@ -52,3 +53,4 @@ namespace UE::Audio::Insights
 		inline static TSharedPtr<FAudioInsightsComponent> Instance;
 	};
 } // namespace UE::Audio::Insights
+#endif // !WITH_EDITOR
\ No newline at end of file
diff --git a/Engine/Plugins/AudioInsights/Source/AudioInsights/Private/Providers/MixerSourceTraceProvider.cpp b/Engine/Plugins/AudioInsights/Source/AudioInsights/Private/Providers/MixerSourceTraceProvider.cpp
--- a/Engine/Plugins/AudioInsights/Source/AudioInsights/Private/Providers/MixerSourceTraceProvider.cpp
+++ b/Engine/Plugins/AudioInsights/Source/AudioInsights/Private/Providers/MixerSourceTraceProvider.cpp
@@ -93,8 +93,7 @@ namespace UE::Audio::Insights
 			}
 		}
 
-		// For now we only retrieve information from AudioDeviceId 1 (main device in standalone games)
-		const FDeviceData* DeviceData = DeviceDataMap.Find(1);
+		const FDeviceData* DeviceData = FindFilteredDeviceData();
 		if (DeviceData)
 		{
 			// Collect messages (volume, pitch, etc.) from active sounds (based on active sounds's PlayOrder)
diff --git a/Engine/Plugins/AudioInsights/Source/AudioInsights/Private/Providers/SoundTraceProvider.cpp b/Engine/Plugins/AudioInsights/Source/AudioInsights/Private/Providers/SoundTraceProvider.cpp
--- a/Engine/Plugins/AudioInsights/Source/AudioInsights/Private/Providers/SoundTraceProvider.cpp
+++ b/Engine/Plugins/AudioInsights/Source/AudioInsights/Private/Providers/SoundTraceProvider.cpp
@@ -1085,8 +1085,7 @@ namespace UE::Audio::Insights
 			}
 		}
 		
-		// For now we only retrieve information from AudioDeviceId 1 (main device in standalone games)
-		const FDeviceData* DeviceData = DeviceDataMap.Find(1);
+		const FDeviceData* DeviceData = FindFilteredDeviceData();
 		if (DeviceData)
 		{
 			// Collect messages (volume, pitch, etc.) from active sounds (based on active sounds's PlayOrder)
diff --git a/Engine/Plugins/AudioInsights/Source/AudioInsights/Private/Providers/VirtualLoopTraceProvider.cpp b/Engine/Plugins/AudioInsights/Source/AudioInsights/Private/Providers/VirtualLoopTraceProvider.cpp
--- a/Engine/Plugins/AudioInsights/Source/AudioInsights/Private/Providers/VirtualLoopTraceProvider.cpp
+++ b/Engine/Plugins/AudioInsights/Source/AudioInsights/Private/Providers/VirtualLoopTraceProvider.cpp
@@ -91,8 +91,7 @@ namespace UE::Audio::Insights
 			}
 		}
 
-		// For now we only retrieve information from AudioDeviceId 1 (main device in standalone games)
-		const FDeviceData* DeviceData = DeviceDataMap.Find(1);
+		const FDeviceData* DeviceData = FindFilteredDeviceData();
 		if (DeviceData)
 		{
 			// Collect update messages from virtualized sounds (based on active sounds's PlayOrder)
diff --git a/Engine/Plugins/AudioInsights/Source/AudioInsights/Public/AudioInsightsTraceProviderBase.h b/Engine/Plugins/AudioInsights/Source/AudioInsights/Public/AudioInsightsTraceProviderBase.h
--- a/Engine/Plugins/AudioInsights/Source/AudioInsights/Public/AudioInsightsTraceProviderBase.h
+++ b/Engine/Plugins/AudioInsights/Source/AudioInsights/Public/AudioInsightsTraceProviderBase.h
@@ -166,8 +166,7 @@ namespace UE::Audio::Insights
 			const IAudioInsightsModule& AudioInsightsModule = IAudioInsightsModule::GetEditorChecked();
 			const ::Audio::FDeviceId AudioDeviceId = AudioInsightsModule.GetDeviceId();
 #else
-			const IAudioInsightsModule& AudioInsightsModule = IAudioInsightsModule::GetChecked();
-			const ::Audio::FDeviceId AudioDeviceId = AudioInsightsModule.GetDeviceId();
+			const ::Audio::FDeviceId AudioDeviceId = GetAudioDeviceIdFromDeviceDataMap();
 #endif // WITH_EDITOR
 
 			return DeviceDataMap.Find(AudioDeviceId);
@@ -229,8 +228,7 @@ namespace UE::Audio::Insights
 			const IAudioInsightsModule& AudioInsightsModule = IAudioInsightsModule::GetEditorChecked();
 			const ::Audio::FDeviceId AudioDeviceId = AudioInsightsModule.GetDeviceId();
 #else
-			const IAudioInsightsModule& AudioInsightsModule = IAudioInsightsModule::GetChecked();
-			const ::Audio::FDeviceId AudioDeviceId = AudioInsightsModule.GetDeviceId();
+			const ::Audio::FDeviceId AudioDeviceId = GetAudioDeviceIdFromDeviceDataMap();
 #endif // WITH_EDITOR
 
 			return DeviceDataMap.Find(AudioDeviceId);
@@ -292,7 +290,26 @@ namespace UE::Audio::Insights
 		{
 			DeviceDataMap.Remove(InDeviceId);
 		}
+#else
+		const uint32 GetAudioDeviceIdFromDeviceDataMap() const
+		{
+			// In Audio Insights for Unreal Insights we don't have access to the Audio Device Manager, 
+			// we need to extract the AudioDeviceId from the already recorded trace info stored in DeviceDataMap
+			uint32 AudioDeviceId = 1; // 1 is the default audio device id
+			
+			for (const auto& [DeviceId, DeviceData] : DeviceDataMap)
+			{
+				if (DeviceId > AudioDeviceId)
+				{
+					AudioDeviceId = DeviceId;
+				}
+			}
+
+			return AudioDeviceId;
+		}
+#endif // WITH_EDITOR
 
+#if WITH_EDITOR
 		bool bIsTraceActive = false;
 #endif // WITH_EDITOR
 
