[Audio Insights] Added icons for trace and trace snapshot buttons (that will be implemented in a subsequent CL) + icons alphabetical order adjustments

diff --git a/Engine/Plugins/AudioInsights/Source/AudioInsights/Public/AudioInsightsStyle.h b/Engine/Plugins/AudioInsights/Source/AudioInsights/Public/AudioInsightsStyle.h
--- a/Engine/Plugins/AudioInsights/Source/AudioInsights/Public/AudioInsightsStyle.h
+++ b/Engine/Plugins/AudioInsights/Source/AudioInsights/Public/AudioInsightsStyle.h
@@ -150,23 +150,24 @@ namespace UE::Audio::Insights
 
 			Set("AudioInsights.Icon", new IMAGE_BRUSH_SVG(TEXT("Icons/audio_insights_icon"), Icon16));
 			Set("AudioInsights.Icon.Dashboard", new IMAGE_BRUSH_SVG(TEXT("Icons/audio_dashboard"), Icon16));
+			Set("AudioInsights.Icon.ContentBrowser", new IMAGE_BRUSH_SVG(TEXT("Icons/content_browser"), Icon20));
 			Set("AudioInsights.Icon.Event", new IMAGE_BRUSH_SVG(TEXT("Icons/audio_event"), Icon16));
 			Set("AudioInsights.Icon.Log", new IMAGE_BRUSH_SVG(TEXT("Icons/audio_log"), Icon16));
-			Set("AudioInsights.Icon.Sources", new IMAGE_BRUSH_SVG(TEXT("Icons/audio_sources"), Icon16));
-			Set("AudioInsights.Icon.Sources.Plots", new IMAGE_BRUSH_SVG(TEXT("Icons/audio_sources_plots"), Icon20));
-			Set("AudioInsights.Icon.Submix", new IMAGE_BRUSH_SVG(TEXT("Icons/audio_submix"), Icon16));
-			Set("AudioInsights.Icon.VirtualLoop", new IMAGE_BRUSH_SVG(TEXT("Icons/audio_virtualloop"), Icon16));
-			Set("AudioInsights.Icon.Viewport", new IMAGE_BRUSH_SVG(TEXT("Icons/viewport"), Icon16));
 			Set("AudioInsights.Icon.Open", new IMAGE_BRUSH_SVG(TEXT("Icons/open"), Icon20));
-			Set("AudioInsights.Icon.ContentBrowser", new IMAGE_BRUSH_SVG(TEXT("Icons/content_browser"), Icon20));
-			Set("AudioInsights.Icon.Start.Active", new IMAGE_BRUSH_SVG(TEXT("Icons/start_active"), Icon20));
-			Set("AudioInsights.Icon.Start.Inactive", new IMAGE_BRUSH_SVG(TEXT("Icons/start_inactive"), Icon20));
 			Set("AudioInsights.Icon.Record.Active", new IMAGE_BRUSH_SVG(TEXT("Icons/record_active"), Icon20));
 			Set("AudioInsights.Icon.Record.Inactive", new IMAGE_BRUSH_SVG(TEXT("Icons/record_inactive"), Icon20));
 			Set("AudioInsights.Icon.Settings", new IMAGE_BRUSH_SVG(TEXT("Icons/settings"), Icon20));
+			Set("AudioInsights.Icon.Sources", new IMAGE_BRUSH_SVG(TEXT("Icons/audio_sources"), Icon16));
+			Set("AudioInsights.Icon.Sources.Plots", new IMAGE_BRUSH_SVG(TEXT("Icons/audio_sources_plots"), Icon20));
+			Set("AudioInsights.Icon.Start.Active", new IMAGE_BRUSH_SVG(TEXT("Icons/start_active"), Icon20));
+			Set("AudioInsights.Icon.Start.Inactive", new IMAGE_BRUSH_SVG(TEXT("Icons/start_inactive"), Icon20));
 			Set("AudioInsights.Icon.Stop.Active", new IMAGE_BRUSH_SVG(TEXT("Icons/stop_active"), Icon20));
 			Set("AudioInsights.Icon.Stop.Inactive", new IMAGE_BRUSH_SVG(TEXT("Icons/stop_inactive"), Icon20));
-
+			Set("AudioInsights.Icon.Submix", new IMAGE_BRUSH_SVG(TEXT("Icons/audio_submix"), Icon16));
+			Set("AudioInsights.Icon.Trace", new IMAGE_BRUSH_SVG(TEXT("Icons/trace"), Icon16));
+			Set("AudioInsights.Icon.Trace.Snapshot", new IMAGE_BRUSH_SVG(TEXT("Icons/trace_snapshot"), Icon16));
+			Set("AudioInsights.Icon.VirtualLoop", new IMAGE_BRUSH_SVG(TEXT("Icons/audio_virtualloop"), Icon16));
+			Set("AudioInsights.Icon.Viewport", new IMAGE_BRUSH_SVG(TEXT("Icons/viewport"), Icon16));
 			Set("AudioInsights.Thumbnail", new IMAGE_BRUSH_SVG(TEXT("Icons/audio_insights"), Icon16));
 
 
