[Audio Insights] Adjusted code to use direct tracing, this means no data is written to the disk when Audio Insights is opened (unless a utrace has been already started, then Audio Insights will join that currently active session) + added support to take trace snapshots

diff --git a/Engine/Plugins/AudioInsights/Source/AudioInsights/Private/AudioInsightsTraceModule.cpp b/Engine/Plugins/AudioInsights/Source/AudioInsights/Private/AudioInsightsTraceModule.cpp
--- a/Engine/Plugins/AudioInsights/Source/AudioInsights/Private/AudioInsightsTraceModule.cpp
+++ b/Engine/Plugins/AudioInsights/Source/AudioInsights/Private/AudioInsightsTraceModule.cpp
@@ -6,7 +6,6 @@
 #include "Insights/IUnrealInsightsModule.h"
 #include "ISessionServicesModule.h"
 #include "Misc/App.h"
-#include "ProfilingDebugging/TraceAuxiliary.h"
 #include "Providers/MixerSourceTraceProvider.h"
 #include "Providers/VirtualLoopTraceProvider.h"
 #include "Templates/SharedPointer.h"
@@ -14,6 +13,9 @@
 #include "TraceServices/Model/Channel.h"
 #include "TraceServices/Model/Diagnostics.h"
 
+#if WITH_EDITOR
+#include "AudioInsightsLog.h"
+#endif // WITH_EDITOR
 
 namespace UE::Audio::Insights
 {
@@ -38,7 +40,16 @@ namespace UE::Audio::Insights
 		, EmptyArray()
 #endif
 	{
+#if WITH_EDITOR
+		FTraceAuxiliary::OnTraceStopped.AddRaw(this, &FTraceModule::OnTraceStopped);
+#endif // WITH_EDITOR
+	}
 
+	FTraceModule::~FTraceModule()
+	{
+#if WITH_EDITOR
+		FTraceAuxiliary::OnTraceStopped.RemoveAll(this);
+#endif // WITH_EDITOR
 	}
 
 	void FTraceModule::GetModuleInfo(TraceServices::FModuleInfo& OutModuleInfo)
@@ -276,10 +287,6 @@ namespace UE::Audio::Insights
 		}
 
 #if WITH_EDITOR
-		// We don't want to stop any traces that were started prior to running Audio Insights
-		// Check the original state here
-		bStopTracingAfterAudioInsightsIsFinished = !FTraceAuxiliary::IsConnected();
-		
 		CacheCurrentlyEnabledTraceChannels();
 
 		if (bOnlyTraceAudioChannels)
@@ -287,43 +294,102 @@ namespace UE::Audio::Insights
 			DisableAllTraceChannels();
 		}
 
-		bTraceAnalysisHasStarted = EnableAudioInsightsTraceChannels();
+		bAudioTraceChannelsEnabled = EnableAudioInsightsTraceChannels();
 
-		if (!FTraceAuxiliary::IsConnected())
+		IUnrealInsightsModule& UnrealInsightsModule = FModuleManager::LoadModuleChecked<IUnrealInsightsModule>("TraceInsights");
+
+		if (FTraceAuxiliary::IsConnected())
 		{
+			// Start analysis with already connected live session
+			const FString& TraceDestination = FTraceAuxiliary::GetTraceDestinationString();
+
+			if (TraceDestination == "localhost" || TraceDestination == "127.0.0.1")
+			{
+				UnrealInsightsModule.StartAnalysisForLastLiveSession();
+			}
+		}
+		else
+		{
+			// Start analysis with a direct trace
+			const FString DirectTraceHost = "127.0.0.1";
+			const uint32 DirectTracePort  = UnrealInsightsModule.StartAnalysisWithDirectTrace();
+			const FString DirectTraceAddress = FString::Printf(TEXT("%s:%u"), *DirectTraceHost, DirectTracePort);
+
 			// Clear all buffered data and prevent data from previous recordings from leaking into the new recording
 			FTraceAuxiliary::FOptions Options;
 			Options.bExcludeTail = true;
 
-			FTraceAuxiliary::Start(FTraceAuxiliary::EConnectionType::Network, TEXT("localhost"), TEXT(""), &Options);
-		}
+			FTraceAuxiliary::Start(FTraceAuxiliary::EConnectionType::Network, *DirectTraceAddress, TEXT(""), &Options);
 
-		// Even if a trace was already active, we still need to notify Audio Insights that analysis has begun
-		IUnrealInsightsModule& InsightsModule = FModuleManager::GetModuleChecked<IUnrealInsightsModule>("TraceInsights");
-		InsightsModule.StartAnalysisForLastLiveSession();
+			bIsAudioInsightsTrace = true;
+		}
 #else
-		bTraceAnalysisHasStarted = EnableAudioInsightsTraceChannels();
+		bAudioTraceChannelsEnabled = EnableAudioInsightsTraceChannels();
 #endif // WITH_EDITOR
+
+		bTraceAnalysisHasStarted = true;
 	}
 
 	void FTraceModule::StopTraceAnalysis()
 	{
 		if (FTraceAuxiliary::IsConnected())
 		{
-			DisableAudioInsightsTraceChannels();
+			FTraceAuxiliary::Stop();
+		}
+	}
 
-			if (bStopTracingAfterAudioInsightsIsFinished)
-			{
-				FTraceAuxiliary::Stop();
-			}
+#if WITH_EDITOR
+	void FTraceModule::OnTraceStopped(FTraceAuxiliary::EConnectionType InTraceType, const FString& InTraceDestination)
+	{
+		if (bIsAudioInsightsTrace)
+		{
+			bIsAudioInsightsTrace = false;
+		}
 
+		if (bAudioTraceChannelsEnabled)
+		{
+			DisableAudioInsightsTraceChannels();
 			RestoreCachedChannels();
+			ChannelsToRestore.Empty();
+
+			bAudioTraceChannelsEnabled = false;
 		}
 
-		ChannelsToRestore.Empty();
 		bTraceAnalysisHasStarted = false;
 	}
 
+	void FTraceModule::SaveTraceSnapshot()
+	{
+		bool bSnapshotSaved = false;
+
+		IUnrealInsightsModule& UnrealInsightsModule = FModuleManager::LoadModuleChecked<IUnrealInsightsModule>("TraceInsights");
+
+		if (const UE::Trace::FStoreClient* StoreClient = UnrealInsightsModule.GetStoreClient())
+		{
+			if (const UE::Trace::FStoreClient::FStatus* Status = StoreClient->GetStatus())
+			{
+				const FString TraceStorePath = FString(Status->GetStoreDir());
+				FString SnapshotFileName = FDateTime::Now().ToString(TEXT("%Y%m%d_%H%M%S_AudioInsights.utrace"));
+
+				FString SnapshotFilePath = FPaths::Combine(TraceStorePath, SnapshotFileName);
+				FPaths::NormalizeFilename(SnapshotFilePath);
+
+				bSnapshotSaved = FTraceAuxiliary::WriteSnapshot(*SnapshotFilePath);
+
+				if (bSnapshotSaved)
+				{
+					UE_LOG(LogAudioInsights, Log, TEXT("Audio Insights snapshot has been saved in %s"), *SnapshotFilePath);
+				}
+			}
+		}
+
+		if (!bSnapshotSaved)
+		{
+			UE_LOG(LogAudioInsights, Warning, TEXT("Failed to save Audio Insights snapshot"));
+		}
+	}
+#endif // WITH_EDITOR
+
 	void FTraceModule::OnOnlyTraceAudioChannelsStateChanged(const bool bOnlyTraceAudioChannels)
 	{
 		if (!bTraceAnalysisHasStarted)
@@ -355,7 +421,7 @@ namespace UE::Audio::Insights
 		return false;
 #else
 		// When attaching to a packaged build, we need to check whether we're connected and the audio channels are enabled
-		return !bTraceAnalysisHasStarted && GetCurrentAnalysisSessionInfo() != nullptr && !GetAudioTracesAreEnabled();
+		return !bTraceAnalysisHasStarted && !bAudioTraceChannelsEnabled && GetCurrentAnalysisSessionInfo() != nullptr && !GetAudioTracesAreEnabled();
 #endif // WITH_EDITOR
 	}
 
diff --git a/Engine/Plugins/AudioInsights/Source/AudioInsights/Public/AudioInsightsTraceModule.h b/Engine/Plugins/AudioInsights/Source/AudioInsights/Public/AudioInsightsTraceModule.h
--- a/Engine/Plugins/AudioInsights/Source/AudioInsights/Public/AudioInsightsTraceModule.h
+++ b/Engine/Plugins/AudioInsights/Source/AudioInsights/Public/AudioInsightsTraceModule.h
@@ -5,6 +5,7 @@
 #include "Containers/Ticker.h"
 #include "IAudioInsightsTraceModule.h"
 #include "ITraceController.h"
+#include "ProfilingDebugging/TraceAuxiliary.h"
 #include "RewindDebuggerRuntimeInterface/IRewindDebuggerRuntimeExtension.h"
 #include "Templates/SharedPointer.h"
 #include "Trace/StoreClient.h"
@@ -25,7 +26,7 @@ namespace UE::Audio::Insights
 	{
 	public:
 		UE_API FTraceModule();
-		virtual ~FTraceModule() = default;
+		virtual ~FTraceModule();
 
 		//~ Begin TraceServices::IModule interface
 		UE_API virtual void GetModuleInfo(TraceServices::FModuleInfo& OutModuleInfo) override;
@@ -55,7 +56,10 @@ namespace UE::Audio::Insights
 		UE_API virtual void RequestChannelUpdate() override;
 		UE_API virtual void ResetTicker() override;
 		UE_API virtual bool TraceControllerIsAvailable() const override;
-#endif
+#else
+		UE_API virtual void SaveTraceSnapshot() override;
+		UE_API virtual bool IsAudioInsightsTrace() const override { return bIsAudioInsightsTrace; }
+#endif // !WITH_EDITOR
 
 		DECLARE_MULTICAST_DELEGATE_OneParam(FOnAnalysisStarting, const double /*Timestamp*/);
 		FOnAnalysisStarting OnAnalysisStarting;
@@ -77,14 +81,17 @@ namespace UE::Audio::Insights
 		UE_API void SendDiscoveryRequestToTraceController() const;
 
 		UE_API bool Tick(float DeltaTime);
-#endif
+#else
+		UE_API void OnTraceStopped(FTraceAuxiliary::EConnectionType InTraceType, const FString& InTraceDestination);
+#endif // !WITH_EDITOR
 
 		TMap<FName, TSharedPtr<FTraceProviderBase>> TraceProviders;
 		TArray<FString> ChannelsToRestore;
 		FRewindDebuggerAudioInsightsRuntime RewindDebugger;
 
 		bool bTraceAnalysisHasStarted = false;
-		bool bStopTracingAfterAudioInsightsIsFinished = false;
+		bool bAudioTraceChannelsEnabled = false;
+		bool bIsAudioInsightsTrace = false;
 
 #if !WITH_EDITOR
 		FGuid InstanceID;
@@ -94,7 +101,7 @@ namespace UE::Audio::Insights
 
 		const TArray<FString> AudioChannels;
 		const TArray<FString> EmptyArray;
-#endif
+#endif // !WITH_EDITOR
 	};
 } // namespace UE::Audio::Insights
 
diff --git a/Engine/Plugins/AudioInsights/Source/AudioInsights/Public/IAudioInsightsTraceModule.h b/Engine/Plugins/AudioInsights/Source/AudioInsights/Public/IAudioInsightsTraceModule.h
--- a/Engine/Plugins/AudioInsights/Source/AudioInsights/Public/IAudioInsightsTraceModule.h
+++ b/Engine/Plugins/AudioInsights/Source/AudioInsights/Public/IAudioInsightsTraceModule.h
@@ -32,5 +32,8 @@ public:
 	virtual void RequestChannelUpdate() = 0;
 	virtual void ResetTicker() = 0;
 	virtual bool TraceControllerIsAvailable() const = 0;
+#else
+	virtual bool IsAudioInsightsTrace() const = 0;
+	virtual void SaveTraceSnapshot() = 0;
 #endif // !WITH_EDITOR
 };
diff --git a/Engine/Plugins/AudioInsights/Source/AudioInsightsEditor/Private/AudioInsightsEditorDashboardFactory.cpp b/Engine/Plugins/AudioInsights/Source/AudioInsightsEditor/Private/AudioInsightsEditorDashboardFactory.cpp
--- a/Engine/Plugins/AudioInsights/Source/AudioInsightsEditor/Private/AudioInsightsEditorDashboardFactory.cpp
+++ b/Engine/Plugins/AudioInsights/Source/AudioInsightsEditor/Private/AudioInsightsEditorDashboardFactory.cpp
@@ -11,13 +11,17 @@
 #include "Engine/World.h"
 #include "Framework/Docking/LayoutService.h"
 #include "Framework/Docking/TabManager.h"
+#include "Framework/Notifications/NotificationManager.h"
 #include "Internationalization/Text.h"
 #include "IPropertyTypeCustomization.h"
 #include "Kismet2/DebuggerCommands.h"
 #include "Templates/SharedPointer.h"
 #include "ToolMenu.h"
 #include "ToolMenus.h"
+#include "Widgets/Input/SButton.h"
 #include "Widgets/Input/SCheckBox.h"
+#include "Widgets/Layout/SSeparator.h"
+#include "Widgets/Notifications/SNotificationList.h"
 
 #define LOCTEXT_NAMESPACE "AudioInsights"
 
@@ -25,16 +29,38 @@ namespace UE::Audio::Insights
 {
 	namespace EditorDashboardFactoryPrivate
 	{
-		static const FText ToolName = LOCTEXT("AudioDashboard_ToolName", "Audio Insights");
+		static const FText ToolNameText = LOCTEXT("AudioDashboard_ToolName_Text", "Audio Insights");
 
-		static const FName MainToolbarName = "MainToolbar";
-		static const FText MainToolbarDisplayName = LOCTEXT("AudioDashboard_MainToolbarDisplayName", "Dashboard Transport");
+		static const FText PreviewDeviceText = LOCTEXT("AudioDashboard_PreviewDevice_Text", "[Preview Audio]");
 
-		static const FText PreviewDeviceDisplayName = LOCTEXT("AudioDashboard_PreviewDevice", "[Preview Audio]");
-		static const FText DashboardWorldSelectDescription = LOCTEXT("AudioDashboard_SelectWorldDescription", "Select world(s) to monitor (worlds may share audio output).");
+		static const FText MainMenuFileText = LOCTEXT("AudioDashboard_MainMenu_File_Text", "File");
+		static const FText MainMenuViewText = LOCTEXT("AudioDashboard_MainMenu_View_Text", "View");
 
-		static const FText OnlyTraceAudioChannelsName = LOCTEXT("AudioDashboard_OnlyTraceAudioChannelsDisplayName", "Only trace audio channels during PIE:");
-		static const FText OnlyTraceAudioChannelsDescription = LOCTEXT("AudioDashboard_OnlyTraceAudioChannelsDescription", "Disable all non-command line trace channels apart from Audio, Audio Mixer and CPU during PIE. This will reduce the file sizes of trace sessions while using Audio Insights.");
+		static const FText FileCloseText = LOCTEXT("AudioDashboard_File_Close_Text", "Close");
+		static const FText FileCloseTooltipText = LOCTEXT("AudioDashboard_File_Close_TooltipText", "Closes the Audio Insights dashboard.");
+
+		static const FText ViewResetLayoutText = LOCTEXT("AudioDashboard_View_ResetLayout_Text", "Reset Layout");
+
+		static const FText OnlyTraceAudioChannelsText = LOCTEXT("AudioDashboard_OnlyTraceAudioChannels_Text", "Only trace audio channels during PIE:");
+		static const FText OnlyTraceAudioChannelsTooltipText = LOCTEXT("AudioDashboard_OnlyTraceAudioChannels_TooltipText", "Disable all non-command line trace channels apart from Audio, Audio Mixer and CPU during PIE. This will reduce the file sizes of trace sessions while using Audio Insights.");
+
+		static const FText StartDirectTraceText = LOCTEXT("AudioDashboard_StartDirectTrace_Text", "Start Live Monitoring");
+		static const FText StartDirectTraceTooltipText = LOCTEXT("AudioDashboard_StartDirectTrace_TooltipText", "Starts a live monitoring, no .utrace file is recorded.");
+
+		static const FText StopDirectTraceText = LOCTEXT("AudioDashboard_StopDirectTrace_Text", "Stop Live Monitoring");
+		static const FText StopDirectTraceTooltipText = LOCTEXT("AudioDashboard_StopDirectTrace_TooltipText", "Stops the currently active Live Monitoring.");
+
+		static const FText StopLiveTraceText = LOCTEXT("AudioDashboard_StopLiveTrace_Text", "Stop Live Trace");
+		static const FText StopLiveTraceTooltipText = LOCTEXT("AudioDashboard_StopLiveTrace_TooltipText", "Stops the currently active Live Trace recording.");
+
+		static const FText SaveSnapshotText = LOCTEXT("AudioDashboard_SnapshotButton_Text", "Save Snapshot");
+		static const FText SaveSnapshotTooltipText = LOCTEXT("AudioDashboard_SnapshotButton_TooltipText", "Saves a snapshot of the current trace, the snapshots are limited to 32MB");
+
+		static const FText SaveSnapshotNotificationText = LOCTEXT("AudioDashboard_SnapshotButton_NotificationText", "Audio Insights Snapshot saved.");
+		static const FText SaveSnapshotNotificationSubText = LOCTEXT("AudioDashboard_SnapshotButton_NotificationSubText", "A snapshot .utrace with the most recent events has been saved to your trace server.");
+
+		static const FText WorldFilterText = LOCTEXT("AudioDashboard_WorldFilter_Text", "World Filter:");
+		static const FText WorldFilterTooltipText = LOCTEXT("AudioDashboard_WorldFilter_TooltipText", "Select world(s) to monitor (worlds may share audio output).");
 
 		FText GetDebugNameFromDeviceId(::Audio::FDeviceId InDeviceId)
 		{
@@ -54,7 +80,7 @@ namespace UE::Audio::Insights
 
 			if (WorldName.IsEmpty())
 			{
-				return PreviewDeviceDisplayName;
+				return PreviewDeviceText;
 			}
 
 			return FText::FromString(WorldName);
@@ -65,7 +91,7 @@ namespace UE::Audio::Insights
 	{
 		if (InDeviceId != INDEX_NONE)
 		{
-			StartTraceAnalysis(InWorld, InDeviceId);
+			SetActiveDeviceId(InDeviceId);
 		}
 
 		RefreshDeviceSelector();
@@ -76,18 +102,15 @@ namespace UE::Audio::Insights
 		RefreshDeviceSelector();
 	}
 
-	void FEditorDashboardFactory::OnPIEStarted(bool bSimulating)
+	void FEditorDashboardFactory::OnTraceStarted(FTraceAuxiliary::EConnectionType InTraceType, const FString& InTraceDestination)
 	{
 		IAudioInsightsTraceModule& TraceModule = FAudioInsightsEditorModule::GetChecked().GetTraceModule();
-		TraceModule.StartTraceAnalysis(bOnlyTraceAudioChannels);
-	}
 
-	void FEditorDashboardFactory::OnPIEStopped(bool bSimulating)
-	{
-		IAudioInsightsTraceModule& TraceModule = FAudioInsightsEditorModule::GetChecked().GetTraceModule();
-		TraceModule.StopTraceAnalysis();
-
-		RefreshDeviceSelector();
+		// If we haven't started a trace, start the analysis with the currently running trace
+		if (!TraceModule.IsAudioInsightsTrace())
+		{
+			TraceModule.StartTraceAnalysis(bOnlyTraceAudioChannels);
+		}
 	}
 
 	void FEditorDashboardFactory::OnDeviceCreated(::Audio::FDeviceId InDeviceId)
@@ -169,16 +192,10 @@ namespace UE::Audio::Insights
 			OnDeviceDestroyedHandle.Reset();
 		}
 
-		if (OnPIEStartedHandle.IsValid())
+		if (OnTraceStartedHandle.IsValid())
 		{
-			FEditorDelegates::PreBeginPIE.Remove(OnPIEStartedHandle);
-			OnPIEStartedHandle.Reset();
-		}
-
-		if (OnPIEStoppedHandle.IsValid())
-		{
-			FEditorDelegates::EndPIE.Remove(OnPIEStoppedHandle);
-			OnPIEStoppedHandle.Reset();
+			FTraceAuxiliary::OnTraceStarted.Remove(OnTraceStartedHandle);
+			OnTraceStartedHandle.Reset();
 		}
 	}
 
@@ -189,8 +206,10 @@ namespace UE::Audio::Insights
 
 	TSharedRef<SDockTab> FEditorDashboardFactory::MakeDockTabWidget(const FSpawnTabArgs& Args)
 	{
+		using namespace EditorDashboardFactoryPrivate;
+
 		const TSharedRef<SDockTab> DockTab = SNew(SDockTab)
-			.Label(EditorDashboardFactoryPrivate::ToolName)
+			.Label(ToolNameText)
 			.Clipping(EWidgetClipping::ClipToBounds)
 			.TabRole(ETabRole::NomadTab);
 
@@ -209,12 +228,14 @@ namespace UE::Audio::Insights
 		RegisterTabSpawners();
 		RefreshDeviceSelector();
 
-		if (FAudioDeviceManager* DeviceManager = FAudioDeviceManager::Get();
-			DeviceManager && GEditor && GEditor->PlayWorld)
+		if (FAudioDeviceManager* DeviceManager = FAudioDeviceManager::Get())
 		{
-			StartTraceAnalysis(GEditor->PlayWorld, DeviceManager->GetActiveAudioDevice()->DeviceID);
+			SetActiveDeviceId(DeviceManager->GetActiveAudioDevice()->DeviceID);
 		}
 
+		IAudioInsightsTraceModule& TraceModule = FAudioInsightsEditorModule::GetChecked().GetTraceModule();
+		TraceModule.StartTraceAnalysis(bOnlyTraceAudioChannels);
+
 		const TSharedRef<FTabManager::FLayout> TabLayout = LoadLayoutFromConfig();
 
 		const TSharedRef<SWidget> TabContent = SNew(SVerticalBox)
@@ -243,11 +264,11 @@ namespace UE::Audio::Insights
 
 		DockTab->SetOnTabClosed(SDockTab::FOnTabClosedCallback::CreateLambda([this](TSharedRef<SDockTab> TabClosed)
 		{
-			// If we are still in PIE, make sure we stop tracing for Audio Insights
-			if (FAudioDeviceManager* DeviceManager = FAudioDeviceManager::Get();
-				DeviceManager && GEditor && GEditor->PlayWorld)
+			// Stop trace when closing Audio Insights (only if we have started it, we don't want to stop external traces)
+			IAudioInsightsTraceModule& TraceModule = FAudioInsightsEditorModule::GetChecked().GetTraceModule();
+
+			if (TraceModule.IsTraceAnalysisActive() && TraceModule.IsAudioInsightsTrace())
 			{
-				IAudioInsightsTraceModule& TraceModule = FAudioInsightsEditorModule::GetChecked().GetTraceModule();
 				TraceModule.StopTraceAnalysis();
 			}
 
@@ -275,15 +296,18 @@ namespace UE::Audio::Insights
 
 	TSharedRef<SWidget> FEditorDashboardFactory::MakeMenuBarWidget()
 	{
+		using namespace EditorDashboardFactoryPrivate;
+
 		FMenuBarBuilder MenuBarBuilder = FMenuBarBuilder(TSharedPtr<FUICommandList>());
 
 		MenuBarBuilder.AddPullDownMenu(
-			LOCTEXT("File_MenuLabel", "File"),
+			MainMenuFileText,
 			FText::GetEmpty(),
 			FNewMenuDelegate::CreateLambda([this](FMenuBuilder& MenuBuilder)
 			{
-				MenuBuilder.AddMenuEntry(LOCTEXT("Close_MenuLabel", "Close"),
-					LOCTEXT("Close_MenuLabel_Tooltip", "Closes the Audio Insights dashboard."),
+				MenuBuilder.AddMenuEntry(
+					FileCloseText,
+					FileCloseTooltipText,
 					FSlateIcon(),
 					FUIAction(FExecuteAction::CreateLambda([this]()
 					{
@@ -301,7 +325,7 @@ namespace UE::Audio::Insights
 		);
 
 		MenuBarBuilder.AddPullDownMenu(
-			LOCTEXT("ViewMenuLabel", "View"),
+			MainMenuViewText,
 			FText::GetEmpty(),
 			FNewMenuDelegate::CreateLambda([this](FMenuBuilder& MenuBuilder)
 			{
@@ -310,7 +334,8 @@ namespace UE::Audio::Insights
 					const FName& FactoryName = KVP.Key;
 					const TSharedPtr<IDashboardViewFactory>& Factory = KVP.Value;
 
-					MenuBuilder.AddMenuEntry(Factory->GetDisplayName(),
+					MenuBuilder.AddMenuEntry(
+						Factory->GetDisplayName(),
 						FText::GetEmpty(),
 						Factory->GetIcon(),
 						FUIAction(FExecuteAction::CreateLambda([this, FactoryName]()
@@ -358,7 +383,8 @@ namespace UE::Audio::Insights
 
 				MenuBuilder.AddMenuSeparator();
 
-				MenuBuilder.AddMenuEntry(LOCTEXT("ViewMenu_ResetLayoutText", "Reset Layout"), 
+				MenuBuilder.AddMenuEntry(
+					ViewResetLayoutText,
 					FText::GetEmpty(), 
 					FSlateIcon(),
 					FUIAction(FExecuteAction::CreateLambda([this]()
@@ -427,6 +453,7 @@ namespace UE::Audio::Insights
 			.BorderImage(&TransportBackgroundColorBrush)
 			[
 				SNew(SHorizontalBox)
+				// Play World Toolbar
 				+ SHorizontalBox::Slot()
 				.HAlign(HAlign_Left)
 				.VAlign(VAlign_Center)
@@ -438,6 +465,7 @@ namespace UE::Audio::Insights
 						UToolMenus::Get()->GenerateWidget(PlayWorldToolBarName, { FPlayWorldCommands::GlobalPlayWorldActions })
 					]
 				]
+				// Only trace audio channels during PIE Label
 				+ SHorizontalBox::Slot()
 				.HAlign(HAlign_Right)
 				.VAlign(VAlign_Center)
@@ -445,9 +473,10 @@ namespace UE::Audio::Insights
 				.Padding(2.0f, 0.0f)
 				[
 					SNew(STextBlock)
-					.Text(OnlyTraceAudioChannelsName)
+					.Text(OnlyTraceAudioChannelsText)
 					.Font(IPropertyTypeCustomizationUtils::GetRegularFont())
 				]
+				// Only trace audio channels during PIE checkbox
 				+ SHorizontalBox::Slot()
 				.HAlign(HAlign_Left)
 				.VAlign(VAlign_Center)
@@ -455,7 +484,7 @@ namespace UE::Audio::Insights
 				.Padding(2.0f, 0.0f)
 				[
 					SNew(SCheckBox)
-					.ToolTipText(OnlyTraceAudioChannelsDescription)
+					.ToolTipText(OnlyTraceAudioChannelsTooltipText)
 					.IsChecked_Lambda([this]() { return bOnlyTraceAudioChannels ? ECheckBoxState::Checked : ECheckBoxState::Unchecked; })
 					.OnCheckStateChanged_Lambda([this](ECheckBoxState NewState) 
 					{ 
@@ -465,16 +494,165 @@ namespace UE::Audio::Insights
 						TraceModule.OnOnlyTraceAudioChannelsStateChanged(bOnlyTraceAudioChannels);
 					})
 				]
+				// Separator
+				+ SHorizontalBox::Slot()
+				.HAlign(HAlign_Left)
+				.VAlign(VAlign_Center)
+				.AutoWidth()
+				.Padding(12.0f, 0.0f, 0.0f, 0.0f)
+				[
+					SNew(SBox)
+					.HeightOverride(38.0f)
+					[
+						SNew(SSeparator)
+						.Orientation(Orient_Vertical)
+						.Thickness(2.0f)
+					]
+				]
+				// Start direct trace button
+				+ SHorizontalBox::Slot()
+				.HAlign(HAlign_Left)
+				.VAlign(VAlign_Center)
+				.AutoWidth()
+				.Padding(12.0f, 0.0f, 0.0f, 0.0f)
+				[
+					SNew(SButton)
+					.ButtonStyle(&FAppStyle::Get().GetWidgetStyle<FButtonStyle>("SimpleButton"))
+					.ToolTipText_Lambda([&TraceModule = FAudioInsightsEditorModule::GetChecked().GetTraceModule()]()
+					{
+						if (TraceModule.IsTraceAnalysisActive())
+						{
+							return TraceModule.IsAudioInsightsTrace() ? StopDirectTraceTooltipText : StopLiveTraceTooltipText;
+						}
+
+						return StartDirectTraceTooltipText;
+					})
+					.OnClicked_Lambda([this]()
+					{
+						IAudioInsightsTraceModule& TraceModule = FAudioInsightsEditorModule::GetChecked().GetTraceModule();
+
+						if (TraceModule.IsTraceAnalysisActive())
+						{
+							TraceModule.StopTraceAnalysis();
+						}
+						else
+						{
+							TraceModule.StartTraceAnalysis(bOnlyTraceAudioChannels);
+						}
+
+						return FReply::Handled();
+					})
+					[
+						SNew(SBox)
+						.HeightOverride(16.0f)
+						[
+							SNew(SHorizontalBox)
+							// Trace record icon
+							+ SHorizontalBox::Slot()
+							.VAlign(VAlign_Center)
+							.AutoWidth()
+							[
+								SNew(SImage)
+								.Image(FSlateStyle::Get().GetBrush("AudioInsights.Icon.Trace"))
+								.ColorAndOpacity_Lambda([&TraceModule = FAudioInsightsEditorModule::GetChecked().GetTraceModule()]()
+				 				{
+									if (TraceModule.IsTraceAnalysisActive())
+									{
+										return TraceModule.IsAudioInsightsTrace() ? FSlateColor(FColor::Green) : FSlateColor(FColor::Red);
+									}
+
+									return FSlateColor::UseForeground();
+				 				})
+							]
+							// Trace record text
+							+ SHorizontalBox::Slot()
+							.VAlign(VAlign_Center)
+							.AutoWidth()
+							.Padding(6.0f, 1.0f, 0.0f, 0.0f)
+							[
+								SNew(STextBlock)
+								.Justification(ETextJustify::Center)
+								.Text_Lambda([&TraceModule = FAudioInsightsEditorModule::GetChecked().GetTraceModule()]()
+				 				{
+									if (TraceModule.IsTraceAnalysisActive())
+									{
+										return TraceModule.IsAudioInsightsTrace() ? StopDirectTraceText : StopLiveTraceText;
+									}
+
+									return StartDirectTraceText;
+				 				})
+								.Font(IPropertyTypeCustomizationUtils::GetRegularFont())
+							]
+						]
+					]
+				]
+				// Save snapshot button
+				+ SHorizontalBox::Slot()
+				.HAlign(HAlign_Left)
+				.VAlign(VAlign_Center)
+				.AutoWidth()
+				.Padding(12.0f, 0.0f, 0.0f, 0.0f)
+				[
+					SNew(SButton)
+					.ButtonStyle(&FAppStyle::Get().GetWidgetStyle<FButtonStyle>("SimpleButton"))
+					.ToolTipText(SaveSnapshotTooltipText)
+					.OnClicked_Lambda([]()
+					{
+						// Save snapshot
+						IAudioInsightsTraceModule& TraceModule = FAudioInsightsEditorModule::GetChecked().GetTraceModule();
+						TraceModule.SaveTraceSnapshot();
+
+						// Show floating notification
+						FNotificationInfo Info(SaveSnapshotNotificationText);
+						Info.SubText = SaveSnapshotNotificationSubText;
+						Info.bFireAndForget  = true;
+						Info.FadeOutDuration = 1.0f;
+						Info.ExpireDuration  = 4.0f;
+
+						FSlateNotificationManager::Get().AddNotification(Info);
+
+						return FReply::Handled();
+					})
+					[
+						SNew(SBox)
+						.HeightOverride(16.0f)
+						[
+							SNew(SHorizontalBox)
+							// Trace snapshot icon
+							+ SHorizontalBox::Slot()
+							.VAlign(VAlign_Center)
+							.AutoWidth()
+							[
+								SNew(SImage)
+								.Image(FSlateStyle::Get().GetBrush("AudioInsights.Icon.Trace.Snapshot"))
+								.ColorAndOpacity(FSlateColor::UseForeground())
+							]
+							// Trace snapshot text
+							+ SHorizontalBox::Slot()
+							.VAlign(VAlign_Center)
+							.AutoWidth()
+							.Padding(6.0f, 1.0f, 0.0f, 0.0f)
+							[
+								SNew(STextBlock)
+								.Justification(ETextJustify::Center)
+								.Text(SaveSnapshotText)
+								.Font(IPropertyTypeCustomizationUtils::GetRegularFont())
+							]
+						]
+					]
+				]
+				// World filter label
 				+ SHorizontalBox::Slot()
 				.HAlign(HAlign_Right)
 				.VAlign(VAlign_Center)
 				.Padding(2.0f, 0.0f)
 				[
 					SNew(STextBlock)
-					.Text(LOCTEXT("SelectDashboardWorld_DisplayName", "World Filter:"))
-					.ToolTipText(DashboardWorldSelectDescription)
+					.Text(WorldFilterText)
+					.ToolTipText(WorldFilterTooltipText)
 					.Font(IPropertyTypeCustomizationUtils::GetRegularFont())
 				]
+				// World filter combobox
 				+ SHorizontalBox::Slot()
 				.HAlign(HAlign_Left)
 				.VAlign(VAlign_Center)
@@ -482,7 +660,7 @@ namespace UE::Audio::Insights
 				.Padding(2.0f, 0.0f)
 				[
 					SAssignNew(AudioDeviceComboBox, SComboBox<TSharedPtr<::Audio::FDeviceId>>)
-					.ToolTipText(DashboardWorldSelectDescription)
+					.ToolTipText(WorldFilterTooltipText)
 					.OptionsSource(&AudioDeviceIds)
 					.OnGenerateWidget_Lambda([](const TSharedPtr<::Audio::FDeviceId>& WidgetDeviceId)
 					{
@@ -513,22 +691,16 @@ namespace UE::Audio::Insights
 			];
 	}
 
-	void FEditorDashboardFactory::StartTraceAnalysis(const TObjectPtr<const UWorld> InWorld, const ::Audio::FDeviceId InDeviceId)
+	void FEditorDashboardFactory::SetActiveDeviceId(const ::Audio::FDeviceId InDeviceId)
 	{
-		if (InWorld && InWorld->IsGameWorld())
-		{
-			IAudioInsightsTraceModule& TraceModule = FAudioInsightsEditorModule::GetChecked().GetTraceModule();
-			TraceModule.StartTraceAnalysis(bOnlyTraceAudioChannels);
-
-			const TObjectPtr<const UAudioInsightsEditorSettings> AudioInsightsEditorSettings = GetDefault<UAudioInsightsEditorSettings>();
+		const TObjectPtr<const UAudioInsightsEditorSettings> AudioInsightsEditorSettings = GetDefault<UAudioInsightsEditorSettings>();
 
-			// We don't want to set ActiveDeviceId if bWorldFilterDefaultsToFirstClient is true and more than 2 PIE clients are running
-			if (AudioInsightsEditorSettings == nullptr ||
-				(AudioInsightsEditorSettings && !AudioInsightsEditorSettings->bWorldFilterDefaultsToFirstClient) ||
-				AudioDeviceIds.Num() < 2)
-			{
-				ActiveDeviceId = InDeviceId;
-			}
+		// We don't want to set ActiveDeviceId if bWorldFilterDefaultsToFirstClient is true and more than 2 PIE clients are running
+		if (AudioInsightsEditorSettings == nullptr ||
+			(AudioInsightsEditorSettings && !AudioInsightsEditorSettings->bWorldFilterDefaultsToFirstClient) ||
+			AudioDeviceIds.Num() < 2)
+		{
+			ActiveDeviceId = InDeviceId;
 		}
 	}
 
@@ -554,14 +726,9 @@ namespace UE::Audio::Insights
 			OnDeviceDestroyedHandle = FAudioDeviceManagerDelegates::OnAudioDeviceDestroyed.AddSP(this, &FEditorDashboardFactory::OnDeviceDestroyed);
 		}
 
-		if (!OnPIEStartedHandle.IsValid())
-		{
-			OnPIEStartedHandle = FEditorDelegates::PreBeginPIE.AddSP(this, &FEditorDashboardFactory::OnPIEStarted);
-		}
-
-		if (!OnPIEStoppedHandle.IsValid())
+		if (!OnTraceStartedHandle.IsValid())
 		{
-			OnPIEStoppedHandle = FEditorDelegates::EndPIE.AddSP(this, &FEditorDashboardFactory::OnPIEStopped);
+			OnTraceStartedHandle = FTraceAuxiliary::OnTraceStarted.AddSP(this, &FEditorDashboardFactory::OnTraceStarted);
 		}
 	}
 
@@ -706,7 +873,7 @@ namespace UE::Audio::Insights
 	{
 		using namespace EditorDashboardFactoryPrivate;
 
-		DashboardWorkspace = DashboardTabManager->AddLocalWorkspaceMenuCategory(ToolName);
+		DashboardWorkspace = DashboardTabManager->AddLocalWorkspaceMenuCategory(ToolNameText);
 
 		for (const auto& KVP : DashboardViewFactories)
 		{
diff --git a/Engine/Plugins/AudioInsights/Source/AudioInsightsEditor/Private/AudioInsightsEditorDashboardFactory.h b/Engine/Plugins/AudioInsights/Source/AudioInsightsEditor/Private/AudioInsightsEditorDashboardFactory.h
--- a/Engine/Plugins/AudioInsights/Source/AudioInsightsEditor/Private/AudioInsightsEditorDashboardFactory.h
+++ b/Engine/Plugins/AudioInsights/Source/AudioInsightsEditor/Private/AudioInsightsEditorDashboardFactory.h
@@ -6,6 +6,7 @@
 #include "Engine/World.h"
 #include "Framework/Docking/TabManager.h"
 #include "IAudioInsightsDashboardFactory.h"
+#include "ProfilingDebugging/TraceAuxiliary.h"
 #include "Templates/SharedPointer.h"
 #include "Views/DashboardViewFactory.h"
 #include "Widgets/Docking/SDockTab.h"
@@ -35,8 +36,7 @@ namespace UE::Audio::Insights
 		void OnDeviceCreated(::Audio::FDeviceId InDeviceId);
 		void OnDeviceDestroyed(::Audio::FDeviceId InDeviceId);
 
-		void OnPIEStarted(bool bSimulating);
-		void OnPIEStopped(bool bSimulating);
+		void OnTraceStarted(FTraceAuxiliary::EConnectionType InTraceType, const FString& InTraceDestination);
 
 		void RefreshDeviceSelector();
 		void ResetDelegates();
@@ -44,7 +44,7 @@ namespace UE::Audio::Insights
 		TSharedRef<SWidget> MakeMenuBarWidget();
 		TSharedRef<SWidget> MakeMainToolbarWidget();
 
-		void StartTraceAnalysis(const TObjectPtr<const UWorld> InWorld, const ::Audio::FDeviceId InDeviceId);
+		void SetActiveDeviceId(const ::Audio::FDeviceId InDeviceId);
 		void InitDelegates();
 		TSharedRef<FTabManager::FLayout> GetDefaultTabLayout();
 
@@ -60,8 +60,7 @@ namespace UE::Audio::Insights
 		FDelegateHandle OnDeviceCreatedHandle;
 		FDelegateHandle OnDeviceDestroyedHandle;
 
-		FDelegateHandle OnPIEStartedHandle;
-		FDelegateHandle OnPIEStoppedHandle;
+		FDelegateHandle OnTraceStartedHandle;
 
 		bool bOnlyTraceAudioChannels = false;
 
