[Audio Insights] Fix for crash for audio bus provider (pointer check) + enforce getting a reference of audio insights modules (runtime and editor) by removing copy/assign/move

diff --git a/Engine/Plugins/AudioInsights/Source/AudioInsights/Private/AudioInsightsModule.h b/Engine/Plugins/AudioInsights/Source/AudioInsights/Private/AudioInsightsModule.h
--- a/Engine/Plugins/AudioInsights/Source/AudioInsights/Private/AudioInsightsModule.h
+++ b/Engine/Plugins/AudioInsights/Source/AudioInsights/Private/AudioInsightsModule.h
@@ -25,9 +25,13 @@ namespace UE::Audio::Insights
 	{
 	public:
 		FAudioInsightsModule() = default;
-
 		virtual ~FAudioInsightsModule() = default;
 
+		FAudioInsightsModule(const FAudioInsightsModule&) = delete;
+		FAudioInsightsModule& operator=(const FAudioInsightsModule&) = delete;
+		FAudioInsightsModule(FAudioInsightsModule&&) = delete;
+		FAudioInsightsModule& operator=(FAudioInsightsModule&&) = delete;
+
 		virtual void StartupModule() override;
 		virtual void ShutdownModule() override;
 		virtual void RegisterDashboardViewFactory(TSharedRef<IDashboardViewFactory> InDashboardFactory) override;
diff --git a/Engine/Plugins/AudioInsights/Source/AudioInsightsEditor/Private/Analyzers/AudioMeterAnalyzer.cpp b/Engine/Plugins/AudioInsights/Source/AudioInsightsEditor/Private/Analyzers/AudioMeterAnalyzer.cpp
--- a/Engine/Plugins/AudioInsights/Source/AudioInsightsEditor/Private/Analyzers/AudioMeterAnalyzer.cpp
+++ b/Engine/Plugins/AudioInsights/Source/AudioInsightsEditor/Private/Analyzers/AudioMeterAnalyzer.cpp
@@ -14,8 +14,7 @@ namespace UE::Audio::Insights
 		{
 			using namespace ::Audio;
 
-			const FAudioInsightsEditorModule AudioInsightsEditorModule = FAudioInsightsEditorModule::GetChecked();
-			const FDeviceId AudioDeviceId = AudioInsightsEditorModule.GetDeviceId();
+			const FDeviceId AudioDeviceId = FAudioInsightsEditorModule::GetChecked().GetDeviceId();
 
 			if (const FAudioDeviceManager* AudioDeviceManager = FAudioDeviceManager::Get())
 			{
@@ -47,8 +46,7 @@ namespace UE::Audio::Insights
 			return;
 		}
 
-		const FAudioInsightsEditorModule AudioInsightsEditorModule = FAudioInsightsEditorModule::GetChecked();
-		const FDeviceId AudioDeviceId = AudioInsightsEditorModule.GetDeviceId();
+		const FDeviceId AudioDeviceId = FAudioInsightsEditorModule::GetChecked().GetDeviceId();
 
 		const FMixerDevice* MixerDevice = static_cast<const FMixerDevice*>(AudioDeviceManager->GetAudioDeviceRaw(AudioDeviceId));
 		if (!MixerDevice)
diff --git a/Engine/Plugins/AudioInsights/Source/AudioInsightsEditor/Private/Analyzers/AudioMeterSubmixAnalyzer.cpp b/Engine/Plugins/AudioInsights/Source/AudioInsightsEditor/Private/Analyzers/AudioMeterSubmixAnalyzer.cpp
--- a/Engine/Plugins/AudioInsights/Source/AudioInsightsEditor/Private/Analyzers/AudioMeterSubmixAnalyzer.cpp
+++ b/Engine/Plugins/AudioInsights/Source/AudioInsightsEditor/Private/Analyzers/AudioMeterSubmixAnalyzer.cpp
@@ -46,8 +46,7 @@ namespace UE::Audio::Insights
 			return;
 		}
 
-		const FAudioInsightsEditorModule AudioInsightsEditorModule = FAudioInsightsEditorModule::GetChecked();
-		const FDeviceId AudioDeviceId = AudioInsightsEditorModule.GetDeviceId();
+		const FDeviceId AudioDeviceId = FAudioInsightsEditorModule::GetChecked().GetDeviceId();
 
 		const FMixerDevice* MixerDevice = static_cast<const FMixerDevice*>(AudioDeviceManager->GetAudioDeviceRaw(AudioDeviceId));
 		if (!MixerDevice)
@@ -94,8 +93,7 @@ namespace UE::Audio::Insights
 			return;
 		}
 
-		const FAudioInsightsEditorModule AudioInsightsEditorModule = FAudioInsightsEditorModule::GetChecked();
-		const FDeviceId AudioDeviceId = AudioInsightsEditorModule.GetDeviceId();
+		const FDeviceId AudioDeviceId = FAudioInsightsEditorModule::GetChecked().GetDeviceId();
 
 		const FMixerDevice* MixerDevice = static_cast<const FMixerDevice*>(AudioDeviceManager->GetAudioDeviceRaw(AudioDeviceId));
 		if (!MixerDevice)
diff --git a/Engine/Plugins/AudioInsights/Source/AudioInsightsEditor/Private/Analyzers/SubmixAudioAnalyzerRack.cpp b/Engine/Plugins/AudioInsights/Source/AudioInsightsEditor/Private/Analyzers/SubmixAudioAnalyzerRack.cpp
--- a/Engine/Plugins/AudioInsights/Source/AudioInsightsEditor/Private/Analyzers/SubmixAudioAnalyzerRack.cpp
+++ b/Engine/Plugins/AudioInsights/Source/AudioInsightsEditor/Private/Analyzers/SubmixAudioAnalyzerRack.cpp
@@ -119,8 +119,7 @@ namespace UE::Audio::Insights
 			return;
 		}
 
-		const FAudioInsightsEditorModule& AudioInsightsEditorModule = FAudioInsightsEditorModule::GetChecked();
-		const FDeviceId AudioDeviceId = AudioInsightsEditorModule.GetDeviceId();
+		const FDeviceId AudioDeviceId = FAudioInsightsEditorModule::GetChecked().GetDeviceId();
 
 		const FMixerDevice* MixerDevice = static_cast<const FMixerDevice*>(AudioDeviceManager->GetAudioDeviceRaw(AudioDeviceId));
 		if (!MixerDevice)
@@ -177,8 +176,7 @@ namespace UE::Audio::Insights
 			return;
 		}
 
-		const FAudioInsightsEditorModule& AudioInsightsEditorModule = FAudioInsightsEditorModule::GetChecked();
-		const FDeviceId AudioDeviceId = AudioInsightsEditorModule.GetDeviceId();
+		const FDeviceId AudioDeviceId = FAudioInsightsEditorModule::GetChecked().GetDeviceId();
 
 		const FMixerDevice* MixerDevice = static_cast<const FMixerDevice*>(AudioDeviceManager->GetAudioDeviceRaw(AudioDeviceId));
 		if (!MixerDevice)
diff --git a/Engine/Plugins/AudioInsights/Source/AudioInsightsEditor/Private/AudioInsightsEditorModule.h b/Engine/Plugins/AudioInsights/Source/AudioInsightsEditor/Private/AudioInsightsEditorModule.h
--- a/Engine/Plugins/AudioInsights/Source/AudioInsightsEditor/Private/AudioInsightsEditorModule.h
+++ b/Engine/Plugins/AudioInsights/Source/AudioInsightsEditor/Private/AudioInsightsEditorModule.h
@@ -17,9 +17,13 @@ namespace UE::Audio::Insights
 	{
 	public:
 		FAudioInsightsEditorModule() = default;
-
 		virtual ~FAudioInsightsEditorModule() = default;
 
+		FAudioInsightsEditorModule(const FAudioInsightsEditorModule&) = delete;
+		FAudioInsightsEditorModule& operator=(const FAudioInsightsEditorModule&) = delete;
+		FAudioInsightsEditorModule(FAudioInsightsEditorModule&&) = delete;
+		FAudioInsightsEditorModule& operator=(FAudioInsightsEditorModule&&) = delete;
+
 		virtual void StartupModule() override;
 		virtual void ShutdownModule() override;
 		virtual void RegisterDashboardViewFactory(TSharedRef<IDashboardViewFactory> InDashboardFactory) override;
diff --git a/Engine/Plugins/AudioInsights/Source/AudioInsightsEditor/Private/Providers/AudioBusProvider.cpp b/Engine/Plugins/AudioInsights/Source/AudioInsightsEditor/Private/Providers/AudioBusProvider.cpp
--- a/Engine/Plugins/AudioInsights/Source/AudioInsightsEditor/Private/Providers/AudioBusProvider.cpp
+++ b/Engine/Plugins/AudioInsights/Source/AudioInsightsEditor/Private/Providers/AudioBusProvider.cpp
@@ -97,8 +97,7 @@ namespace UE::Audio::Insights
 
 		if (!bIsAudioBusAssetAlreadAdded)
 		{
-			const FAudioInsightsEditorModule AudioInsightsEditorModule = FAudioInsightsEditorModule::GetChecked();
-			const ::Audio::FDeviceId AudioDeviceId = AudioInsightsEditorModule.GetDeviceId();
+			const ::Audio::FDeviceId AudioDeviceId = FAudioInsightsEditorModule::GetChecked().GetDeviceId();
 
 			TSharedPtr<FAudioBusAssetDashboardEntry> AudioBusAssetDashboardEntry = MakeShared<FAudioBusAssetDashboardEntry>();
 			AudioBusAssetDashboardEntry->DeviceId   = AudioDeviceId;
@@ -126,8 +125,7 @@ namespace UE::Audio::Insights
 
 		if (FoundAudioBusAssetNameIndex != INDEX_NONE)
 		{
-			const FAudioInsightsEditorModule AudioInsightsEditorModule = FAudioInsightsEditorModule::GetChecked();
-			const ::Audio::FDeviceId AudioDeviceId = AudioInsightsEditorModule.GetDeviceId();
+			const ::Audio::FDeviceId AudioDeviceId = FAudioInsightsEditorModule::GetChecked().GetDeviceId();
 
 			RemoveDeviceEntry(AudioDeviceId, AudioBusDataViewEntries[FoundAudioBusAssetNameIndex]->AudioBus->GetUniqueID());
 
@@ -198,8 +196,7 @@ namespace UE::Audio::Insights
 		// Process messages
 		if (bAssetEntriesNeedRefreshing)
 		{
-			const FAudioInsightsEditorModule AudioInsightsEditorModule = FAudioInsightsEditorModule::GetChecked();
-			const ::Audio::FDeviceId AudioDeviceId = AudioInsightsEditorModule.GetDeviceId();
+			const ::Audio::FDeviceId AudioDeviceId = FAudioInsightsEditorModule::GetChecked().GetDeviceId();
 
 			for (const TSharedPtr<FAudioBusAssetDashboardEntry>& AudioBusDataViewEntry : AudioBusDataViewEntries)
 			{
diff --git a/Engine/Plugins/AudioInsights/Source/AudioInsightsEditor/Private/Providers/SoundSubmixProvider.cpp b/Engine/Plugins/AudioInsights/Source/AudioInsightsEditor/Private/Providers/SoundSubmixProvider.cpp
--- a/Engine/Plugins/AudioInsights/Source/AudioInsightsEditor/Private/Providers/SoundSubmixProvider.cpp
+++ b/Engine/Plugins/AudioInsights/Source/AudioInsightsEditor/Private/Providers/SoundSubmixProvider.cpp
@@ -97,8 +97,7 @@ namespace UE::Audio::Insights
 
 		if (!bIsSubmixAssetAlreadAdded)
 		{
-			const FAudioInsightsEditorModule AudioInsightsEditorModule = FAudioInsightsEditorModule::GetChecked();
-			const ::Audio::FDeviceId AudioDeviceId = AudioInsightsEditorModule.GetDeviceId();
+			const ::Audio::FDeviceId AudioDeviceId = FAudioInsightsEditorModule::GetChecked().GetDeviceId();
 
 			TSharedPtr<FSoundSubmixAssetDashboardEntry>& SoundSubmixAssetDashboardEntry = SubmixDataViewEntries.Add_GetRef(MakeShared<FSoundSubmixAssetDashboardEntry>());
 			SoundSubmixAssetDashboardEntry->DeviceId    = AudioDeviceId;
@@ -122,8 +121,7 @@ namespace UE::Audio::Insights
 
 		if (FoundSubmixAssetNameIndex != INDEX_NONE)
 		{
-			const FAudioInsightsEditorModule AudioInsightsEditorModule = FAudioInsightsEditorModule::GetChecked();
-			const ::Audio::FDeviceId AudioDeviceId = AudioInsightsEditorModule.GetDeviceId();
+			const ::Audio::FDeviceId AudioDeviceId = FAudioInsightsEditorModule::GetChecked().GetDeviceId();
 
 			const uint32 SubmixId = SubmixDataViewEntries[FoundSubmixAssetNameIndex]->SubmixId;
 
@@ -175,8 +173,7 @@ namespace UE::Audio::Insights
 		// Process messages
 		if (bAssetEntriesNeedRefreshing)
 		{
-			const FAudioInsightsEditorModule AudioInsightsEditorModule = FAudioInsightsEditorModule::GetChecked();
-			const ::Audio::FDeviceId AudioDeviceId = AudioInsightsEditorModule.GetDeviceId();
+			const ::Audio::FDeviceId AudioDeviceId = FAudioInsightsEditorModule::GetChecked().GetDeviceId();
 
 			for (const TSharedPtr<FSoundSubmixAssetDashboardEntry>& SubmixDataViewEntry : SubmixDataViewEntries)
 			{
