[Audio Insights] Delete Mixer Source dashboard code - move and rename messages to match Sounds Dashboard messages

diff --git a/Engine/Plugins/AudioInsights/Source/AudioInsights/Private/AudioInsightsTraceModule.cpp b/Engine/Plugins/AudioInsights/Source/AudioInsights/Private/AudioInsightsTraceModule.cpp
--- a/Engine/Plugins/AudioInsights/Source/AudioInsights/Private/AudioInsightsTraceModule.cpp
+++ b/Engine/Plugins/AudioInsights/Source/AudioInsights/Private/AudioInsightsTraceModule.cpp
@@ -6,7 +6,6 @@
 #include "Insights/IUnrealInsightsModule.h"
 #include "ISessionServicesModule.h"
 #include "Misc/App.h"
-#include "Providers/MixerSourceTraceProvider.h"
 #include "Providers/VirtualLoopTraceProvider.h"
 #include "Templates/SharedPointer.h"
 #include "Trace/Trace.h"
diff --git a/Engine/Plugins/AudioInsights/Source/AudioInsights/Private/Providers/MixerSourceTraceProvider.cpp b/Engine/Plugins/AudioInsights/Source/AudioInsights/Private/Providers/MixerSourceTraceProvider.cpp
deleted file mode 100644
diff --git a/Engine/Plugins/AudioInsights/Source/AudioInsights/Private/Providers/MixerSourceTraceProvider.h b/Engine/Plugins/AudioInsights/Source/AudioInsights/Private/Providers/MixerSourceTraceProvider.h
deleted file mode 100644
diff --git a/Engine/Plugins/AudioInsights/Source/AudioInsights/Private/Providers/SoundTraceProvider.cpp b/Engine/Plugins/AudioInsights/Source/AudioInsights/Private/Providers/SoundTraceProvider.cpp
--- a/Engine/Plugins/AudioInsights/Source/AudioInsights/Private/Providers/SoundTraceProvider.cpp
+++ b/Engine/Plugins/AudioInsights/Source/AudioInsights/Private/Providers/SoundTraceProvider.cpp
@@ -2,7 +2,6 @@
 #include "Providers/SoundTraceProvider.h"
 
 #include "AudioInsightsModule.h"
-#include "Messages/MixerSourceTraceMessages.h"
 #include "Settings/SoundDashboardSettings.h"
 #include "Trace/Analyzer.h"
 #include "TraceServices/Model/AnalysisSession.h"
@@ -77,7 +76,7 @@ namespace UE::Audio::Insights
 			return DataPointsDisjointedArrayView.FirstBuffer.Last();
 		}
 
-		void UpdateParameterEntry(const FMixerSourceParameterMessage& Msg, const float DataPoint, ::Audio::TCircularAudioBuffer<FDataPoint>& OutDataPoints, TArray<int32, TInlineAllocator<64>>& OutEntriesWithPoppedDataPoints)
+		void UpdateParameterEntry(const FSoundParameterMessage& Msg, const float DataPoint, ::Audio::TCircularAudioBuffer<FDataPoint>& OutDataPoints, TArray<int32, TInlineAllocator<64>>& OutEntriesWithPoppedDataPoints)
 		{
 			if (!OutEntriesWithPoppedDataPoints.Contains(Msg.PlayOrder))
 			{
@@ -209,12 +208,12 @@ namespace UE::Audio::Insights
 				Builder.RouteEvent(RouteId_Stop,				"Audio", "SoundStop");
 				Builder.RouteEvent(RouteId_Priority,			"Audio", "SoundPriority");
 				Builder.RouteEvent(RouteId_Distance,			"Audio", "SoundDistance");
-				Builder.RouteEvent(RouteId_DistanceAttenuation,	"Audio", "MixerSourceDistanceAttenuation");
-				Builder.RouteEvent(RouteId_Filters,				"Audio", "MixerSourceFilters");
-				Builder.RouteEvent(RouteId_Amplitude,			"Audio", "MixerSourceEnvelope");
-				Builder.RouteEvent(RouteId_Volume,				"Audio", "MixerSourceVolume");
-				Builder.RouteEvent(RouteId_Pitch,				"Audio", "MixerSourcePitch");
-				Builder.RouteEvent(RouteId_RelativeRenderCost,	"Audio", "MixerSourceRelativeRenderCost");
+				Builder.RouteEvent(RouteId_DistanceAttenuation,	"Audio", "SoundDistanceAttenuation");
+				Builder.RouteEvent(RouteId_Filters,				"Audio", "SoundSourceFilters");
+				Builder.RouteEvent(RouteId_Amplitude,			"Audio", "SoundSourceEnvelope");
+				Builder.RouteEvent(RouteId_Volume,				"Audio", "SoundVolume");
+				Builder.RouteEvent(RouteId_Pitch,				"Audio", "SoundPitch");
+				Builder.RouteEvent(RouteId_RelativeRenderCost,	"Audio", "SoundRelativeRenderCost");
 			}
 
 			virtual bool OnEvent(uint16 RouteId, EStyle Style, const FOnEventContext& Context) override
@@ -257,38 +256,38 @@ namespace UE::Audio::Insights
 
 					case RouteId_DistanceAttenuation:
 					{
-						Messages.DistanceAttenuationMessages.Enqueue(FMixerSourceDistanceAttenuationMessage{ Context });
+						Messages.DistanceAttenuationMessages.Enqueue(FSoundDistanceAttenuationMessage{ Context });
 						break;
 					}
 
 					case RouteId_Filters:
 					{
-						Messages.LPFFreqMessages.Enqueue(FMixerSourceLPFFreqMessage{ Context });
-						Messages.HPFFreqMessages.Enqueue(FMixerSourceHPFFreqMessage{ Context });
+						Messages.LPFFreqMessages.Enqueue(FSoundLPFFreqMessage{ Context });
+						Messages.HPFFreqMessages.Enqueue(FSoundHPFFreqMessage{ Context });
 						break;
 					}
 
 					case RouteId_Amplitude:
 					{
-						Messages.AmplitudeMessages.Enqueue(FMixerSourceEnvelopeMessage{ Context });
+						Messages.AmplitudeMessages.Enqueue(FSoundEnvelopeMessage{ Context });
 						break;
 					}
 
 					case RouteId_Volume:
 					{
-						Messages.VolumeMessages.Enqueue(FMixerSourceVolumeMessage{ Context });
+						Messages.VolumeMessages.Enqueue(FSoundVolumeMessage{ Context });
 						break;
 					}
 
 					case RouteId_Pitch:
 					{
-						Messages.PitchMessages.Enqueue(FMixerSourcePitchMessage{ Context });
+						Messages.PitchMessages.Enqueue(FSoundPitchMessage{ Context });
 						break;
 					}
 
 					case RouteId_RelativeRenderCost:
 					{
-						Messages.RelativeRenderCostMessages.Enqueue(FMixerSourceRelativeRenderCostMessage{ Context });
+						Messages.RelativeRenderCostMessages.Enqueue(FSoundRelativeRenderCostMessage{ Context });
 						break;
 					}
 
@@ -363,7 +362,7 @@ namespace UE::Audio::Insights
 		};
 
 		TSharedPtr<FSoundDashboardEntry> ParamReturn = nullptr;
-		auto GetActiveSoundEntryFromParameterMessage = [this, &ParamReturn](const FMixerSourceParameterMessage& Msg) -> TSharedPtr<FSoundDashboardEntry>*
+		auto GetActiveSoundEntryFromParameterMessage = [this, &ParamReturn](const FSoundParameterMessage& Msg) -> TSharedPtr<FSoundDashboardEntry>*
 		{
 			ParamReturn = nullptr;
 
@@ -484,8 +483,8 @@ namespace UE::Audio::Insights
 		});
 		EntriesWithPoppedDataPoints.Reset();
 
-		ProcessMessageQueue<FMixerSourceDistanceAttenuationMessage>(TraceMessages.DistanceAttenuationMessages, GetActiveSoundEntryFromParameterMessage,
-		[this, &EntriesWithPoppedDataPoints](const FMixerSourceDistanceAttenuationMessage& Msg, const TSharedPtr<FSoundDashboardEntry>* OutActiveSoundEntry)
+		ProcessMessageQueue<FSoundDistanceAttenuationMessage>(TraceMessages.DistanceAttenuationMessages, GetActiveSoundEntryFromParameterMessage,
+		[this, &EntriesWithPoppedDataPoints](const FSoundDistanceAttenuationMessage& Msg, const TSharedPtr<FSoundDashboardEntry>* OutActiveSoundEntry)
 		{
 			FSoundDashboardEntry* SoundEntry = FSoundTraceProviderPrivate::FindSoundEntryForActiveSound(Msg.PlayOrder, OutActiveSoundEntry);
 			if (SoundEntry == nullptr)
@@ -504,8 +503,8 @@ namespace UE::Audio::Insights
 		});
 		EntriesWithPoppedDataPoints.Reset();
 
-		ProcessMessageQueue<FMixerSourceLPFFreqMessage>(TraceMessages.LPFFreqMessages, GetActiveSoundEntryFromParameterMessage,
-		[this, &EntriesWithPoppedDataPoints](const FMixerSourceLPFFreqMessage& Msg, const TSharedPtr<FSoundDashboardEntry>* OutActiveSoundEntry)
+		ProcessMessageQueue<FSoundLPFFreqMessage>(TraceMessages.LPFFreqMessages, GetActiveSoundEntryFromParameterMessage,
+		[this, &EntriesWithPoppedDataPoints](const FSoundLPFFreqMessage& Msg, const TSharedPtr<FSoundDashboardEntry>* OutActiveSoundEntry)
 		{
 			FSoundDashboardEntry* SoundEntry = FSoundTraceProviderPrivate::FindSoundEntryForActiveSound(Msg.PlayOrder, OutActiveSoundEntry);
 			if (SoundEntry == nullptr)
@@ -524,8 +523,8 @@ namespace UE::Audio::Insights
 		});
 		EntriesWithPoppedDataPoints.Reset();
 
-		ProcessMessageQueue<FMixerSourceHPFFreqMessage>(TraceMessages.HPFFreqMessages, GetActiveSoundEntryFromParameterMessage,
-		[this, &EntriesWithPoppedDataPoints](const FMixerSourceHPFFreqMessage& Msg, const TSharedPtr<FSoundDashboardEntry>* OutActiveSoundEntry)
+		ProcessMessageQueue<FSoundHPFFreqMessage>(TraceMessages.HPFFreqMessages, GetActiveSoundEntryFromParameterMessage,
+		[this, &EntriesWithPoppedDataPoints](const FSoundHPFFreqMessage& Msg, const TSharedPtr<FSoundDashboardEntry>* OutActiveSoundEntry)
 		{
 			FSoundDashboardEntry* SoundEntry = FSoundTraceProviderPrivate::FindSoundEntryForActiveSound(Msg.PlayOrder, OutActiveSoundEntry);
 			if (SoundEntry == nullptr)
@@ -544,8 +543,8 @@ namespace UE::Audio::Insights
 		});
 		EntriesWithPoppedDataPoints.Reset();
 
-		ProcessMessageQueue<FMixerSourceEnvelopeMessage>(TraceMessages.AmplitudeMessages, GetActiveSoundEntryFromParameterMessage,
-		[this, &EntriesWithPoppedDataPoints](const FMixerSourceEnvelopeMessage& Msg, const TSharedPtr<FSoundDashboardEntry>* OutActiveSoundEntry)
+		ProcessMessageQueue<FSoundEnvelopeMessage>(TraceMessages.AmplitudeMessages, GetActiveSoundEntryFromParameterMessage,
+		[this, &EntriesWithPoppedDataPoints](const FSoundEnvelopeMessage& Msg, const TSharedPtr<FSoundDashboardEntry>* OutActiveSoundEntry)
 		{
 			FSoundDashboardEntry* SoundEntry = FSoundTraceProviderPrivate::FindSoundEntryForActiveSound(Msg.PlayOrder, OutActiveSoundEntry);
 			if (SoundEntry == nullptr)
@@ -564,8 +563,8 @@ namespace UE::Audio::Insights
 		});
 		EntriesWithPoppedDataPoints.Reset();
 
-		ProcessMessageQueue<FMixerSourceVolumeMessage>(TraceMessages.VolumeMessages, GetActiveSoundEntryFromParameterMessage,
-		[this, &EntriesWithPoppedDataPoints](const FMixerSourceVolumeMessage& Msg, const TSharedPtr<FSoundDashboardEntry>* OutActiveSoundEntry)
+		ProcessMessageQueue<FSoundVolumeMessage>(TraceMessages.VolumeMessages, GetActiveSoundEntryFromParameterMessage,
+		[this, &EntriesWithPoppedDataPoints](const FSoundVolumeMessage& Msg, const TSharedPtr<FSoundDashboardEntry>* OutActiveSoundEntry)
 		{
 			FSoundDashboardEntry* SoundEntry = FSoundTraceProviderPrivate::FindSoundEntryForActiveSound(Msg.PlayOrder, OutActiveSoundEntry);
 			if (SoundEntry == nullptr)
@@ -584,8 +583,8 @@ namespace UE::Audio::Insights
 		});
 		EntriesWithPoppedDataPoints.Reset();
 
-		ProcessMessageQueue<FMixerSourcePitchMessage>(TraceMessages.PitchMessages, GetActiveSoundEntryFromParameterMessage,
-		[this, &EntriesWithPoppedDataPoints](const FMixerSourcePitchMessage& Msg, const TSharedPtr<FSoundDashboardEntry>* OutActiveSoundEntry)
+		ProcessMessageQueue<FSoundPitchMessage>(TraceMessages.PitchMessages, GetActiveSoundEntryFromParameterMessage,
+		[this, &EntriesWithPoppedDataPoints](const FSoundPitchMessage& Msg, const TSharedPtr<FSoundDashboardEntry>* OutActiveSoundEntry)
 		{
 			FSoundDashboardEntry* SoundEntry = FSoundTraceProviderPrivate::FindSoundEntryForActiveSound(Msg.PlayOrder, OutActiveSoundEntry);
 			if (SoundEntry == nullptr)
@@ -604,8 +603,8 @@ namespace UE::Audio::Insights
 		});
 		EntriesWithPoppedDataPoints.Reset();
 
-		ProcessMessageQueue<FMixerSourceRelativeRenderCostMessage>(TraceMessages.RelativeRenderCostMessages, GetActiveSoundEntryFromParameterMessage,
-		[this, &EntriesWithPoppedDataPoints](const FMixerSourceRelativeRenderCostMessage& Msg, const TSharedPtr<FSoundDashboardEntry>* OutActiveSoundEntry)
+		ProcessMessageQueue<FSoundRelativeRenderCostMessage>(TraceMessages.RelativeRenderCostMessages, GetActiveSoundEntryFromParameterMessage,
+		[this, &EntriesWithPoppedDataPoints](const FSoundRelativeRenderCostMessage& Msg, const TSharedPtr<FSoundDashboardEntry>* OutActiveSoundEntry)
 		{
 			FSoundDashboardEntry* SoundEntry = FSoundTraceProviderPrivate::FindSoundEntryForActiveSound(Msg.PlayOrder, OutActiveSoundEntry);
 			if (SoundEntry == nullptr)
@@ -994,7 +993,7 @@ namespace UE::Audio::Insights
 		};
 
 		TSharedPtr<FSoundDashboardEntry> ParamReturn = nullptr;
-		auto GetActiveSoundEntryFromParameterMessage = [this, &ParamReturn](const FMixerSourceParameterMessage& Msg) -> TSharedPtr<FSoundDashboardEntry>*
+		auto GetActiveSoundEntryFromParameterMessage = [this, &ParamReturn](const FSoundParameterMessage& Msg) -> TSharedPtr<FSoundDashboardEntry>*
 		{
 			ParamReturn = nullptr;
 
@@ -1093,13 +1092,13 @@ namespace UE::Audio::Insights
 			{
 				FSoundPriorityMessage PriorityMessage;
 				FSoundDistanceMessage DistanceMessage;
-				FMixerSourceDistanceAttenuationMessage DistanceAttenuationMessage;
-				FMixerSourceLPFFreqMessage LPFFreqMessage;
-				FMixerSourceHPFFreqMessage HPFFreqMessage;
-				FMixerSourceEnvelopeMessage AmplitudeMessage;
-				FMixerSourceVolumeMessage VolumeMessage;
-				FMixerSourcePitchMessage PitchMessage;
-				FMixerSourceRelativeRenderCostMessage RelativeRenderCostMessage;
+				FSoundDistanceAttenuationMessage DistanceAttenuationMessage;
+				FSoundLPFFreqMessage LPFFreqMessage;
+				FSoundHPFFreqMessage HPFFreqMessage;
+				FSoundEnvelopeMessage AmplitudeMessage;
+				FSoundVolumeMessage VolumeMessage;
+				FSoundPitchMessage PitchMessage;
+				FSoundRelativeRenderCostMessage RelativeRenderCostMessage;
 			};
 
 			TArray<uint32> PlayOrderArray;
@@ -1129,49 +1128,49 @@ namespace UE::Audio::Insights
 				}
 
 				// Distance Attenuation
-				const FMixerSourceDistanceAttenuationMessage* FoundDistanceAttenuationCachedMessage = FindClosestMessageToTimestamp(SessionCachedMessages->DistanceAttenuationCachedMessages, TimeMarker, PlayOrder);
+				const FSoundDistanceAttenuationMessage* FoundDistanceAttenuationCachedMessage = FindClosestMessageToTimestamp(SessionCachedMessages->DistanceAttenuationCachedMessages, TimeMarker, PlayOrder);
 				if (FoundDistanceAttenuationCachedMessage)
 				{
 					CachedEntryInfos[Index].DistanceAttenuationMessage = *FoundDistanceAttenuationCachedMessage;
 				}
 
 				// LPF Freq
-				const FMixerSourceLPFFreqMessage* FoundLPFFreqCachedMessage = FindClosestMessageToTimestamp(SessionCachedMessages->LPFFreqCachedMessages, TimeMarker, PlayOrder);
+				const FSoundLPFFreqMessage* FoundLPFFreqCachedMessage = FindClosestMessageToTimestamp(SessionCachedMessages->LPFFreqCachedMessages, TimeMarker, PlayOrder);
 				if (FoundLPFFreqCachedMessage)
 				{
 					CachedEntryInfos[Index].LPFFreqMessage = *FoundLPFFreqCachedMessage;
 				}
 
 				// HPF Freq
-				const FMixerSourceHPFFreqMessage* FoundHPFFreqCachedMessage = FindClosestMessageToTimestamp(SessionCachedMessages->HPFFreqCachedMessages, TimeMarker, PlayOrder);
+				const FSoundHPFFreqMessage* FoundHPFFreqCachedMessage = FindClosestMessageToTimestamp(SessionCachedMessages->HPFFreqCachedMessages, TimeMarker, PlayOrder);
 				if (FoundHPFFreqCachedMessage)
 				{
 					CachedEntryInfos[Index].HPFFreqMessage = *FoundHPFFreqCachedMessage;
 				}
 
 				// Amplitude
-				const FMixerSourceEnvelopeMessage* FoundAmplitudeCachedMessage = FindClosestMessageToTimestamp(SessionCachedMessages->AmplitudeCachedMessages, TimeMarker, PlayOrder);
+				const FSoundEnvelopeMessage* FoundAmplitudeCachedMessage = FindClosestMessageToTimestamp(SessionCachedMessages->AmplitudeCachedMessages, TimeMarker, PlayOrder);
 				if (FoundAmplitudeCachedMessage)
 				{
 					CachedEntryInfos[Index].AmplitudeMessage = *FoundAmplitudeCachedMessage;
 				}
 
 				// Volume
-				const FMixerSourceVolumeMessage* FoundVolumeCachedMessage = FindClosestMessageToTimestamp(SessionCachedMessages->VolumeCachedMessages, TimeMarker, PlayOrder);
+				const FSoundVolumeMessage* FoundVolumeCachedMessage = FindClosestMessageToTimestamp(SessionCachedMessages->VolumeCachedMessages, TimeMarker, PlayOrder);
 				if (FoundVolumeCachedMessage)
 				{
 					CachedEntryInfos[Index].VolumeMessage = *FoundVolumeCachedMessage;
 				}
 
 				// Pitch
-				const FMixerSourcePitchMessage* FoundPitchCachedMessage = FindClosestMessageToTimestamp(SessionCachedMessages->PitchCachedMessages, TimeMarker, PlayOrder);
+				const FSoundPitchMessage* FoundPitchCachedMessage = FindClosestMessageToTimestamp(SessionCachedMessages->PitchCachedMessages, TimeMarker, PlayOrder);
 				if (FoundPitchCachedMessage)
 				{
 					CachedEntryInfos[Index].PitchMessage = *FoundPitchCachedMessage;
 				}
 
 				// Relative Render Cost
-				const FMixerSourceRelativeRenderCostMessage* FoundRelativeRenderCostCachedMessage = FindClosestMessageToTimestamp(SessionCachedMessages->RelativeRenderCostCachedMessages, TimeMarker, PlayOrder);
+				const FSoundRelativeRenderCostMessage* FoundRelativeRenderCostCachedMessage = FindClosestMessageToTimestamp(SessionCachedMessages->RelativeRenderCostCachedMessages, TimeMarker, PlayOrder);
 				if (FoundRelativeRenderCostCachedMessage)
 				{
 					CachedEntryInfos[Index].RelativeRenderCostMessage = *FoundRelativeRenderCostCachedMessage;
diff --git a/Engine/Plugins/AudioInsights/Source/AudioInsights/Private/Views/MixerSourceDashboardViewFactory.cpp b/Engine/Plugins/AudioInsights/Source/AudioInsights/Private/Views/MixerSourceDashboardViewFactory.cpp
deleted file mode 100644
diff --git a/Engine/Plugins/AudioInsights/Source/AudioInsights/Public/AudioInsightsDataSource.h b/Engine/Plugins/AudioInsights/Source/AudioInsights/Public/AudioInsightsDataSource.h
--- a/Engine/Plugins/AudioInsights/Source/AudioInsights/Public/AudioInsightsDataSource.h
+++ b/Engine/Plugins/AudioInsights/Source/AudioInsights/Public/AudioInsightsDataSource.h
@@ -6,6 +6,8 @@
 
 namespace UE::Audio::Insights
 {
+	using FDataPoint = TPair<double, float>; // (Timestamp, Value)
+
 	class IDashboardDataViewEntry : public TSharedFromThis<IDashboardDataViewEntry>
 	{
 	public:
diff --git a/Engine/Plugins/AudioInsights/Source/AudioInsights/Public/Messages/MixerSourceTraceMessages.h b/Engine/Plugins/AudioInsights/Source/AudioInsights/Public/Messages/MixerSourceTraceMessages.h
deleted file mode 100644
diff --git a/Engine/Plugins/AudioInsights/Source/AudioInsights/Public/Messages/SoundTraceMessages.h b/Engine/Plugins/AudioInsights/Source/AudioInsights/Public/Messages/SoundTraceMessages.h
--- a/Engine/Plugins/AudioInsights/Source/AudioInsights/Public/Messages/SoundTraceMessages.h
+++ b/Engine/Plugins/AudioInsights/Source/AudioInsights/Public/Messages/SoundTraceMessages.h
@@ -4,7 +4,6 @@
 #include "AudioDefines.h"
 #include "DSP/Dsp.h"
 #include "Messages/AnalyzerMessageQueue.h"
-#include "Messages/MixerSourceTraceMessages.h"
 #include "Trace/Analyzer.h"
 #include "UObject/SoftObjectPath.h"
 #include "Views/TreeDashboardViewFactory.h"
@@ -40,6 +39,19 @@ namespace UE::Audio::Insights
 		double Timestamp = 0.0;
 	};
 
+	struct FSoundParameterMessage : public FSoundMessageBase
+	{
+		FSoundParameterMessage() = default;
+
+		FSoundParameterMessage(const Trace::IAnalyzer::FOnEventContext& InContext)
+			: FSoundMessageBase(InContext)
+		{
+			ActiveSoundPlayOrder = InContext.EventData.GetValue<uint32>("ActiveSoundPlayOrder");
+		}
+
+		uint32 ActiveSoundPlayOrder = INDEX_NONE;
+	};
+
 	struct FSoundStartMessage : public FSoundMessageBase
 	{
 		FSoundStartMessage() = default;
@@ -68,31 +80,29 @@ namespace UE::Audio::Insights
 
 	using FSoundStopMessage = FSoundMessageBase;
 
-	struct FSoundPriorityMessage : public FMixerSourceParameterMessage
-	{
-		FSoundPriorityMessage() = default;
-		FSoundPriorityMessage(const UE::Trace::IAnalyzer::FOnEventContext& InContext)
-			: FMixerSourceParameterMessage(InContext)
-		{
-			Priority = InContext.EventData.GetValue<float>("Priority");
-		}
-
-		float Priority = 0.0f;
+#define DEFINE_SOUND_PARAM_MESSAGE(ClassName, SuperClassName, ParamName, Type, Default)	\
+	struct ClassName : public SuperClassName													\
+	{																							\
+		ClassName() = default;																	\
+		ClassName(const Trace::IAnalyzer::FOnEventContext& InContext)							\
+			: SuperClassName(InContext)															\
+		{																						\
+			const Trace::IAnalyzer::FEventData& EventData = InContext.EventData;				\
+			ParamName = EventData.GetValue<Type>(#ParamName);									\
+		}																						\
+		Type ParamName = Default;																\
 	};
 
-	struct FSoundDistanceMessage : public FMixerSourceParameterMessage
-	{
-		FSoundDistanceMessage() = default;
-		FSoundDistanceMessage(const UE::Trace::IAnalyzer::FOnEventContext& InContext)
-			: FMixerSourceParameterMessage(InContext)
-		{
-			Distance = InContext.EventData.GetValue<float>("Distance");
-		}
-
-		float Distance = 0.0f;
-	};
-
-
+	DEFINE_SOUND_PARAM_MESSAGE(FSoundPriorityMessage, FSoundParameterMessage, Priority, float, 0.0f)
+	DEFINE_SOUND_PARAM_MESSAGE(FSoundDistanceMessage, FSoundParameterMessage, Distance, float, 0.0f)
+	DEFINE_SOUND_PARAM_MESSAGE(FSoundDistanceAttenuationMessage, FSoundParameterMessage, DistanceAttenuation, float, 1.0f)
+	DEFINE_SOUND_PARAM_MESSAGE(FSoundHPFFreqMessage, FSoundParameterMessage, HPFFrequency, float, MIN_FILTER_FREQUENCY)
+	DEFINE_SOUND_PARAM_MESSAGE(FSoundLPFFreqMessage, FSoundParameterMessage, LPFFrequency, float, MAX_FILTER_FREQUENCY)
+	DEFINE_SOUND_PARAM_MESSAGE(FSoundEnvelopeMessage, FSoundParameterMessage, Envelope, float, 0.0f)
+	DEFINE_SOUND_PARAM_MESSAGE(FSoundPitchMessage, FSoundParameterMessage, Pitch, float, 1.0f)
+	DEFINE_SOUND_PARAM_MESSAGE(FSoundVolumeMessage, FSoundParameterMessage, Volume, float, 1.0f)
+	DEFINE_SOUND_PARAM_MESSAGE(FSoundRelativeRenderCostMessage, FSoundParameterMessage, RelativeRenderCost, float, 0.0f)
+#undef DEFINE_SOUND_PARAM_MESSAGE
 
 	// Trace message queues
 	class FSoundMessages
@@ -102,13 +112,13 @@ namespace UE::Audio::Insights
 		TAnalyzerMessageQueue<FSoundWaveStartMessage> SoundWaveStartMessages{ 0.1 };
 		TAnalyzerMessageQueue<FSoundPriorityMessage> PriorityMessages{ 0.1 };
 		TAnalyzerMessageQueue<FSoundDistanceMessage> DistanceMessages { 0.1 };
-		TAnalyzerMessageQueue<FMixerSourceDistanceAttenuationMessage> DistanceAttenuationMessages { 0.1 };
-		TAnalyzerMessageQueue<FMixerSourceHPFFreqMessage> HPFFreqMessages{ 0.1 };
-		TAnalyzerMessageQueue<FMixerSourceLPFFreqMessage> LPFFreqMessages{ 0.1 };
-		TAnalyzerMessageQueue<FMixerSourceEnvelopeMessage> AmplitudeMessages { 0.1 };
-		TAnalyzerMessageQueue<FMixerSourceVolumeMessage> VolumeMessages { 0.1 };
-		TAnalyzerMessageQueue<FMixerSourcePitchMessage> PitchMessages { 0.1 };
-		TAnalyzerMessageQueue<FMixerSourceRelativeRenderCostMessage> RelativeRenderCostMessages{ 0.1 };
+		TAnalyzerMessageQueue<FSoundDistanceAttenuationMessage> DistanceAttenuationMessages { 0.1 };
+		TAnalyzerMessageQueue<FSoundHPFFreqMessage> HPFFreqMessages{ 0.1 };
+		TAnalyzerMessageQueue<FSoundLPFFreqMessage> LPFFreqMessages{ 0.1 };
+		TAnalyzerMessageQueue<FSoundEnvelopeMessage> AmplitudeMessages { 0.1 };
+		TAnalyzerMessageQueue<FSoundVolumeMessage> VolumeMessages { 0.1 };
+		TAnalyzerMessageQueue<FSoundPitchMessage> PitchMessages { 0.1 };
+		TAnalyzerMessageQueue<FSoundRelativeRenderCostMessage> RelativeRenderCostMessages{ 0.1 };
 		TAnalyzerMessageQueue<FSoundStopMessage> StopMessages { 0.1 };
 	};
 
@@ -136,13 +146,13 @@ namespace UE::Audio::Insights
 		TraceServices::TPagedArray<FSoundWaveStartMessage> SoundWaveStartCachedMessages;
 		TraceServices::TPagedArray<FSoundPriorityMessage> PriorityCachedMessages;
 		TraceServices::TPagedArray<FSoundDistanceMessage> DistanceCachedMessages;
-		TraceServices::TPagedArray<FMixerSourceDistanceAttenuationMessage> DistanceAttenuationCachedMessages;
-		TraceServices::TPagedArray<FMixerSourceHPFFreqMessage> HPFFreqCachedMessages;
-		TraceServices::TPagedArray<FMixerSourceLPFFreqMessage> LPFFreqCachedMessages;
-		TraceServices::TPagedArray<FMixerSourceEnvelopeMessage> AmplitudeCachedMessages;
-		TraceServices::TPagedArray<FMixerSourceVolumeMessage> VolumeCachedMessages;
-		TraceServices::TPagedArray<FMixerSourcePitchMessage> PitchCachedMessages;
-		TraceServices::TPagedArray<FMixerSourceRelativeRenderCostMessage> RelativeRenderCostCachedMessages;
+		TraceServices::TPagedArray<FSoundDistanceAttenuationMessage> DistanceAttenuationCachedMessages;
+		TraceServices::TPagedArray<FSoundHPFFreqMessage> HPFFreqCachedMessages;
+		TraceServices::TPagedArray<FSoundLPFFreqMessage> LPFFreqCachedMessages;
+		TraceServices::TPagedArray<FSoundEnvelopeMessage> AmplitudeCachedMessages;
+		TraceServices::TPagedArray<FSoundVolumeMessage> VolumeCachedMessages;
+		TraceServices::TPagedArray<FSoundPitchMessage> PitchCachedMessages;
+		TraceServices::TPagedArray<FSoundRelativeRenderCostMessage> RelativeRenderCostCachedMessages;
 		TraceServices::TPagedArray<FSoundStopMessage> StopCachedMessages;
 	};
 #endif // !WITH_EDITOR
diff --git a/Engine/Plugins/AudioInsights/Source/AudioInsights/Public/Views/MixerSourceDashboardViewFactory.h b/Engine/Plugins/AudioInsights/Source/AudioInsights/Public/Views/MixerSourceDashboardViewFactory.h
deleted file mode 100644
diff --git a/Engine/Plugins/AudioInsights/Source/AudioInsights/Public/Views/SoundPlotsWidgetView.h b/Engine/Plugins/AudioInsights/Source/AudioInsights/Public/Views/SoundPlotsWidgetView.h
--- a/Engine/Plugins/AudioInsights/Source/AudioInsights/Public/Views/SoundPlotsWidgetView.h
+++ b/Engine/Plugins/AudioInsights/Source/AudioInsights/Public/Views/SoundPlotsWidgetView.h
@@ -2,7 +2,7 @@
 #pragma once
 
 #include "AudioInsightsDataSource.h"
-#include "Messages/MixerSourceTraceMessages.h"
+#include "DSP/Dsp.h"
 #include "Templates/Function.h"
 #include "Views/SAudioCurveView.h"
 
diff --git a/Engine/Plugins/AudioInsights/Source/AudioInsightsEditor/Private/AudioInsightsEditorModule.h b/Engine/Plugins/AudioInsights/Source/AudioInsightsEditor/Private/AudioInsightsEditorModule.h
--- a/Engine/Plugins/AudioInsights/Source/AudioInsightsEditor/Private/AudioInsightsEditorModule.h
+++ b/Engine/Plugins/AudioInsights/Source/AudioInsightsEditor/Private/AudioInsightsEditorModule.h
@@ -7,7 +7,6 @@
 #include "IAudioInsightsEditorModule.h"
 #include "Templates/SharedPointer.h"
 #include "Views/DashboardViewFactory.h"
-#include "Views/MixerSourceMuteSoloFilter.h"
 #include "Views/VirtualLoopsDebugDraw.h"
 #include "Widgets/Docking/SDockTab.h"
 
@@ -46,6 +45,5 @@ namespace UE::Audio::Insights
 		TSharedPtr<FEditorDashboardFactory> DashboardFactory;
 		
 		FVirtualLoopsDebugDraw VirtualLoopsDebugDraw;
-		FMuteSoloFilter MuteSoloFilter;
 	};
 } // namespace UE::Audio::Insights
diff --git a/Engine/Plugins/AudioInsights/Source/AudioInsightsEditor/Private/Views/MixerSourceMuteSoloFilter.cpp b/Engine/Plugins/AudioInsights/Source/AudioInsightsEditor/Private/Views/MixerSourceMuteSoloFilter.cpp
deleted file mode 100644
diff --git a/Engine/Plugins/AudioInsights/Source/AudioInsightsEditor/Private/Views/MixerSourceMuteSoloFilter.h b/Engine/Plugins/AudioInsights/Source/AudioInsightsEditor/Private/Views/MixerSourceMuteSoloFilter.h
deleted file mode 100644
diff --git a/Engine/Source/Runtime/AudioMixer/Private/AudioMixerSource.cpp b/Engine/Source/Runtime/AudioMixer/Private/AudioMixerSource.cpp
--- a/Engine/Source/Runtime/AudioMixer/Private/AudioMixerSource.cpp
+++ b/Engine/Source/Runtime/AudioMixer/Private/AudioMixerSource.cpp
@@ -24,22 +24,7 @@
 CSV_DECLARE_CATEGORY_MODULE_EXTERN(AUDIOMIXERCORE_API, Audio);
 
 #if UE_AUDIO_PROFILERTRACE_ENABLED
-UE_TRACE_EVENT_BEGIN(Audio, MixerSourceStart)
-	UE_TRACE_EVENT_FIELD(uint32, DeviceId)
-	UE_TRACE_EVENT_FIELD(uint64, Timestamp)
-	UE_TRACE_EVENT_FIELD(uint32, PlayOrder)
-	UE_TRACE_EVENT_FIELD(int32, SourceId)
-	UE_TRACE_EVENT_FIELD(uint64, ComponentId)
-	UE_TRACE_EVENT_FIELD(UE::Trace::WideString, Name)
-UE_TRACE_EVENT_END()
-
-UE_TRACE_EVENT_BEGIN(Audio, MixerSourceStop)
-	UE_TRACE_EVENT_FIELD(uint32, DeviceId)
-	UE_TRACE_EVENT_FIELD(uint64, Timestamp)
-	UE_TRACE_EVENT_FIELD(uint32, PlayOrder)
-UE_TRACE_EVENT_END()
-
-UE_TRACE_EVENT_BEGIN(Audio, MixerSourceRelativeRenderCost)
+UE_TRACE_EVENT_BEGIN(Audio, SoundRelativeRenderCost)
 	UE_TRACE_EVENT_FIELD(uint32, DeviceId)
 	UE_TRACE_EVENT_FIELD(uint64, Timestamp)
 	UE_TRACE_EVENT_FIELD(uint32, PlayOrder)
@@ -201,10 +186,6 @@ namespace Audio
 
 	FMixerSource::~FMixerSource()
 	{
-#if UE_AUDIO_PROFILERTRACE_ENABLED
-		FTraceAuxiliary::OnTraceStarted.RemoveAll(this);
-#endif // UE_AUDIO_PROFILERTRACE_ENABLED
-
 		FreeResources();
 	}
 
@@ -602,10 +583,6 @@ namespace Audio
 			}
 #endif
 
-#if UE_AUDIO_PROFILERTRACE_ENABLED
-			FTraceAuxiliary::OnTraceStarted.AddRaw(this, &FMixerSource::OnTraceStarted);
-#endif // UE_AUDIO_PROFILERTRACE_ENABLED
-
 			// Pass the decompression state off to the mixer source buffer if it hasn't already done so
 			ICompressedAudioInfo* Decoder = MixerBuffer->GetDecompressionState(true);
 			MixerSourceBuffer->SetDecoder(Decoder);
@@ -1047,29 +1024,6 @@ namespace Audio
 				WaveInstance->WaveData->CachedRealtimeFirstBuffer,ToCStr(LexToString(MixerSourceBuffer.IsValid())), TickCount, ToCStr(LexToString(InitializationState)));
 
 			MixerSourceVoice->Play();
-
-#if UE_AUDIO_PROFILERTRACE_ENABLED
-			const bool bChannelEnabled = UE_TRACE_CHANNELEXPR_IS_ENABLED(AudioMixerChannel);
-			if (bChannelEnabled && WaveInstance)
-			{
-				if (const FActiveSound* ActiveSound = WaveInstance->ActiveSound)
-				{
-					int32 TraceSourceId = INDEX_NONE;
-					if (MixerSourceVoice)
-					{
-						TraceSourceId = MixerSourceVoice->GetSourceId();
-					}
-
-					UE_TRACE_LOG(Audio, MixerSourceStart, AudioMixerChannel)
-						<< MixerSourceStart.DeviceId(MixerDevice->DeviceID)
-						<< MixerSourceStart.Timestamp(FPlatformTime::Cycles64())
-						<< MixerSourceStart.PlayOrder(WaveInstance->GetPlayOrder())
-						<< MixerSourceStart.SourceId(TraceSourceId)
-						<< MixerSourceStart.ComponentId(ActiveSound->GetAudioComponentID())
-						<< MixerSourceStart.Name(*WaveInstance->WaveData->GetPathName());
-				}
-			}
-#endif // UE_AUDIO_PROFILERTRACE_ENABLED
 		}
 
 		bIsStopping = false;
@@ -1152,23 +1106,6 @@ namespace Audio
 		{
 			if (MixerSourceVoice && Playing)
 			{
-#if UE_AUDIO_PROFILERTRACE_ENABLED
-				const bool bChannelEnabled = UE_TRACE_CHANNELEXPR_IS_ENABLED(AudioMixerChannel);
-				if (bChannelEnabled)
-				{
-					int32 TraceSourceId = INDEX_NONE;
-					if (MixerSourceVoice)
-					{
-						TraceSourceId = MixerSourceVoice->GetSourceId();
-					}
-
-					UE_TRACE_LOG(Audio, MixerSourceStop, AudioMixerChannel)
-						<< MixerSourceStop.DeviceId(MixerDevice->DeviceID)
-						<< MixerSourceStop.Timestamp(FPlatformTime::Cycles64())
-						<< MixerSourceStop.PlayOrder(WaveInstance->GetPlayOrder());
-				}
-#endif // UE_AUDIO_PROFILERTRACE_ENABLED
-
 				MixerSourceVoice->Stop();
 			}
 
@@ -1831,12 +1768,12 @@ namespace Audio
 			const bool bChannelEnabled = UE_TRACE_CHANNELEXPR_IS_ENABLED(AudioMixerChannel);
 			if (bChannelEnabled && WaveInstance->ActiveSound != nullptr)
 			{
-				UE_TRACE_LOG(Audio, MixerSourceRelativeRenderCost, AudioMixerChannel)
-					<< MixerSourceRelativeRenderCost.DeviceId(MixerDevice->DeviceID)
-					<< MixerSourceRelativeRenderCost.Timestamp(FPlatformTime::Cycles64())
-					<< MixerSourceRelativeRenderCost.PlayOrder(WaveInstance->GetPlayOrder())
-					<< MixerSourceRelativeRenderCost.ActiveSoundPlayOrder(WaveInstance->ActiveSound->GetPlayOrder())
-					<< MixerSourceRelativeRenderCost.RelativeRenderCost(RelativeRenderCost);
+				UE_TRACE_LOG(Audio, SoundRelativeRenderCost, AudioMixerChannel)
+					<< SoundRelativeRenderCost.DeviceId(MixerDevice->DeviceID)
+					<< SoundRelativeRenderCost.Timestamp(FPlatformTime::Cycles64())
+					<< SoundRelativeRenderCost.PlayOrder(WaveInstance->GetPlayOrder())
+					<< SoundRelativeRenderCost.ActiveSoundPlayOrder(WaveInstance->ActiveSound->GetPlayOrder())
+					<< SoundRelativeRenderCost.RelativeRenderCost(RelativeRenderCost);
 			}
 #endif // UE_AUDIO_PROFILERTRACE_ENABLED
 		}
@@ -1856,34 +1793,6 @@ namespace Audio
 	}
 #endif // if ENABLE_AUDIO_DEBUG
 
-#if UE_AUDIO_PROFILERTRACE_ENABLED
-	void FMixerSource::OnTraceStarted(FTraceAuxiliary::EConnectionType TraceType, const FString& TraceDestination)
-	{
-		// Make sure we send mixer source data to Audio Insights if it's opened in the middle of a PIE session
-		const bool bChannelEnabled = UE_TRACE_CHANNELEXPR_IS_ENABLED(AudioMixerChannel);
-		if (bChannelEnabled && WaveInstance)
-		{
-			if (const FActiveSound* ActiveSound = WaveInstance->ActiveSound;
-				ActiveSound && ActiveSound->IsPlayingAudio())
-			{
-				int32 TraceSourceId = INDEX_NONE;
-				if (MixerSourceVoice)
-				{
-					TraceSourceId = MixerSourceVoice->GetSourceId();
-				}
-
-				UE_TRACE_LOG(Audio, MixerSourceStart, AudioMixerChannel)
-					<< MixerSourceStart.DeviceId(MixerDevice->DeviceID)
-					<< MixerSourceStart.Timestamp(FPlatformTime::Cycles64())
-					<< MixerSourceStart.PlayOrder(WaveInstance->GetPlayOrder())
-					<< MixerSourceStart.SourceId(TraceSourceId)
-					<< MixerSourceStart.ComponentId(ActiveSound->GetAudioComponentID())
-					<< MixerSourceStart.Name(*WaveInstance->WaveData->GetPathName());
-			}
-		}
-	}
-#endif // UE_AUDIO_PROFILERTRACE_ENABLED
-
 	bool FMixerSource::ComputeMonoChannelMap(Audio::FAlignedFloatBuffer& OutChannelMap)
 	{
 		if (IsUsingObjectBasedSpatialization())
diff --git a/Engine/Source/Runtime/AudioMixer/Private/AudioMixerSource.h b/Engine/Source/Runtime/AudioMixer/Private/AudioMixerSource.h
--- a/Engine/Source/Runtime/AudioMixer/Private/AudioMixerSource.h
+++ b/Engine/Source/Runtime/AudioMixer/Private/AudioMixerSource.h
@@ -98,10 +98,6 @@ namespace Audio
 		void UpdateCPUCoreUtilization();
 #endif // ENABLE_AUDIO_DEBUG
 
-#if UE_AUDIO_PROFILERTRACE_ENABLED
-		void OnTraceStarted(FTraceAuxiliary::EConnectionType TraceType, const FString& TraceDestination);
-#endif // UE_AUDIO_PROFILERTRACE_ENABLED
-
 		/** Computes the mono-channel map. */
 		bool ComputeMonoChannelMap(Audio::FAlignedFloatBuffer& OutChannelMap);
 
diff --git a/Engine/Source/Runtime/AudioMixer/Private/AudioMixerSourceManager.cpp b/Engine/Source/Runtime/AudioMixer/Private/AudioMixerSourceManager.cpp
--- a/Engine/Source/Runtime/AudioMixer/Private/AudioMixerSourceManager.cpp
+++ b/Engine/Source/Runtime/AudioMixer/Private/AudioMixerSourceManager.cpp
@@ -221,7 +221,7 @@ DEFINE_STAT(STAT_AudioMixerSourceOutputBuffers);
 
 #if UE_AUDIO_PROFILERTRACE_ENABLED
 // Mixer Source messages
-UE_TRACE_EVENT_BEGIN(Audio, MixerSourceVolume)
+UE_TRACE_EVENT_BEGIN(Audio, SoundVolume)
 	UE_TRACE_EVENT_FIELD(uint32, DeviceId)
 	UE_TRACE_EVENT_FIELD(uint64, Timestamp)
 	UE_TRACE_EVENT_FIELD(uint32, PlayOrder)
@@ -229,7 +229,7 @@ UE_TRACE_EVENT_BEGIN(Audio, MixerSourceVolume)
 	UE_TRACE_EVENT_FIELD(float, Volume)
 UE_TRACE_EVENT_END()
 
-UE_TRACE_EVENT_BEGIN(Audio, MixerSourceDistanceAttenuation)
+UE_TRACE_EVENT_BEGIN(Audio, SoundDistanceAttenuation)
 	UE_TRACE_EVENT_FIELD(uint32, DeviceId)
 	UE_TRACE_EVENT_FIELD(uint64, Timestamp)
 	UE_TRACE_EVENT_FIELD(uint32, PlayOrder)
@@ -237,7 +237,7 @@ UE_TRACE_EVENT_BEGIN(Audio, MixerSourceDistanceAttenuation)
 	UE_TRACE_EVENT_FIELD(float, DistanceAttenuation)
 UE_TRACE_EVENT_END()
 
-UE_TRACE_EVENT_BEGIN(Audio, MixerSourcePitch)
+UE_TRACE_EVENT_BEGIN(Audio, SoundPitch)
 	UE_TRACE_EVENT_FIELD(uint32, DeviceId)
 	UE_TRACE_EVENT_FIELD(uint64, Timestamp)
 	UE_TRACE_EVENT_FIELD(uint32, PlayOrder)
@@ -245,7 +245,7 @@ UE_TRACE_EVENT_BEGIN(Audio, MixerSourcePitch)
 	UE_TRACE_EVENT_FIELD(float, Pitch)
 UE_TRACE_EVENT_END()
 
-UE_TRACE_EVENT_BEGIN(Audio, MixerSourceFilters)
+UE_TRACE_EVENT_BEGIN(Audio, SoundSourceFilters)
 	UE_TRACE_EVENT_FIELD(uint32, DeviceId)
 	UE_TRACE_EVENT_FIELD(uint64, Timestamp)
 	UE_TRACE_EVENT_FIELD(uint32, PlayOrder)
@@ -254,7 +254,7 @@ UE_TRACE_EVENT_BEGIN(Audio, MixerSourceFilters)
 	UE_TRACE_EVENT_FIELD(float, HPFFrequency)
 UE_TRACE_EVENT_END()
 
-UE_TRACE_EVENT_BEGIN(Audio, MixerSourceEnvelope)
+UE_TRACE_EVENT_BEGIN(Audio, SoundSourceEnvelope)
 	UE_TRACE_EVENT_FIELD(uint32, DeviceId)
 	UE_TRACE_EVENT_FIELD(uint64, Timestamp)
 	UE_TRACE_EVENT_FIELD(uint32, PlayOrder)
@@ -2316,12 +2316,12 @@ namespace Audio
 			const bool bChannelEnabled = UE_TRACE_CHANNELEXPR_IS_ENABLED(AudioMixerChannel);
 			if (bChannelEnabled)
 			{
-				UE_TRACE_LOG(Audio, MixerSourcePitch, AudioMixerChannel)
-					<< MixerSourcePitch.DeviceId(MixerDevice->DeviceID)
-					<< MixerSourcePitch.Timestamp(FPlatformTime::Cycles64())
-					<< MixerSourcePitch.PlayOrder(SourceInfo.PlayOrder)
-					<< MixerSourcePitch.ActiveSoundPlayOrder(SourceInfo.ActiveSoundPlayOrder)
-					<< MixerSourcePitch.Pitch(FinalPitch);
+				UE_TRACE_LOG(Audio, SoundPitch, AudioMixerChannel)
+					<< SoundPitch.DeviceId(MixerDevice->DeviceID)
+					<< SoundPitch.Timestamp(FPlatformTime::Cycles64())
+					<< SoundPitch.PlayOrder(SourceInfo.PlayOrder)
+					<< SoundPitch.ActiveSoundPlayOrder(SourceInfo.ActiveSoundPlayOrder)
+					<< SoundPitch.Pitch(FinalPitch);
 			}
 #endif // UE_AUDIO_PROFILERTRACE_ENABLED
 
@@ -2974,19 +2974,19 @@ namespace Audio
 			const bool bChannelEnabled = UE_TRACE_CHANNELEXPR_IS_ENABLED(AudioMixerChannel);
 			if (bChannelEnabled)
 			{
-				UE_TRACE_LOG(Audio, MixerSourceVolume, AudioMixerChannel)
-					<< MixerSourceVolume.DeviceId(MixerDevice->DeviceID)
-					<< MixerSourceVolume.Timestamp(FPlatformTime::Cycles64())
-					<< MixerSourceVolume.PlayOrder(SourceInfo.PlayOrder)
-					<< MixerSourceVolume.ActiveSoundPlayOrder(SourceInfo.ActiveSoundPlayOrder)
-					<< MixerSourceVolume.Volume(VolumeDestination);
+				UE_TRACE_LOG(Audio, SoundVolume, AudioMixerChannel)
+					<< SoundVolume.DeviceId(MixerDevice->DeviceID)
+					<< SoundVolume.Timestamp(FPlatformTime::Cycles64())
+					<< SoundVolume.PlayOrder(SourceInfo.PlayOrder)
+					<< SoundVolume.ActiveSoundPlayOrder(SourceInfo.ActiveSoundPlayOrder)
+					<< SoundVolume.Volume(VolumeDestination);
 
-				UE_TRACE_LOG(Audio, MixerSourceDistanceAttenuation, AudioMixerChannel)
-					<< MixerSourceDistanceAttenuation.DeviceId(MixerDevice->DeviceID)
-					<< MixerSourceDistanceAttenuation.Timestamp(FPlatformTime::Cycles64())
-					<< MixerSourceDistanceAttenuation.PlayOrder(SourceInfo.PlayOrder)
-					<< MixerSourceDistanceAttenuation.ActiveSoundPlayOrder(SourceInfo.ActiveSoundPlayOrder)
-					<< MixerSourceDistanceAttenuation.DistanceAttenuation(SourceInfo.DistanceAttenuationSourceDestination);
+				UE_TRACE_LOG(Audio, SoundDistanceAttenuation, AudioMixerChannel)
+					<< SoundDistanceAttenuation.DeviceId(MixerDevice->DeviceID)
+					<< SoundDistanceAttenuation.Timestamp(FPlatformTime::Cycles64())
+					<< SoundDistanceAttenuation.PlayOrder(SourceInfo.PlayOrder)
+					<< SoundDistanceAttenuation.ActiveSoundPlayOrder(SourceInfo.ActiveSoundPlayOrder)
+					<< SoundDistanceAttenuation.DistanceAttenuation(SourceInfo.DistanceAttenuationSourceDestination);
 			}
 #endif // UE_AUDIO_PROFILERTRACE_ENABLED
 		}
@@ -3121,19 +3121,19 @@ namespace Audio
 					HPFFrequency = SourceInfo.HighpassModulation.GetValue();
 				}
 
-				UE_TRACE_LOG(Audio, MixerSourceFilters, AudioMixerChannel)
-					<< MixerSourceFilters.DeviceId(MixerDevice->DeviceID)
-					<< MixerSourceFilters.Timestamp(FPlatformTime::Cycles64())
-					<< MixerSourceFilters.PlayOrder(SourceInfo.PlayOrder)
-					<< MixerSourceFilters.ActiveSoundPlayOrder(SourceInfo.ActiveSoundPlayOrder)
-					<< MixerSourceFilters.HPFFrequency(HPFFrequency)
-					<< MixerSourceFilters.LPFFrequency(LPFFrequency);
-				UE_TRACE_LOG(Audio, MixerSourceEnvelope, AudioMixerChannel)
-					<< MixerSourceEnvelope.DeviceId(MixerDevice->DeviceID)
-					<< MixerSourceEnvelope.Timestamp(FPlatformTime::Cycles64())
-					<< MixerSourceEnvelope.PlayOrder(SourceInfo.PlayOrder)
-					<< MixerSourceEnvelope.ActiveSoundPlayOrder(SourceInfo.ActiveSoundPlayOrder)
-					<< MixerSourceEnvelope.Envelope(SourceInfo.SourceEnvelopeValue);
+				UE_TRACE_LOG(Audio, SoundSourceFilters, AudioMixerChannel)
+					<< SoundSourceFilters.DeviceId(MixerDevice->DeviceID)
+					<< SoundSourceFilters.Timestamp(FPlatformTime::Cycles64())
+					<< SoundSourceFilters.PlayOrder(SourceInfo.PlayOrder)
+					<< SoundSourceFilters.ActiveSoundPlayOrder(SourceInfo.ActiveSoundPlayOrder)
+					<< SoundSourceFilters.HPFFrequency(HPFFrequency)
+					<< SoundSourceFilters.LPFFrequency(LPFFrequency);
+				UE_TRACE_LOG(Audio, SoundSourceEnvelope, AudioMixerChannel)
+					<< SoundSourceEnvelope.DeviceId(MixerDevice->DeviceID)
+					<< SoundSourceEnvelope.Timestamp(FPlatformTime::Cycles64())
+					<< SoundSourceEnvelope.PlayOrder(SourceInfo.PlayOrder)
+					<< SoundSourceEnvelope.ActiveSoundPlayOrder(SourceInfo.ActiveSoundPlayOrder)
+					<< SoundSourceEnvelope.Envelope(SourceInfo.SourceEnvelopeValue);
 			}
 #endif // UE_AUDIO_PROFILERTRACE_ENABLED
 
