[Audio Insights] Fix Plots widget not working when using direct trace in editor. Pause the plot graph when PIE is paused.

diff --git a/Engine/Plugins/AudioInsights/Source/AudioInsights/Private/Views/SoundPlotsWidgetView.cpp b/Engine/Plugins/AudioInsights/Source/AudioInsights/Private/Views/SoundPlotsWidgetView.cpp
--- a/Engine/Plugins/AudioInsights/Source/AudioInsights/Private/Views/SoundPlotsWidgetView.cpp
+++ b/Engine/Plugins/AudioInsights/Source/AudioInsights/Private/Views/SoundPlotsWidgetView.cpp
@@ -41,6 +41,8 @@ namespace UE::Audio::Insights
 		FEditorDelegates::EndPIE.AddSP(this, &FSoundPlotsWidgetView::OnPIEStopped);
 		FEditorDelegates::PausePIE.AddSP(this, &FSoundPlotsWidgetView::OnPIEPaused);
 		FEditorDelegates::ResumePIE.AddSP(this, &FSoundPlotsWidgetView::OnPIEResumed);
+
+		FTraceAuxiliary::OnTraceStopped.AddSP(this, &FSoundPlotsWidgetView::OnTraceStopped);
 #else
 		FAudioInsightsModule& AudioInsightsModule = FAudioInsightsModule::GetChecked();
 
@@ -59,7 +61,7 @@ namespace UE::Audio::Insights
 
 	void FSoundPlotsWidgetView::ProcessPlotData(const TArray<TSharedPtr<IDashboardDataTreeViewEntry>>& DataViewEntries)
 	{
-		if (!PlotWidgetMetadataPerCurve.IsValid() || DataViewEntries.IsEmpty())
+		if (!PlotWidgetMetadataPerCurve.IsValid() || DataViewEntries.IsEmpty() || GameState == EGameState::Paused)
 		{
 			return;
 		}
@@ -126,8 +128,6 @@ namespace UE::Audio::Insights
 
 	void FSoundPlotsWidgetView::OnPIEStopped(bool bSimulating)
 	{
-		ResetPlots();
-
 		GameState = EGameState::Stopped;
 	}
 
@@ -140,6 +140,11 @@ namespace UE::Audio::Insights
 	{
 		GameState = EGameState::Running;
 	}
+
+	void FSoundPlotsWidgetView::OnTraceStopped(FTraceAuxiliary::EConnectionType InTraceType, const FString& InTraceDestination)
+	{
+		ResetPlots();
+	}
 #else
 	void FSoundPlotsWidgetView::OnAudioInsightsComponentTabSpawn()
 	{
@@ -409,11 +414,7 @@ namespace UE::Audio::Insights
 
 	TRange<double> FSoundPlotsWidgetView::GetViewRange()
 	{
-#if WITH_EDITOR
-		if (GameState == EGameState::Stopped || BeginTimestamp == TNumericLimits<double>::Max())
-#else
 		if (BeginTimestamp == TNumericLimits<double>::Max())
-#endif // WITH_EDITOR
 		{
 			return TRange<double>(0, MaxPlotHistorySeconds);
 		}
@@ -427,7 +428,7 @@ namespace UE::Audio::Insights
 
 		const bool bAnyMessageReceived = CurrentTimestamp != TNumericLimits<double>::Lowest();
 
-		if (GameState == EGameState::Running && bAnyMessageReceived)
+		if (GameState != EGameState::Paused && bAnyMessageReceived)
 		{
 			const double RelativeCurrentTime = CurrentTime - BeginTimestamp;
 			TimestampsDiff = RelativeCurrentTime - (CurrentTimestamp - BeginTimestamp);
diff --git a/Engine/Plugins/AudioInsights/Source/AudioInsights/Public/Views/SoundPlotsWidgetView.h b/Engine/Plugins/AudioInsights/Source/AudioInsights/Public/Views/SoundPlotsWidgetView.h
--- a/Engine/Plugins/AudioInsights/Source/AudioInsights/Public/Views/SoundPlotsWidgetView.h
+++ b/Engine/Plugins/AudioInsights/Source/AudioInsights/Public/Views/SoundPlotsWidgetView.h
@@ -42,6 +42,8 @@ namespace UE::Audio::Insights
 		void OnPIEStopped(bool bSimulating);
 		void OnPIEPaused(bool bSimulating);
 		void OnPIEResumed(bool bSimulating);
+
+		void OnTraceStopped(FTraceAuxiliary::EConnectionType InTraceType, const FString& InTraceDestination);
 #else
 		void OnAudioInsightsComponentTabSpawn();
 		void OnSessionAnalysisCompleted();
