[Audio Insights] Fix Sound Dashboard not retaining expanded state of items when switching viewing modes

diff --git a/Engine/Plugins/AudioInsights/Source/AudioInsights/Private/Views/SoundDashboardViewFactory.cpp b/Engine/Plugins/AudioInsights/Source/AudioInsights/Private/Views/SoundDashboardViewFactory.cpp
--- a/Engine/Plugins/AudioInsights/Source/AudioInsights/Private/Views/SoundDashboardViewFactory.cpp
+++ b/Engine/Plugins/AudioInsights/Source/AudioInsights/Private/Views/SoundDashboardViewFactory.cpp
@@ -176,13 +176,15 @@ namespace UE::Audio::Insights
 				];
 		};
 
-		void CacheInitExpandStateRecursive(FSoundDashboardEntry& SoundEntry)
+		void RefreshInitExpandStateRecursive(FSoundDashboardEntry& SoundEntry)
 		{
+			SoundEntry.bHasSetInitExpansion = false;
+
 			for (TSharedPtr<IDashboardDataTreeViewEntry>& ChildEntry : SoundEntry.Children)
 			{
 				if (ChildEntry.IsValid())
 				{
-					CacheInitExpandStateRecursive(CastEntry(*ChildEntry));
+					RefreshInitExpandStateRecursive(CastEntry(*ChildEntry));
 				}
 			}
 		}
@@ -1333,7 +1335,7 @@ namespace UE::Audio::Insights
 			for (const TSharedPtr<IDashboardDataTreeViewEntry>& Entry : DataViewEntries)
 			{
 				FSoundDashboardEntry& SoundEntry = CastEntry(*Entry);
-				CacheInitExpandStateRecursive(SoundEntry);
+				RefreshInitExpandStateRecursive(SoundEntry);
 			}
 		}
 
@@ -1371,7 +1373,7 @@ namespace UE::Audio::Insights
 
 			if (!bEntryTypePassesFilter)
 			{
-				CacheInitExpandStateRecursive(SoundCategoryEntry);
+				RefreshInitExpandStateRecursive(SoundCategoryEntry);
 			}
 
 			return bEntryTypePassesFilter;
@@ -1385,11 +1387,19 @@ namespace UE::Audio::Insights
 
 	void FSoundDashboardViewFactory::ChangeTreeViewingMode(const ESoundDashboardTreeViewingOptions SelectedMode, const bool bUpdateEditorSettings /*= true*/)
 	{
+		using namespace FSoundDashboardViewFactoryPrivate;
+
 		if (TreeViewingMode == SelectedMode)
 		{
 			return;
 		}
 
+		for (const TSharedPtr<IDashboardDataTreeViewEntry>& Entry : DataViewEntries)
+		{
+			FSoundDashboardEntry& SoundEntry = CastEntry(*Entry);
+			RefreshInitExpandStateRecursive(SoundEntry);
+		}
+
 		TreeViewingMode = SelectedMode;
 		UpdateFilterReason = EProcessReason::FilterUpdated;
 		ClearSelection();
