[Audio Insights] Allow users to create a trace bookmark when tracing data with Audio Insights

diff --git a/Engine/Plugins/AudioInsights/Source/AudioInsights/Private/AudioInsightsDashboardAssetCommands.cpp b/Engine/Plugins/AudioInsights/Source/AudioInsights/Private/AudioInsightsDashboardAssetCommands.cpp
--- a/Engine/Plugins/AudioInsights/Source/AudioInsights/Private/AudioInsightsDashboardAssetCommands.cpp
+++ b/Engine/Plugins/AudioInsights/Source/AudioInsights/Private/AudioInsightsDashboardAssetCommands.cpp
@@ -25,6 +25,8 @@ namespace UE::Audio::Insights
 
 		UI_COMMAND(Start, "Start Trace", "Starts the active trace session used by Audio Insights.", EUserInterfaceActionType::Button, FInputChord());
 		UI_COMMAND(Stop, "Stop Trace", "Stops the active trace session used by Audio Insights.", EUserInterfaceActionType::Button, FInputChord());
+
+		UI_COMMAND(CreateTraceBookmark, "Create Trace Bookmark", "Places a bookmark inside a trace file when tracing with Audio Insights.", EUserInterfaceActionType::Button, FInputChord(EKeys::M, EModifierKey::Control));
 	}
 
 	void FDashboardAssetCommands::AddAssetCommands(FToolBarBuilder& OutToolbarBuilder) const
@@ -68,6 +70,11 @@ namespace UE::Audio::Insights
 		return Stop;
 	}
 
+	TSharedPtr<const FUICommandInfo> FDashboardAssetCommands::GetTraceBookmarkCommand() const
+	{
+		return CreateTraceBookmark;
+	}
+
 	FSlateIcon FDashboardAssetCommands::GetStartIcon() const
 	{
 		const IAudioInsightsTraceModule& TraceModule = FAudioInsightsModule::GetChecked().GetTraceModule();
diff --git a/Engine/Plugins/AudioInsights/Source/AudioInsights/Private/AudioInsightsTraceModule.cpp b/Engine/Plugins/AudioInsights/Source/AudioInsights/Private/AudioInsightsTraceModule.cpp
--- a/Engine/Plugins/AudioInsights/Source/AudioInsights/Private/AudioInsightsTraceModule.cpp
+++ b/Engine/Plugins/AudioInsights/Source/AudioInsights/Private/AudioInsightsTraceModule.cpp
@@ -23,6 +23,7 @@ namespace UE::Audio::Insights
 		static const TCHAR* CPUChannelName = TEXT("Cpu");
 		static const TCHAR* AudioChannelName = TEXT("Audio");
 		static const TCHAR* AudioMixerChannelName = TEXT("AudioMixer");
+		static const TCHAR* BookmarkChannelName = TEXT("Bookmark");
 	}
 
 	void FRewindDebuggerAudioInsightsRuntime::RecordingStarted()
@@ -31,11 +32,12 @@ namespace UE::Audio::Insights
 
 		UE::Trace::ToggleChannel(AudioChannelName, true);
 		UE::Trace::ToggleChannel(AudioMixerChannelName, true);
+		UE::Trace::ToggleChannel(BookmarkChannelName, true);
 	}
 
 	FTraceModule::FTraceModule()
 #if !WITH_EDITOR
-		: AudioChannels({ FTraceModulePrivate::AudioChannelName, FTraceModulePrivate::AudioMixerChannelName } )
+		: AudioChannels({ FTraceModulePrivate::AudioChannelName, FTraceModulePrivate::AudioMixerChannelName, FTraceModulePrivate::BookmarkChannelName } )
 		, EmptyArray()
 #endif
 	{
@@ -90,8 +92,9 @@ namespace UE::Audio::Insights
  		const bool bCPUOn = UE::Trace::ToggleChannel(CPUChannelName, true);
  		const bool bAudioOn = UE::Trace::ToggleChannel(AudioChannelName, true);
  		const bool bAudioMixerOn = UE::Trace::ToggleChannel(AudioMixerChannelName, true);
+		const bool bBookmarkChannelOn = UE::Trace::ToggleChannel(BookmarkChannelName, true);
 
- 		return bCPUOn && bAudioOn && bAudioMixerOn;
+ 		return bCPUOn && bAudioOn && bAudioMixerOn && bBookmarkChannelOn;
 #else
 		if (!InstanceID.IsValid())
 		{
@@ -137,6 +140,7 @@ namespace UE::Audio::Insights
 		DisableIfNotMarkedToRestore(CPUChannelName);
 		DisableIfNotMarkedToRestore(AudioChannelName);
 		DisableIfNotMarkedToRestore(AudioMixerChannelName);
+		DisableIfNotMarkedToRestore(BookmarkChannelName);
 	}
 
 	void FTraceModule::CacheCurrentlyEnabledTraceChannels()
diff --git a/Engine/Plugins/AudioInsights/Source/AudioInsights/Public/AudioInsightsDashboardAssetCommands.h b/Engine/Plugins/AudioInsights/Source/AudioInsights/Public/AudioInsightsDashboardAssetCommands.h
--- a/Engine/Plugins/AudioInsights/Source/AudioInsights/Public/AudioInsightsDashboardAssetCommands.h
+++ b/Engine/Plugins/AudioInsights/Source/AudioInsights/Public/AudioInsightsDashboardAssetCommands.h
@@ -24,6 +24,8 @@ namespace UE::Audio::Insights
 		virtual TSharedPtr<const FUICommandInfo> GetStartCommand() const;
 		virtual TSharedPtr<const FUICommandInfo> GetStopCommand() const;
 
+		virtual TSharedPtr<const FUICommandInfo> GetTraceBookmarkCommand() const;
+
 	private:
 		FSlateIcon GetStartIcon() const;
 		FSlateIcon GetStopIcon() const;
@@ -39,5 +41,8 @@ namespace UE::Audio::Insights
 
 		/** Stops the trace session & disables the required audio channels off. */
 		TSharedPtr<FUICommandInfo> Stop;
+
+		/** Places a bookmark inside a trace file. */
+		TSharedPtr<FUICommandInfo> CreateTraceBookmark;
 	};
 } // namespace UE::Audio::Insights
diff --git a/Engine/Plugins/AudioInsights/Source/AudioInsights/Public/AudioInsightsStyle.h b/Engine/Plugins/AudioInsights/Source/AudioInsights/Public/AudioInsightsStyle.h
--- a/Engine/Plugins/AudioInsights/Source/AudioInsights/Public/AudioInsightsStyle.h
+++ b/Engine/Plugins/AudioInsights/Source/AudioInsights/Public/AudioInsightsStyle.h
@@ -166,6 +166,7 @@ namespace UE::Audio::Insights
 			Set("AudioInsights.Icon.Submix", new IMAGE_BRUSH_SVG(TEXT("Icons/audio_submix"), Icon16));
 			Set("AudioInsights.Icon.Trace", new IMAGE_BRUSH_SVG(TEXT("Icons/trace"), Icon16));
 			Set("AudioInsights.Icon.Trace.Snapshot", new IMAGE_BRUSH_SVG(TEXT("Icons/trace_snapshot"), Icon16));
+			Set("AudioInsights.Icon.Trace.Bookmark", new IMAGE_BRUSH_SVG(TEXT("Icons/bookmark"), Icon16));
 			Set("AudioInsights.Icon.VirtualLoop", new IMAGE_BRUSH_SVG(TEXT("Icons/audio_virtualloop"), Icon16));
 			Set("AudioInsights.Icon.Viewport", new IMAGE_BRUSH_SVG(TEXT("Icons/viewport"), Icon16));
 			Set("AudioInsights.Thumbnail", new IMAGE_BRUSH_SVG(TEXT("Icons/audio_insights"), Icon16));
diff --git a/Engine/Plugins/AudioInsights/Source/AudioInsightsEditor/Private/AudioInsightsEditorDashboardFactory.cpp b/Engine/Plugins/AudioInsights/Source/AudioInsightsEditor/Private/AudioInsightsEditorDashboardFactory.cpp
--- a/Engine/Plugins/AudioInsights/Source/AudioInsightsEditor/Private/AudioInsightsEditorDashboardFactory.cpp
+++ b/Engine/Plugins/AudioInsights/Source/AudioInsightsEditor/Private/AudioInsightsEditorDashboardFactory.cpp
@@ -4,11 +4,13 @@
 #include "Async/Async.h"
 #include "AudioDevice.h"
 #include "AudioDeviceManager.h"
+#include "AudioInsightsDashboardAssetCommands.h"
 #include "AudioInsightsEditorModule.h"
 #include "AudioInsightsEditorSettings.h"
 #include "AudioInsightsStyle.h"
 #include "Editor.h"
 #include "Engine/World.h"
+#include "Framework/Commands/UICommandList.h"
 #include "Framework/Docking/LayoutService.h"
 #include "Framework/Docking/TabManager.h"
 #include "Framework/Notifications/NotificationManager.h"
@@ -59,6 +61,14 @@ namespace UE::Audio::Insights
 		static const FText SaveSnapshotNotificationText = LOCTEXT("AudioDashboard_SnapshotButton_NotificationText", "Audio Insights Snapshot saved.");
 		static const FText SaveSnapshotNotificationSubText = LOCTEXT("AudioDashboard_SnapshotButton_NotificationSubText", "A snapshot .utrace with the most recent events has been saved to your trace server.");
 
+		static const FText SaveBookmarkText = LOCTEXT("AudioDashboard_BookmarkButton_Text", "Bookmark");
+		static const FText SaveBookmarkTooltipText = LOCTEXT("AudioDashboard_BookmarkButton_TooltipText", "Ctrl+M - Creates an 'Audio Insights' bookmark in the .utrace file generated by Save Snapshot or Live Trace at the timestamp the button was pressed.");
+
+		static const FText BookmarkLabelText = LOCTEXT("AudioDashboard_BookmarkLabelTitle_Text", "Audio Insights Bookmark");
+		static const FText SaveBookmarkNotificationText = LOCTEXT("AudioDashboard_BookmarkButton_NotificationText", "Audio Insights Bookmark created.");
+		static const FText SaveBookmarkNotificationDirectTraceSubText = LOCTEXT("AudioDashboard_BookmarkButton_NotificationSubText", "bookmark has been added to the Audio Insights Snapshot.");
+		static const FText SaveBookmarkNotificationLiveTraceSubText = LOCTEXT("AudioDashboard_BookmarkButton_NotificationSubText", "bookmark has been added to the live .utrace file on your trace server.");
+
 		static const FText WorldFilterText = LOCTEXT("AudioDashboard_WorldFilter_Text", "World Filter:");
 		static const FText WorldFilterTooltipText = LOCTEXT("AudioDashboard_WorldFilter_TooltipText", "Select world(s) to monitor (worlds may share audio output).");
 
@@ -110,6 +120,20 @@ namespace UE::Audio::Insights
 		if (!TraceModule.IsAudioInsightsTrace())
 		{
 			TraceModule.StartTraceAnalysis(bOnlyTraceAudioChannels);
+			BookmarkIndex = 0;
+		}
+	}
+
+	void FEditorDashboardFactory::CreateTraceBookmark()
+	{
+		using namespace EditorDashboardFactoryPrivate;
+
+		IAudioInsightsTraceModule& TraceModule = FAudioInsightsEditorModule::GetChecked().GetTraceModule();
+
+		if (TraceModule.IsTraceAnalysisActive())
+		{
+			TRACE_BOOKMARK(TEXT("%s %d"), *BookmarkLabelText.ToString(), BookmarkIndex);
+			BookmarkIndex++;
 		}
 	}
 
@@ -199,11 +223,29 @@ namespace UE::Audio::Insights
 		}
 	}
 
+	void FEditorDashboardFactory::BindCommands()
+	{
+		CommandList = MakeShared<FUICommandList>();
+
+		const FDashboardAssetCommands& Commands = FDashboardAssetCommands::Get();
+		CommandList->MapAction(Commands.GetTraceBookmarkCommand(), FExecuteAction::CreateRaw(this, &FEditorDashboardFactory::CreateTraceBookmark));
+	}
+
 	::Audio::FDeviceId FEditorDashboardFactory::GetDeviceId() const
 	{
 		return ActiveDeviceId;
 	}
 
+	FEditorDashboardFactory::FEditorDashboardFactory()
+	{
+		FDashboardAssetCommands::Register();
+	}
+
+	FEditorDashboardFactory::~FEditorDashboardFactory()
+	{
+		FDashboardAssetCommands::Unregister();
+	}
+
 	TSharedRef<SDockTab> FEditorDashboardFactory::MakeDockTabWidget(const FSpawnTabArgs& Args)
 	{
 		using namespace EditorDashboardFactoryPrivate;
@@ -641,6 +683,71 @@ namespace UE::Audio::Insights
 						]
 					]
 				]
+				// Trace Bookmark Button
+				+ SHorizontalBox::Slot()
+				.HAlign(HAlign_Left)
+				.VAlign(VAlign_Center)
+				.AutoWidth()
+				.Padding(12.0f, 0.0f, 0.0f, 0.0f)
+				[
+					SNew(SButton)
+					.ButtonStyle(&FAppStyle::Get().GetWidgetStyle<FButtonStyle>("SimpleButton"))
+					.ToolTipText(SaveBookmarkTooltipText)
+					.IsEnabled_Lambda([this]()
+					{
+						IAudioInsightsTraceModule& TraceModule = FAudioInsightsEditorModule::GetChecked().GetTraceModule();
+						return TraceModule.IsTraceAnalysisActive();
+					})
+					.OnClicked_Lambda([this]()
+					{
+						CreateTraceBookmark();
+
+						IAudioInsightsTraceModule& TraceModule = FAudioInsightsEditorModule::GetChecked().GetTraceModule();
+
+						// Show floating notification
+						FNotificationInfo Info(SaveBookmarkNotificationText);
+
+						const FText SubText = TraceModule.IsTraceAnalysisActive() ? SaveBookmarkNotificationLiveTraceSubText : SaveBookmarkNotificationDirectTraceSubText;
+
+						static const FTextFormat SubTextFormat = FTextFormat::FromString(TEXT("{0} {1} {2}"));
+						Info.SubText = FText::Format(SubTextFormat, BookmarkLabelText, BookmarkIndex, SaveBookmarkNotificationLiveTraceSubText);
+
+						Info.bFireAndForget  = true;
+						Info.FadeOutDuration = 1.0f;
+						Info.ExpireDuration  = 4.0f;
+
+						FSlateNotificationManager::Get().AddNotification(Info);
+
+						return FReply::Handled();
+					})
+					[
+						SNew(SBox)
+						.HeightOverride(16.0f)
+						[
+							SNew(SHorizontalBox)
+							// Save Bookmark icon
+							+ SHorizontalBox::Slot()
+							.VAlign(VAlign_Center)
+							.AutoWidth()
+							[
+								SNew(SImage)
+								.Image(FSlateStyle::Get().GetBrush("AudioInsights.Icon.Trace.Bookmark"))
+								.ColorAndOpacity(FSlateColor::UseForeground())
+							]
+							// Save bookmark text
+							+ SHorizontalBox::Slot()
+							.VAlign(VAlign_Center)
+							.AutoWidth()
+							.Padding(6.0f, 1.0f, 0.0f, 0.0f)
+							[
+								SNew(STextBlock)
+								.Justification(ETextJustify::Center)
+								.Text(SaveBookmarkText)
+								.Font(IPropertyTypeCustomizationUtils::GetRegularFont())
+							]
+						]
+					]
+				]
 				// World filter label
 				+ SHorizontalBox::Slot()
 				.HAlign(HAlign_Right)
diff --git a/Engine/Plugins/AudioInsights/Source/AudioInsightsEditor/Private/AudioInsightsEditorDashboardFactory.h b/Engine/Plugins/AudioInsights/Source/AudioInsightsEditor/Private/AudioInsightsEditorDashboardFactory.h
--- a/Engine/Plugins/AudioInsights/Source/AudioInsightsEditor/Private/AudioInsightsEditorDashboardFactory.h
+++ b/Engine/Plugins/AudioInsights/Source/AudioInsightsEditor/Private/AudioInsightsEditorDashboardFactory.h
@@ -12,13 +12,15 @@
 #include "Widgets/Docking/SDockTab.h"
 #include "Widgets/Input/SComboBox.h"
 
+class FUICommandList;
+
 namespace UE::Audio::Insights
 {
 	class FEditorDashboardFactory : public IDashboardFactory, public TSharedFromThis<FEditorDashboardFactory>
 	{
 	public:
-		FEditorDashboardFactory() = default;
-		virtual ~FEditorDashboardFactory() = default;
+		FEditorDashboardFactory();
+		virtual ~FEditorDashboardFactory();
 
 		TSharedRef<SDockTab> MakeDockTabWidget(const FSpawnTabArgs& Args);
 
@@ -37,10 +39,13 @@ namespace UE::Audio::Insights
 		void OnDeviceDestroyed(::Audio::FDeviceId InDeviceId);
 
 		void OnTraceStarted(FTraceAuxiliary::EConnectionType InTraceType, const FString& InTraceDestination);
+		void CreateTraceBookmark();
 
 		void RefreshDeviceSelector();
 		void ResetDelegates();
 
+		void BindCommands();
+
 		TSharedRef<SWidget> MakeMenuBarWidget();
 		TSharedRef<SWidget> MakeMainToolbarWidget();
 
@@ -67,10 +72,13 @@ namespace UE::Audio::Insights
 		TSharedPtr<FTabManager> DashboardTabManager;
 		TSharedPtr<FWorkspaceItem> DashboardWorkspace;
 
+		TSharedPtr<FUICommandList> CommandList;
+
 		TArray<TSharedPtr<::Audio::FDeviceId>> AudioDeviceIds;
 		TSharedPtr<SComboBox<TSharedPtr<::Audio::FDeviceId>>> AudioDeviceComboBox;
 
 		::Audio::FDeviceId ActiveDeviceId = INDEX_NONE;
+		int32 BookmarkIndex = 0;
 
 		TMap<FName, TSharedPtr<IDashboardViewFactory>> DashboardViewFactories;
 	};
