[Audio Insights] Set bIsAudioInsightsTrace before starting a trace to avoid StartTraceAnalysis to be called twice (this was making caching trace channels incorrectly) + adjusted "Only trace audio channels during PIE" checkbox text (since tracing doesn't start when PIE starts anymore)

diff --git a/Engine/Plugins/AudioInsights/Source/AudioInsightsEditor/Private/AudioInsightsEditorDashboardFactory.cpp b/Engine/Plugins/AudioInsights/Source/AudioInsightsEditor/Private/AudioInsightsEditorDashboardFactory.cpp
--- a/Engine/Plugins/AudioInsights/Source/AudioInsightsEditor/Private/AudioInsightsEditorDashboardFactory.cpp
+++ b/Engine/Plugins/AudioInsights/Source/AudioInsightsEditor/Private/AudioInsightsEditorDashboardFactory.cpp
@@ -43,8 +43,8 @@ namespace UE::Audio::Insights
 
 		static const FText ViewResetLayoutText = LOCTEXT("AudioDashboard_View_ResetLayout_Text", "Reset Layout");
 
-		static const FText OnlyTraceAudioChannelsText = LOCTEXT("AudioDashboard_OnlyTraceAudioChannels_Text", "Only trace audio channels during PIE:");
-		static const FText OnlyTraceAudioChannelsTooltipText = LOCTEXT("AudioDashboard_OnlyTraceAudioChannels_TooltipText", "Disable all non-command line trace channels apart from Audio, Audio Mixer and CPU during PIE. This will reduce the file sizes of trace sessions while using Audio Insights.");
+		static const FText OnlyTraceAudioChannelsText = LOCTEXT("AudioDashboard_OnlyTraceAudioChannels_Text", "Only trace audio channels:");
+		static const FText OnlyTraceAudioChannelsTooltipText = LOCTEXT("AudioDashboard_OnlyTraceAudioChannels_TooltipText", "Disable all non-command line trace channels apart from Audio, Audio Mixer, Bookmark and CPU when starting a live monitoring / live trace. This will reduce the file size of snapshots / trace sessions while using Audio Insights.");
 
 		static const FText StartDirectTraceText = LOCTEXT("AudioDashboard_StartDirectTrace_Text", "Start Live Monitoring");
 		static const FText StartDirectTraceTooltipText = LOCTEXT("AudioDashboard_StartDirectTrace_TooltipText", "Starts a live monitoring, no .utrace file is recorded.");
@@ -507,7 +507,7 @@ namespace UE::Audio::Insights
 						UToolMenus::Get()->GenerateWidget(PlayWorldToolBarName, { FPlayWorldCommands::GlobalPlayWorldActions })
 					]
 				]
-				// Only trace audio channels during PIE Label
+				// Only trace audio channels label
 				+ SHorizontalBox::Slot()
 				.HAlign(HAlign_Right)
 				.VAlign(VAlign_Center)
@@ -518,7 +518,7 @@ namespace UE::Audio::Insights
 					.Text(OnlyTraceAudioChannelsText)
 					.Font(IPropertyTypeCustomizationUtils::GetRegularFont())
 				]
-				// Only trace audio channels during PIE checkbox
+				// Only trace audio channels checkbox
 				+ SHorizontalBox::Slot()
 				.HAlign(HAlign_Left)
 				.VAlign(VAlign_Center)
