[Audio Insights] Enable Ctrl + M keyboard command when using Audio Insights to create a bookmark inside .utrace files

diff --git a/Engine/Plugins/AudioInsights/Source/AudioInsightsEditor/Private/AudioInsightsEditorDashboardFactory.cpp b/Engine/Plugins/AudioInsights/Source/AudioInsightsEditor/Private/AudioInsightsEditorDashboardFactory.cpp
--- a/Engine/Plugins/AudioInsights/Source/AudioInsightsEditor/Private/AudioInsightsEditorDashboardFactory.cpp
+++ b/Engine/Plugins/AudioInsights/Source/AudioInsightsEditor/Private/AudioInsightsEditorDashboardFactory.cpp
@@ -95,6 +95,40 @@ namespace UE::Audio::Insights
 
 			return FText::FromString(WorldName);
 		}
+
+		// Helper class required to intercept OnKeyDown events
+		class SAudioInsightsTabContents : public SVerticalBox
+		{
+		public:
+			SLATE_BEGIN_ARGS(SAudioInsightsTabContents)
+			{}
+
+			SLATE_EVENT(FOnKeyDown, OnKeyDown)
+
+			SLATE_END_ARGS()
+
+			void Construct(const FArguments& InArgs)
+			{
+				OnKeyDownHandler = InArgs._OnKeyDown;
+
+				SVerticalBox::Construct(SVerticalBox::FArguments());
+			}
+
+			virtual FReply OnKeyDown(const FGeometry& MyGeometry, const FKeyEvent& InKeyEvent) override
+			{
+				if (OnKeyDownHandler.IsBound())
+				{
+					return OnKeyDownHandler.Execute(MyGeometry, InKeyEvent);
+				}
+
+				return SVerticalBox::OnKeyDown(MyGeometry, InKeyEvent);
+			}
+
+			virtual bool SupportsKeyboardFocus() const override { return true; }
+
+		private:
+			FOnKeyDown OnKeyDownHandler;
+		};
 	} // namespace EditorDashboardFactoryPrivate
 
 	void FEditorDashboardFactory::OnWorldRegisteredToAudioDevice(const UWorld* InWorld, ::Audio::FDeviceId InDeviceId)
@@ -134,6 +168,20 @@ namespace UE::Audio::Insights
 		{
 			TRACE_BOOKMARK(TEXT("%s %d"), *BookmarkLabelText.ToString(), BookmarkIndex);
 			BookmarkIndex++;
+
+			// Show floating notification
+			FNotificationInfo Info(SaveBookmarkNotificationText);
+
+			const FText SubText = TraceModule.IsTraceAnalysisActive() ? SaveBookmarkNotificationLiveTraceSubText : SaveBookmarkNotificationDirectTraceSubText;
+
+			static const FTextFormat SubTextFormat = FTextFormat::FromString(TEXT("{0} {1} {2}"));
+			Info.SubText = FText::Format(SubTextFormat, BookmarkLabelText, BookmarkIndex, SaveBookmarkNotificationLiveTraceSubText);
+
+			Info.bFireAndForget = true;
+			Info.FadeOutDuration = 1.0f;
+			Info.ExpireDuration = 4.0f;
+
+			FSlateNotificationManager::Get().AddNotification(Info);
 		}
 	}
 
@@ -239,6 +287,8 @@ namespace UE::Audio::Insights
 	FEditorDashboardFactory::FEditorDashboardFactory()
 	{
 		FDashboardAssetCommands::Register();
+
+		BindCommands();
 	}
 
 	FEditorDashboardFactory::~FEditorDashboardFactory()
@@ -280,27 +330,16 @@ namespace UE::Audio::Insights
 
 		const TSharedRef<FTabManager::FLayout> TabLayout = LoadLayoutFromConfig();
 
-		const TSharedRef<SWidget> TabContent = SNew(SVerticalBox)
-			+ SVerticalBox::Slot()
-			.AutoHeight()
-			[
-				MakeMenuBarWidget()
-			]
-			+ SVerticalBox::Slot()
-			.AutoHeight()
-			[
-				MakeMainToolbarWidget()
-			]
-			+ SVerticalBox::Slot()
-			.AutoHeight()
-			[
-				SNew(SBox)
-				.HeightOverride(4.0f)
-			]
-			+ SVerticalBox::Slot()
-			[
-				DashboardTabManager->RestoreFrom(TabLayout, Args.GetOwnerWindow()).ToSharedRef()
-			];
+		TSharedRef<SVerticalBox> TabContent = SNew(SAudioInsightsTabContents)
+			.OnKeyDown_Lambda([this](const FGeometry& MyGeometry, const FKeyEvent& InKeyEvent)
+			{
+				return (CommandList && CommandList->ProcessCommandBindings(InKeyEvent)) ? FReply::Handled() : FReply::Unhandled();
+			});
+
+		TabContent->AddSlot().AutoHeight()[ MakeMenuBarWidget() ];
+		TabContent->AddSlot().AutoHeight()[ MakeMainToolbarWidget() ];
+		TabContent->AddSlot().AutoHeight()[ SNew(SBox).HeightOverride(4.0f) ];
+		TabContent->AddSlot()[ DashboardTabManager->RestoreFrom(TabLayout, Args.GetOwnerWindow()).ToSharedRef() ];
 
 		DockTab->SetContent(TabContent);
 
@@ -702,22 +741,6 @@ namespace UE::Audio::Insights
 					{
 						CreateTraceBookmark();
 
-						IAudioInsightsTraceModule& TraceModule = FAudioInsightsEditorModule::GetChecked().GetTraceModule();
-
-						// Show floating notification
-						FNotificationInfo Info(SaveBookmarkNotificationText);
-
-						const FText SubText = TraceModule.IsTraceAnalysisActive() ? SaveBookmarkNotificationLiveTraceSubText : SaveBookmarkNotificationDirectTraceSubText;
-
-						static const FTextFormat SubTextFormat = FTextFormat::FromString(TEXT("{0} {1} {2}"));
-						Info.SubText = FText::Format(SubTextFormat, BookmarkLabelText, BookmarkIndex, SaveBookmarkNotificationLiveTraceSubText);
-
-						Info.bFireAndForget  = true;
-						Info.FadeOutDuration = 1.0f;
-						Info.ExpireDuration  = 4.0f;
-
-						FSlateNotificationManager::Get().AddNotification(Info);
-
 						return FReply::Handled();
 					})
 					[
