[AudioInsights]
* Deleted duplicated implementation of FSlateStyle::get

diff --git a/Engine/Plugins/AudioInsights/Source/AudioInsights/Private/AudioInsightsModule.cpp b/Engine/Plugins/AudioInsights/Source/AudioInsights/Private/AudioInsightsModule.cpp
--- a/Engine/Plugins/AudioInsights/Source/AudioInsights/Private/AudioInsightsModule.cpp
+++ b/Engine/Plugins/AudioInsights/Source/AudioInsights/Private/AudioInsightsModule.cpp
@@ -26,12 +26,6 @@ DEFINE_LOG_CATEGORY(LogAudioInsights);
 
 namespace UE::Audio::Insights
 {
-	FSlateStyle& FSlateStyle::Get()
-	{
-		static FSlateStyle InsightsStyle;
-		return InsightsStyle;
-	}
-
 	void FAudioInsightsModule::StartupModule()
 	{
 		// Don't run providers in any commandlet to avoid additional, unnecessary overhead as audio insights is dormant.
