[Audio Insights] TraceModule: added ExecuteConsoleCommand method so cvars can be sent from both Audio Insights editor and standalone

diff --git a/Engine/Plugins/AudioInsights/Source/AudioInsights/Private/AudioInsightsTraceModule.cpp b/Engine/Plugins/AudioInsights/Source/AudioInsights/Private/AudioInsightsTraceModule.cpp
--- a/Engine/Plugins/AudioInsights/Source/AudioInsights/Private/AudioInsightsTraceModule.cpp
+++ b/Engine/Plugins/AudioInsights/Source/AudioInsights/Private/AudioInsightsTraceModule.cpp
@@ -478,6 +478,31 @@ namespace UE::Audio::Insights
 	}
 #endif // !WITH_EDITOR
 
+	void FTraceModule::ExecuteConsoleCommand(const FString& InCommandStr)const 
+	{
+		ISessionServicesModule& SessionServicesModule = FModuleManager::LoadModuleChecked<ISessionServicesModule>("SessionServices");
+		TSharedPtr<ISessionManager> SessionManager = SessionServicesModule.GetSessionManager();
+
+		if (SessionManager.IsValid())
+		{
+#if WITH_EDITOR
+			for (const TSharedPtr<ISessionInstanceInfo>& SessionInstance : SessionManager->GetSelectedInstances())
+			{
+				if (SessionInstance.IsValid())
+				{
+					SessionInstance->ExecuteCommand(InCommandStr);
+				}
+			}
+#else
+			TSharedPtr<ISessionInstanceInfo> SessionInstance = SessionManager->GetInstance(InstanceID);
+			if (SessionInstance.IsValid())
+			{
+				SessionInstance->ExecuteCommand(InCommandStr);
+			}
+#endif // WITH_EDITOR
+		}
+	}
+
 	void FTraceModule::GetLoggers(TArray<const TCHAR *>& OutLoggers)
 	{
 		OutLoggers.Add(TEXT("Audio"));
diff --git a/Engine/Plugins/AudioInsights/Source/AudioInsights/Public/AudioInsightsTraceModule.h b/Engine/Plugins/AudioInsights/Source/AudioInsights/Public/AudioInsightsTraceModule.h
--- a/Engine/Plugins/AudioInsights/Source/AudioInsights/Public/AudioInsightsTraceModule.h
+++ b/Engine/Plugins/AudioInsights/Source/AudioInsights/Public/AudioInsightsTraceModule.h
@@ -51,6 +51,8 @@ namespace UE::Audio::Insights
 
 		UE_API virtual bool AudioChannelsCanBeManuallyEnabled() const override;
 
+		UE_API virtual void ExecuteConsoleCommand(const FString& InCommandStr) const override;
+
 #if !WITH_EDITOR
 		UE_API virtual void InitializeSessionInfo(const TraceServices::FSessionInfo& SessionInfo) override;
 		UE_API virtual void RequestChannelUpdate() override;
diff --git a/Engine/Plugins/AudioInsights/Source/AudioInsights/Public/IAudioInsightsTraceModule.h b/Engine/Plugins/AudioInsights/Source/AudioInsights/Public/IAudioInsightsTraceModule.h
--- a/Engine/Plugins/AudioInsights/Source/AudioInsights/Public/IAudioInsightsTraceModule.h
+++ b/Engine/Plugins/AudioInsights/Source/AudioInsights/Public/IAudioInsightsTraceModule.h
@@ -27,6 +27,8 @@ public:
 
 	virtual bool AudioChannelsCanBeManuallyEnabled() const = 0;
 
+	virtual void ExecuteConsoleCommand(const FString& InCommandStr) const = 0;
+
 #if !WITH_EDITOR
 	virtual void InitializeSessionInfo(const TraceServices::FSessionInfo& SessionInfo) = 0;
 	virtual void RequestChannelUpdate() = 0;
