[Audio Insights] Optimizations to SAudioCurveView::PaintCurves - Filter out repeated points before forwarding them on to FSlateDrawElement::MakeLines

diff --git a/Engine/Plugins/AudioInsights/Source/AudioInsights/Private/Views/SAudioCurveView.cpp b/Engine/Plugins/AudioInsights/Source/AudioInsights/Private/Views/SAudioCurveView.cpp
--- a/Engine/Plugins/AudioInsights/Source/AudioInsights/Private/Views/SAudioCurveView.cpp
+++ b/Engine/Plugins/AudioInsights/Source/AudioInsights/Private/Views/SAudioCurveView.cpp
@@ -464,6 +464,12 @@ int32 SAudioCurveView::PaintCurves(const FGeometry& AllottedGeometry, const FSla
 		Points.Reserve(CurvePoints.Num());
 
 		float PrevX = CurvePoints[0].Key;
+
+		// Keep track of the last 3 recorded Y values
+		// These are used later to detect constant lines in the plot and remove all middle points that would be redundant
+		float PrevY = CurvePoints[0].Value;
+		float LineStartY = CurvePoints[0].Value;
+
 		for (const FCurvePoint& Point : CurvePoints)
 		{
 			if (Point.Key - PrevX > LargeFrameTime && Points.Num() > 1)
@@ -483,10 +489,26 @@ int32 SAudioCurveView::PaintCurves(const FGeometry& AllottedGeometry, const FSla
 			}
 
 			const float X = RangeToScreen.InputToLocalX(Point.Key);
-			const float Y = ValueToLocalY(Size, Point.Value);
-			Points.Emplace(X, Y);
 
-			PrevX = Point.Key;
+			if (Points.Num() >= 2 && Point.Value == PrevY && Point.Value == LineStartY)
+			{
+				// Filter out unnecessary points in the plot line
+				// - if 3 or more points in a row have the same Y value, move the line end X value forward in the points array (Points.Last(0).X = X;)
+				// - this can heavily reduce the number of points sent to FSlateDrawElement::MakeLines to draw
+				Points.Last(0).X = X;
+				PrevX = Point.Key;
+			}
+			else
+			{
+				const float Y = ValueToLocalY(Size, Point.Value);
+				Points.Emplace(X, Y);
+
+				PrevX = Point.Key;
+
+				// Keep track of the previous values of Y to detect constant lines
+				LineStartY = PrevY;
+				PrevY = Point.Value;
+			}
 		}
 
 		FSlateDrawElement::MakeLines(
