[Audio Insights] Allow users to plot all audio data captured in a .utrace file when scrubbing back and forth with the time marker

diff --git a/Engine/Plugins/AudioInsights/Source/AudioInsights/Private/Views/SoundDashboardViewFactory.cpp b/Engine/Plugins/AudioInsights/Source/AudioInsights/Private/Views/SoundDashboardViewFactory.cpp
--- a/Engine/Plugins/AudioInsights/Source/AudioInsights/Private/Views/SoundDashboardViewFactory.cpp
+++ b/Engine/Plugins/AudioInsights/Source/AudioInsights/Private/Views/SoundDashboardViewFactory.cpp
@@ -996,6 +996,9 @@ namespace UE::Audio::Insights
 	{
 #if WITH_EDITOR
 		FSoundDashboardSettings::OnRequestReadSettings.Broadcast();
+#else
+		FAudioInsightsModule& AudioInsightsModule = FAudioInsightsModule::GetChecked();
+		AudioInsightsModule.GetTimingViewExtender().OnTimingViewTimeMarkerChanged.AddSP(this, &FSoundDashboardViewFactory::OnTimingViewTimeMarkerChanged);
 #endif // WITH_EDITOR
 
 		return SNew(SVerticalBox)
@@ -2723,6 +2726,11 @@ namespace UE::Audio::Insights
 
 		return Objects;
 	}
+#else
+	void FSoundDashboardViewFactory::OnTimingViewTimeMarkerChanged(double InTimeMarker)
+	{
+		OnUpdatePlotVisibility.Broadcast(DataViewEntries);
+	}
 #endif // WITH_EDITOR
 
 	TSharedRef<SWidget> FSoundDashboardViewFactory::MakeShowPlotWidget()
@@ -2856,6 +2864,10 @@ namespace UE::Audio::Insights
 				SetShowPlot(Entry, bShowPlot);
 			}
 		}
+
+#if !WITH_EDITOR
+		OnUpdatePlotVisibility.Broadcast(DataViewEntries);
+#endif // !WITH_EDITOR
 	}
 
 	void FSoundDashboardViewFactory::SetShowPlotRecursive(const TSharedPtr<IDashboardDataTreeViewEntry>& Entry, const bool bShowPlot)
diff --git a/Engine/Plugins/AudioInsights/Source/AudioInsights/Private/Views/SoundPlotsDashboardViewFactory.cpp b/Engine/Plugins/AudioInsights/Source/AudioInsights/Private/Views/SoundPlotsDashboardViewFactory.cpp
--- a/Engine/Plugins/AudioInsights/Source/AudioInsights/Private/Views/SoundPlotsDashboardViewFactory.cpp
+++ b/Engine/Plugins/AudioInsights/Source/AudioInsights/Private/Views/SoundPlotsDashboardViewFactory.cpp
@@ -54,6 +54,10 @@ namespace UE::Audio::Insights
 		PlotsWidget = MakeShared<FSoundPlotsWidgetView>(MoveTemp(ColumnInfo), [](const IDashboardDataTreeViewEntry& InEntry) -> bool { return CastEntry(InEntry).bIsPlotActive; });
 
 		OnProcessPlotData = SoundsDashboard->OnProcessPlotData.AddSP(this, &FSoundPlotsDashboardViewFactory::ProcessPlotData);
+
+#if !WITH_EDITOR
+		OnUpdatePlotVisibility = SoundsDashboard->OnUpdatePlotVisibility.AddSP(this, &FSoundPlotsDashboardViewFactory::UpdatePlotVisibility);
+#endif // !WITH_EDITOR
 	}
 
 	void FSoundPlotsDashboardViewFactory::ProcessPlotData(const TArray<TSharedPtr<IDashboardDataTreeViewEntry>>& DataViewEntries)
@@ -63,5 +67,15 @@ namespace UE::Audio::Insights
 			PlotsWidget->ProcessPlotData(DataViewEntries);
 		}
 	}
+
+#if !WITH_EDITOR
+	void FSoundPlotsDashboardViewFactory::UpdatePlotVisibility(const TArray<TSharedPtr<IDashboardDataTreeViewEntry>>& DataViewEntries)
+	{
+		if (PlotsWidget.IsValid())
+		{
+			PlotsWidget->UpdatePlotVisibility(DataViewEntries);
+		}
+	}
+#endif // !WITH_EDITOR
 } // namespace UE::Audio::Insights
 #undef LOCTEXT_NAMESPACE
diff --git a/Engine/Plugins/AudioInsights/Source/AudioInsights/Private/Views/SoundPlotsWidgetView.cpp b/Engine/Plugins/AudioInsights/Source/AudioInsights/Private/Views/SoundPlotsWidgetView.cpp
--- a/Engine/Plugins/AudioInsights/Source/AudioInsights/Private/Views/SoundPlotsWidgetView.cpp
+++ b/Engine/Plugins/AudioInsights/Source/AudioInsights/Private/Views/SoundPlotsWidgetView.cpp
@@ -17,12 +17,26 @@
 
 namespace UE::Audio::Insights
 {
-	FSoundPlotsWidgetView::FSoundPlotsWidgetView(TMap<FName, FSoundPlotsWidgetView::FPlotColumnInfo>&& InColumnInfo, const TFunctionRef<bool(const IDashboardDataTreeViewEntry& InEntry)> InIsPlotEnabledForEntryFunc)
+	FSoundPlotsWidgetView::FSoundPlotsWidgetView(TMap<FName, FSoundPlotsWidgetView::FPlotColumnInfo>&& InColumnInfo, const TFunction<bool(const IDashboardDataTreeViewEntry& InEntry)> InIsPlotEnabledForEntryFunc)
 		: ColumnInfo(InColumnInfo)
 		, IsPlotEnabledForEntryFunc(InIsPlotEnabledForEntryFunc)
 	{
 		FTraceModule& AudioInsightsTraceModule = static_cast<FTraceModule&>(FAudioInsightsModule::GetChecked().GetTraceModule());
 		AudioInsightsTraceModule.OnAnalysisStarting.AddRaw(this, &FSoundPlotsWidgetView::OnAnalysisStarting);
+
+#if !WITH_EDITOR
+		FAudioInsightsModule& AudioInsightsModule = FAudioInsightsModule::GetChecked();
+		const TSharedPtr<FAudioInsightsComponent> AudioInsightsComponent = AudioInsightsModule.GetAudioInsightsComponent();
+		if (AudioInsightsComponent.IsValid())
+		{
+			AudioInsightsComponent->OnSessionAnalysisCompleted.AddRaw(this, &FSoundPlotsWidgetView::OnSessionAnalysisCompleted);
+
+			if (!AudioInsightsComponent->GetIsLiveSession())
+			{
+				GameState = EGameState::LoadingTraceFile;
+			}
+		}
+#endif // !WITH_EDITOR
 	}
 
 	FSoundPlotsWidgetView::~FSoundPlotsWidgetView()
@@ -31,6 +45,15 @@ namespace UE::Audio::Insights
 		{
 			FTraceModule& TraceModule = static_cast<FTraceModule&>(FAudioInsightsModule::GetChecked().GetTraceModule());
 			TraceModule.OnAnalysisStarting.RemoveAll(this);
+
+#if !WITH_EDITOR
+			FAudioInsightsModule& AudioInsightsModule = FAudioInsightsModule::GetChecked();
+			const TSharedPtr<FAudioInsightsComponent> AudioInsightsComponent = AudioInsightsModule.GetAudioInsightsComponent();
+			if (AudioInsightsComponent.IsValid())
+			{
+				AudioInsightsComponent->OnSessionAnalysisCompleted.RemoveAll(this);
+			}
+#endif // !WITH_EDITOR
 		}
 	}
 
@@ -52,7 +75,6 @@ namespace UE::Audio::Insights
 		if (AudioInsightsComponent.IsValid())
 		{
 			AudioInsightsComponent->OnTabSpawn.AddSP(this, &FSoundPlotsWidgetView::OnAudioInsightsComponentTabSpawn);
-			AudioInsightsComponent->OnSessionAnalysisCompleted.AddSP(this, &FSoundPlotsWidgetView::OnSessionAnalysisCompleted);
 		}
 #endif // WITH_EDITOR
 
@@ -61,10 +83,28 @@ namespace UE::Audio::Insights
 
 	void FSoundPlotsWidgetView::ProcessPlotData(const TArray<TSharedPtr<IDashboardDataTreeViewEntry>>& DataViewEntries)
 	{
-		if (!PlotWidgetMetadataPerCurve.IsValid() || DataViewEntries.IsEmpty() || GameState == EGameState::Paused)
+#if !WITH_EDITOR
+		if (GameState == EGameState::ControlledByTimemarker && !bTimeMarkerMoved)
 		{
 			return;
 		}
+#endif // !WITH_EDITOR
+
+		if (DataViewEntries.IsEmpty() || GameState == EGameState::Paused)
+		{
+			return;
+		}
+
+		if (!PlotWidgetMetadataPerCurve.IsValid())
+		{
+#if WITH_EDITOR
+			return;
+#else
+			// In standalone, we still want to process any plot data while the .utrace file is loading in
+			// Ensure we create the plot widget and start collecting data from the point we start receiving data
+			ConstructPlotsWidgetContents();
+#endif // WITH_EDITOR
+		}
 
 		// Process new data
 		bool bHasNewMetadata = false;
@@ -80,28 +120,44 @@ namespace UE::Audio::Insights
 		}
 
 		// Remove old points and set curve data for each widget
+		const double PlotDrawMinLimitTimestamp = CurrentTimestamp - BeginTimestamp - (MaxPlotHistorySeconds + 0.2 /* extra grace time to avoid curve cuts being displayed */);
 #if WITH_EDITOR
-		const double PlotDrawLimitTimestamp = CurrentTimestamp - BeginTimestamp - (MaxPlotHistorySeconds + 0.2 /* extra grace time to avoid curve cuts being displayed */);
+		
+		const TSharedPtr<FPointDataPerCurveMap> CurveDataMapPtr = *PlotWidgetCurveIdToPointDataMapPerColumn.Find(SelectedPlotColumnName);
+#else
+		// Take a deep copy of the curve data in standalone - we need to retain the original data and filter which points to process
+		const TSharedPtr<FPointDataPerCurveMap> CurveDataMapPtr = MakeShared<FPointDataPerCurveMap>(*(*PlotWidgetCurveIdToPointDataMapPerColumn.Find(SelectedPlotColumnName)).Get());
+		const double PlotDrawMaxLimitTimestamp = CurrentTimestamp + 0.2 /* extra grace time to avoid curve cuts being displayed */;
 #endif // WITH_EDITOR
 
-		const TSharedPtr<FPointDataPerCurveMap> CurveDataMapPtr = *PlotWidgetCurveIdToPointDataMapPerColumn.Find(SelectedPlotColumnName);
 		if (CurveDataMapPtr.IsValid())
 		{
-#if WITH_EDITOR
 			// Remove points that are older than max history limit from the most recent timestamp
 			for (auto& [CurveId, CurvePoints] : *CurveDataMapPtr)
 			{
-				const int32 FoundIndex = CurvePoints.IndexOfByPredicate([&PlotDrawLimitTimestamp](const FDataPoint& InDataPoint)
+				const int32 FoundMinIndex = CurvePoints.IndexOfByPredicate([&PlotDrawMinLimitTimestamp](const FDataPoint& InDataPoint)
 				{
-					return InDataPoint.Key >= PlotDrawLimitTimestamp;
+					return InDataPoint.Key >= PlotDrawMinLimitTimestamp;
 				});
 
-				if (FoundIndex > 0)
+				if (FoundMinIndex > 0)
 				{
-					CurvePoints.RemoveAt(0, FoundIndex, EAllowShrinking::No);
+					CurvePoints.RemoveAt(0, FoundMinIndex, EAllowShrinking::No);
 				}
+
+#if !WITH_EDITOR
+				// In standalone, also remove any points that are beyond the current timestamp
+				const int32 FoundMaxIndex = CurvePoints.IndexOfByPredicate([&PlotDrawMaxLimitTimestamp](const FDataPoint& InDataPoint)
+				{
+					return InDataPoint.Key >= PlotDrawMaxLimitTimestamp;
+				});
+
+				if (FoundMaxIndex > 0)
+				{
+					CurvePoints.RemoveAt(FoundMaxIndex, CurvePoints.Num() - FoundMaxIndex - 1, EAllowShrinking::No);
+				}
+#endif // !WITH_EDITOR
 			}
-#endif // WITH_EDITOR
 
 			if (!CurveDataMapPtr->IsEmpty())
 			{
@@ -109,6 +165,28 @@ namespace UE::Audio::Insights
 				PlotWidget->SetYValueFormattingOptions(*GetPlotColumnNumberFormat(SelectedPlotColumnName));
 			}
 		}
+
+#if !WITH_EDITOR
+		bTimeMarkerMoved = false;
+#endif // !WITH_EDITOR
+	}
+
+	void FSoundPlotsWidgetView::UpdatePlotVisibility(const TArray<TSharedPtr<IDashboardDataTreeViewEntry>>& DataViewEntries)
+	{
+		if (!PlotWidgetMetadataPerCurve.IsValid() || DataViewEntries.IsEmpty())
+		{
+			return;
+		}
+
+		for (auto& [ID, PlotMetaData] : *PlotWidgetMetadataPerCurve)
+		{
+			PlotMetaData.CurveColor.A = 0.0f;
+		}
+
+		for (const TSharedPtr<IDashboardDataTreeViewEntry>& Entry : DataViewEntries)
+		{
+			UpdatePlotVisibilityRecursive(Entry);
+		}
 	}
 
 	void FSoundPlotsWidgetView::OnAnalysisStarting(const double Timestamp)
@@ -149,31 +227,82 @@ namespace UE::Audio::Insights
 	void FSoundPlotsWidgetView::OnAudioInsightsComponentTabSpawn()
 	{
 		const TSharedPtr<const FAudioInsightsComponent> AudioInsightsComponent = FAudioInsightsModule::GetChecked().GetAudioInsightsComponent();
-		if (AudioInsightsComponent.IsValid())
+		if (AudioInsightsComponent.IsValid() && AudioInsightsComponent->GetIsLiveSession())
 		{
-			GameState = AudioInsightsComponent->GetIsLiveSession() ? EGameState::Running : EGameState::Stopped;
+			GameState = EGameState::Running;
 		}
 	}
 
 	void FSoundPlotsWidgetView::OnSessionAnalysisCompleted()
 	{
-		GameState = EGameState::Stopped;
+		if (GameState == EGameState::LoadingTraceFile && PlotWidgetMetadataPerCurve.IsValid())
+		{
+			for (auto& [SourceID, MetaData] : *PlotWidgetMetadataPerCurve)
+			{
+				MetaData.CurveColor.A = 0.0f;
+			}
+
+			if (PlotWidget.IsValid())
+			{
+				PlotWidget->SetCurvesMetadata(PlotWidgetMetadataPerCurve);
+			}
+
+			GameState = EGameState::ControlledByTimemarker;
+		}
+		else
+		{
+			GameState = EGameState::Stopped;
+		}
 	}
 
 	void FSoundPlotsWidgetView::OnTimingViewTimeMarkerChanged(double InTimeMarker)
 	{
 		CurrentTimestamp = InTimeMarker;
+		GameState = EGameState::ControlledByTimemarker;
+		bTimeMarkerMoved = true;
 	}
 #endif // WITH_EDITOR
 
 	TSharedRef<SWidget> FSoundPlotsWidgetView::MakePlotsWidget()
+	{
+		if (!PlotWidget.IsValid() || !ComboBox.IsValid())
+		{
+			ConstructPlotsWidgetContents();
+		}
+
+		return SNew(SVerticalBox)
+			.Clipping(EWidgetClipping::ClipToBounds)
+			+ SVerticalBox::Slot()
+			.AutoHeight()
+			.HAlign(HAlign_Fill)
+			[
+				SNew(SSimpleTimeSlider)
+				.ViewRange_Lambda([this]() { return GetViewRange(); })
+				.ClampRangeHighlightSize(0.0f) // Hide clamp range
+				.ScrubPosition_Lambda([]() { return TNumericLimits<double>::Lowest(); }) // Hide scrub
+				.PixelSnappingMethod(EWidgetPixelSnapping::Disabled)
+			]
+			+ SVerticalBox::Slot()
+			.AutoHeight()
+			.HAlign(HAlign_Fill)
+			[
+				ComboBox.ToSharedRef()
+			]
+			+ SVerticalBox::Slot()
+			.HAlign(HAlign_Fill)
+			[
+				PlotWidget.ToSharedRef()
+			];
+	}
+
+	void FSoundPlotsWidgetView::ConstructPlotsWidgetContents()
 	{
 		// Initialize column options and initially selected column
 		ColumnInfo.GenerateKeyArray(ColumnNames);
 
 		if (ColumnNames.IsEmpty())
 		{
-			return SNullWidget::NullWidget;
+			return;
 		}
 
 		SelectedPlotColumnName = ColumnNames[0];
@@ -197,60 +326,36 @@ namespace UE::Audio::Insights
 		PlotWidget->SetYValueFormattingOptions(*GetPlotColumnNumberFormat(SelectedPlotColumnName));
 		PlotWidget->SetInitYAxisRange(GetPlotColumnDefaultYRange(SelectedPlotColumnName));
 
-		// Create plot column combo box widgets  
-		auto CreatePlotColumnComboBoxWidget = [this]()
-		{
-			return SNew(SComboBox<FName>)
-				.ToolTipText(LOCTEXT("SoundPlotsWidgetView_SelectPlotColumnDescription", "Select a column from the table to plot."))
-				.OptionsSource(&ColumnNames)
-				.OnGenerateWidget_Lambda([this](const FName& ColumnName)
-				{
-					return SNew(STextBlock)
-					.Font(FAppStyle::GetFontStyle(TEXT("PropertyWindow.NormalFont")))
-					.Text(GetPlotColumnDisplayName(ColumnName));
-				})
-				.OnSelectionChanged_Lambda([this](FName NewColumnName, ESelectInfo::Type)
+		SAssignNew(ComboBox, SComboBox<FName>)
+			.ToolTipText(LOCTEXT("SoundPlotsWidgetView_SelectPlotColumnDescription", "Select a column from the table to plot."))
+			.OptionsSource(&ColumnNames)
+			.OnGenerateWidget_Lambda([this](const FName& ColumnName)
+			{
+				return SNew(STextBlock)
+				.Font(FAppStyle::GetFontStyle(TEXT("PropertyWindow.NormalFont")))
+				.Text(GetPlotColumnDisplayName(ColumnName));
+			})
+			.OnSelectionChanged_Lambda([this](FName NewColumnName, ESelectInfo::Type)
+			{
+				SelectedPlotColumnName = NewColumnName;
+				if (TSharedPtr<FPointDataPerCurveMap>* DataMap = PlotWidgetCurveIdToPointDataMapPerColumn.Find(NewColumnName))
 				{
-					SelectedPlotColumnName = NewColumnName;
-					if (TSharedPtr<FPointDataPerCurveMap>* DataMap = PlotWidgetCurveIdToPointDataMapPerColumn.Find(NewColumnName))
-					{
-						PlotWidget->SetCurvesPointData(*DataMap);
-						PlotWidget->SetYValueFormattingOptions(*GetPlotColumnNumberFormat(NewColumnName));
-						PlotWidget->SetInitYAxisRange(GetPlotColumnDefaultYRange(NewColumnName));
-					}
-				})
-				[
-					SNew(STextBlock)
-					.Font(FAppStyle::GetFontStyle(TEXT("PropertyWindow.NormalFont")))
-					.Text_Lambda([this]()
-					{
-						return GetPlotColumnDisplayName(SelectedPlotColumnName);
-					})
-				];
-		};
+					PlotWidget->SetCurvesPointData(*DataMap);
+					PlotWidget->SetYValueFormattingOptions(*GetPlotColumnNumberFormat(NewColumnName));
+					PlotWidget->SetInitYAxisRange(GetPlotColumnDefaultYRange(NewColumnName));
 
-		return SNew(SVerticalBox)
-			.Clipping(EWidgetClipping::ClipToBounds)
-			+ SVerticalBox::Slot()
-			.AutoHeight()
-			.HAlign(HAlign_Fill)
-			[
-				SNew(SSimpleTimeSlider)
-				.ViewRange_Lambda([this]() { return GetViewRange(); })
-				.ClampRangeHighlightSize(0.0f) // Hide clamp range
-				.ScrubPosition_Lambda([]() { return TNumericLimits<double>::Lowest(); }) // Hide scrub
-				.PixelSnappingMethod(EWidgetPixelSnapping::Disabled)
-			]
-			+ SVerticalBox::Slot()
-			.AutoHeight()
-			.HAlign(HAlign_Fill)
-			[
-				CreatePlotColumnComboBoxWidget()
-			]
-			+ SVerticalBox::Slot()
-			.HAlign(HAlign_Fill)
+#if !WITH_EDITOR
+					bTimeMarkerMoved = true;
+#endif
+				}
+			})
 			[
-				PlotWidget.ToSharedRef()
+				SNew(STextBlock)
+				.Font(FAppStyle::GetFontStyle(TEXT("PropertyWindow.NormalFont")))
+				.Text_Lambda([this]()
+				{
+					return GetPlotColumnDisplayName(SelectedPlotColumnName);
+				})
 			];
 	}
 
@@ -262,7 +367,7 @@ namespace UE::Audio::Insights
 		}
 
 		const uint32 SourceId = Entry->GetEntryID();
-		const bool bPlotIsEnabled = IsPlotEnabledForEntryFunc(*Entry);
+		const bool bPlotIsEnabled = GetIsPlotEnabledForEntry(*Entry);
 
 		if (SourceId == INDEX_NONE || !bPlotIsEnabled)
 		{
@@ -322,11 +427,15 @@ namespace UE::Audio::Insights
 					}
 				}
 
-				CurrentTimestamp = FMath::Max(CurrentTimestamp, Timestamp);
-
 #if WITH_EDITOR
+				CurrentTimestamp = FMath::Max(CurrentTimestamp, Timestamp);
 				const double DataPointTime = Timestamp - BeginTimestamp;
 #else
+				if (GameState != EGameState::ControlledByTimemarker)
+				{
+					CurrentTimestamp = FMath::Max(CurrentTimestamp, Timestamp);
+				}
+
 				const double DataPointTime = Timestamp;
 #endif // WITH_EDITOR
 				DataPoints->Emplace(DataPointTime, UnitConversionFunc.IsSet() ? UnitConversionFunc(Value) : Value);
@@ -369,6 +478,24 @@ namespace UE::Audio::Insights
 		CurrentTimestamp = TNumericLimits<double>::Lowest();
 	}
 
+	void FSoundPlotsWidgetView::UpdatePlotVisibilityRecursive(const TSharedPtr<IDashboardDataTreeViewEntry>& Entry)
+	{
+		if (!Entry.IsValid())
+		{
+			return;
+		}
+
+		if (FPlotCurveMetadata* PlotMetaData = PlotWidgetMetadataPerCurve->Find(Entry->GetEntryID()))
+		{
+			PlotMetaData->CurveColor.A = GetIsPlotEnabledForEntry(*Entry) ? 1.0f : 0.0f;
+		}
+
+		for (const TSharedPtr<IDashboardDataTreeViewEntry>& Child : Entry->Children)
+		{
+			UpdatePlotVisibilityRecursive(Child);
+		}
+	}
+
 	const TFunction<const::Audio::TCircularAudioBuffer<FDataPoint>& (const IDashboardDataTreeViewEntry& InData)> FSoundPlotsWidgetView::GetPlotColumnDataFunc(const FName& ColumnName)
 	{
 		return ColumnInfo.Find(ColumnName)->DataFunc;
@@ -462,6 +589,15 @@ namespace UE::Audio::Insights
 
 		return TRange<double>(RangeUpperBound - MaxPlotHistorySeconds, RangeUpperBound);
 	}
+
+	bool FSoundPlotsWidgetView::GetIsPlotEnabledForEntry(const IDashboardDataTreeViewEntry& InEntry) const
+	{
+#if WITH_EDITOR
+		return IsPlotEnabledForEntryFunc.IsSet() && IsPlotEnabledForEntryFunc(InEntry);
+#else
+		return GameState == EGameState::LoadingTraceFile || (IsPlotEnabledForEntryFunc.IsSet() && IsPlotEnabledForEntryFunc(InEntry));
+#endif // WITH_EDITOR
+	}
 } // namespace UE::Audio::Insights
 
 #undef LOCTEXT_NAMESPACE
diff --git a/Engine/Plugins/AudioInsights/Source/AudioInsights/Public/Views/SoundDashboardViewFactory.h b/Engine/Plugins/AudioInsights/Source/AudioInsights/Public/Views/SoundDashboardViewFactory.h
--- a/Engine/Plugins/AudioInsights/Source/AudioInsights/Public/Views/SoundDashboardViewFactory.h
+++ b/Engine/Plugins/AudioInsights/Source/AudioInsights/Public/Views/SoundDashboardViewFactory.h
@@ -116,6 +116,11 @@ namespace UE::Audio::Insights
 		DECLARE_MULTICAST_DELEGATE_OneParam(FOnProcessPlotData, const TArray<TSharedPtr<IDashboardDataTreeViewEntry>>& /*DataViewEntries*/);
 		FOnProcessPlotData OnProcessPlotData;
 
+#if !WITH_EDITOR
+		DECLARE_MULTICAST_DELEGATE_OneParam(FOnUpdatePlotVisibility, const TArray<TSharedPtr<IDashboardDataTreeViewEntry>>& /*DataViewEntries*/);
+		FOnUpdatePlotVisibility OnUpdatePlotVisibility;
+#endif // !WITH_EDITOR
+
 	private:
 		UE_API void BindCommands();
 
@@ -174,6 +179,8 @@ namespace UE::Audio::Insights
 		UE_API void ToggleMuteSoloEntries(const TArray<TSharedPtr<IDashboardDataTreeViewEntry>>& InEntries, const EMuteSoloMode InMuteSoloMode);
 
 		UE_API TArray<TObjectPtr<UObject>> GetSelectedEditableAssets() const;
+#else
+		void OnTimingViewTimeMarkerChanged(double InTimeMarker);
 #endif // WITH_EDITOR
 
 		UE_API TSharedRef<SWidget> MakeShowPlotWidget();
diff --git a/Engine/Plugins/AudioInsights/Source/AudioInsights/Public/Views/SoundPlotsDashboardViewFactory.h b/Engine/Plugins/AudioInsights/Source/AudioInsights/Public/Views/SoundPlotsDashboardViewFactory.h
--- a/Engine/Plugins/AudioInsights/Source/AudioInsights/Public/Views/SoundPlotsDashboardViewFactory.h
+++ b/Engine/Plugins/AudioInsights/Source/AudioInsights/Public/Views/SoundPlotsDashboardViewFactory.h
@@ -30,6 +30,12 @@ namespace UE::Audio::Insights
 
 		void ProcessPlotData(const TArray<TSharedPtr<IDashboardDataTreeViewEntry>>& DataViewEntries);
 
+#if !WITH_EDITOR
+		void UpdatePlotVisibility(const TArray<TSharedPtr<IDashboardDataTreeViewEntry>>& DataViewEntries);
+
+		FDelegateHandle OnUpdatePlotVisibility;
+#endif // !WITH_EDITOR
+
 		FDelegateHandle OnProcessPlotData;
 	};
 } // namespace UE::Audio::Insights
diff --git a/Engine/Plugins/AudioInsights/Source/AudioInsights/Public/Views/SoundPlotsWidgetView.h b/Engine/Plugins/AudioInsights/Source/AudioInsights/Public/Views/SoundPlotsWidgetView.h
--- a/Engine/Plugins/AudioInsights/Source/AudioInsights/Public/Views/SoundPlotsWidgetView.h
+++ b/Engine/Plugins/AudioInsights/Source/AudioInsights/Public/Views/SoundPlotsWidgetView.h
@@ -8,6 +8,8 @@
 
 #define UE_API AUDIOINSIGHTS_API
 
+template <typename ItemType> class SComboBox;
+
 namespace UE::Audio::Insights
 {
 	class FSoundPlotsWidgetView : public TSharedFromThis<FSoundPlotsWidgetView>
@@ -24,11 +26,12 @@ namespace UE::Audio::Insights
 			const TFunction<FVector2f()> GetDefaultYDataRangeFunc;
 		};
 
-		UE_API FSoundPlotsWidgetView(TMap<FName, FSoundPlotsWidgetView::FPlotColumnInfo>&& InColumnInfo, const TFunctionRef<bool(const IDashboardDataTreeViewEntry& InEntry)> InIsPlotEnabledForEntryFunc);
+		UE_API FSoundPlotsWidgetView(TMap<FName, FSoundPlotsWidgetView::FPlotColumnInfo>&& InColumnInfo, const TFunction<bool(const IDashboardDataTreeViewEntry& InEntry)> InIsPlotEnabledForEntryFunc);
 		UE_API virtual ~FSoundPlotsWidgetView();
 
 		UE_API TSharedRef<SWidget> MakeWidget();
 		UE_API void ProcessPlotData(const TArray<TSharedPtr<IDashboardDataTreeViewEntry>>& DataViewEntries);
+		UE_API void UpdatePlotVisibility(const TArray<TSharedPtr<IDashboardDataTreeViewEntry>>& DataViewEntries);
 
 	private:
 		using FPlotCurvePoint = SAudioCurveView::FCurvePoint;
@@ -51,10 +54,13 @@ namespace UE::Audio::Insights
 #endif // WITH_EDITOR
 
 		TSharedRef<SWidget> MakePlotsWidget();
+		void ConstructPlotsWidgetContents();
 
 		void CollectDataRecursive(const TSharedPtr<IDashboardDataTreeViewEntry>& Entry, bool& bOutHasNewMetadata);
 		void ResetPlots();
 
+		void UpdatePlotVisibilityRecursive(const TSharedPtr<IDashboardDataTreeViewEntry>& Entry);
+
 		// Column information used by plot widgets, keyed by column name. These keys should be a subset of the keys in GetColumns(). 
 		const TFunction<const ::Audio::TCircularAudioBuffer<FDataPoint>& (const IDashboardDataTreeViewEntry& InEntry)> GetPlotColumnDataFunc(const FName& ColumnName);
 		const TFunction<float (const float)> GetPlotColumnUnitConversionFunc(const FName& ColumnName);
@@ -64,6 +70,8 @@ namespace UE::Audio::Insights
 
 		TRange<double> GetViewRange();
 
+		bool GetIsPlotEnabledForEntry(const IDashboardDataTreeViewEntry& InEntry) const;
+
 		// Maximum amount of data history kept for plots (in seconds)
 		inline static constexpr double MaxPlotHistorySeconds = 5.0;
 
@@ -84,7 +92,7 @@ namespace UE::Audio::Insights
 		const TMap<FName, FSoundPlotsWidgetView::FPlotColumnInfo> ColumnInfo;
 
 		// Func to enquire whether plotting is enabled for an entry
-		const TFunctionRef<bool(const IDashboardDataTreeViewEntry& InEntry)> IsPlotEnabledForEntryFunc;
+		const TFunction<bool(const IDashboardDataTreeViewEntry& InEntry)> IsPlotEnabledForEntryFunc;
 
 		double BeginTimestamp = TNumericLimits<double>::Max();
 		double CurrentTimestamp = TNumericLimits<double>::Lowest();
@@ -93,17 +101,24 @@ namespace UE::Audio::Insights
 		{
 			Running,
 			Stopped,
-			Paused
+			Paused,
+#if !WITH_EDITOR
+			LoadingTraceFile,
+			ControlledByTimemarker
+#endif
 		};
 
 		EGameState GameState = EGameState::Stopped;
 
 		TSharedPtr<SAudioCurveView> PlotWidget;
+		TSharedPtr<SComboBox<FName>> ComboBox;
 		FName SelectedPlotColumnName;
 
 #if !WITH_EDITOR
 		double PreviousTime = 0.0;
 		double CurrentRangeUpperBound = 0.0;
+
+		bool bTimeMarkerMoved = false;
 #endif // !WITH_EDITOR
 	};
 } // namespace UE::Audio::Insights
