[Audio Insights] If a sound virtualizes and then realizes, keep plotting active if it was previously.

diff --git a/Engine/Plugins/AudioInsights/Source/AudioInsights/Private/Providers/SoundTraceProvider.cpp b/Engine/Plugins/AudioInsights/Source/AudioInsights/Private/Providers/SoundTraceProvider.cpp
--- a/Engine/Plugins/AudioInsights/Source/AudioInsights/Private/Providers/SoundTraceProvider.cpp
+++ b/Engine/Plugins/AudioInsights/Source/AudioInsights/Private/Providers/SoundTraceProvider.cpp
@@ -329,6 +329,7 @@ namespace UE::Audio::Insights
 		ActiveSoundToEntryKeysMap.Empty();
 		EntriesTimingOut.Empty();
 		SoundsStoppedBeforeStart.Empty();
+		PlottingSoundEntries.Empty();
 
 		return new FSoundTraceAnalyzer(AsShared(), InSession);
 	}
@@ -400,6 +401,11 @@ namespace UE::Audio::Insights
 			ActiveSoundEntryRef.EntryType = Msg.EntryType;
 			ActiveSoundEntryRef.CategoryName = FText::FromString(FSoundTraceProviderPrivate::GetEntryTypeDisplayName(Msg.EntryType));
 
+			if (PlottingSoundEntries.Contains(ActiveSoundEntryRef.PlayOrder))
+			{
+				ActiveSoundEntryRef.bIsPlotActive = true;
+			}
+
 #if WITH_EDITOR
 			ActiveSoundEntryRef.OwnerActorName = FText::FromString(*Msg.OwnerActorName);
 #endif // WITH_EDITOR
@@ -647,6 +653,16 @@ namespace UE::Audio::Insights
 					}
 
 					EntriesTimingOut.Push(SoundMessageIDs{ Msg.DeviceId, Msg.PlayOrder });
+
+					// Cache the bIsPlotActive flag in case this sound returns in future
+					if ((*OutActiveSoundEntry)->bIsPlotActive)
+					{
+						PlottingSoundEntries.FindOrAdd((*OutActiveSoundEntry)->PlayOrder);
+					}
+					else if (PlottingSoundEntries.Contains((*OutActiveSoundEntry)->PlayOrder))
+					{
+						PlottingSoundEntries.Remove((*OutActiveSoundEntry)->PlayOrder);
+					}
 				}
 			}
 			else
@@ -965,6 +981,14 @@ namespace UE::Audio::Insights
 			return;
 		}
 
+		for (const auto& [AudioDeviceID, DeviceData] : DeviceDataMap)
+		{
+			for (const auto& [EntryType, SoundDashboardEntry] : DeviceData)
+			{
+				CacheIsPlottingFlagRecursive(SoundDashboardEntry);
+			}
+		}
+
 		TSharedPtr<FSoundDashboardEntry> ActiveSoundEntryReturn = nullptr;
 		auto GetOrCreateActiveSoundEntryFromActiveStartMessage = [this, &ActiveSoundEntryReturn](const FSoundStartMessage& Msg) -> TSharedPtr<FSoundDashboardEntry>*
 		{
@@ -1018,6 +1042,11 @@ namespace UE::Audio::Insights
 			FSoundDashboardEntry& EntryRef = *ActiveSoundEntry->Get();
 			EntryRef.SetName(*StartCachedMessage.Name);
 			EntryRef.EntryType = StartCachedMessage.EntryType;
+
+			if (!EntryRef.bIsPlotActive)
+			{
+				EntryRef.bIsPlotActive = PlottingSoundEntries.Contains(EntryRef.PlayOrder);
+			}
 		}
 
 		////////////////////////////////////////
@@ -1049,6 +1078,11 @@ namespace UE::Audio::Insights
 			ChildEntry->SetName(*SoundWaveStartCachedMessage.Name);
 			ChildEntry->EntryType = SoundWaveStartCachedMessage.EntryType;
 
+			if (!ChildEntry->bIsPlotActive)
+			{
+				ChildEntry->bIsPlotActive = PlottingSoundEntries.Contains(ChildEntry->PlayOrder);
+			}
+
 			ParentSoundEntryReturn->Children.Add(ChildEntry);
 		}
 
@@ -1210,6 +1244,33 @@ namespace UE::Audio::Insights
 		// Call parent method to update LastMessageId
 		FTraceProviderBase::OnTimingViewTimeMarkerChanged(TimeMarker);
 	}
+
+	void FSoundTraceProvider::CacheIsPlottingFlagRecursive(const TSharedPtr<IDashboardDataTreeViewEntry>& Entry)
+	{
+		using namespace FSoundTraceProviderPrivate;
+
+		if (!Entry.IsValid())
+		{
+			return;
+		}
+
+		FSoundDashboardEntry& SoundEntry = CastEntry(*Entry);
+
+		// Cache the bIsPlotActive flag in case this sound returns in future
+		if (SoundEntry.bIsPlotActive)
+		{
+			PlottingSoundEntries.FindOrAdd(SoundEntry.PlayOrder);
+		}
+		else if (PlottingSoundEntries.Contains(SoundEntry.PlayOrder))
+		{
+			PlottingSoundEntries.Remove(SoundEntry.PlayOrder);
+		}
+
+		for (const TSharedPtr<IDashboardDataTreeViewEntry>& Child : SoundEntry.Children)
+		{
+			CacheIsPlottingFlagRecursive(Child);
+		}
+	}
 #endif // !WITH_EDITOR
 
 } // namespace UE::Audio::Insights
diff --git a/Engine/Plugins/AudioInsights/Source/AudioInsights/Private/Providers/SoundTraceProvider.h b/Engine/Plugins/AudioInsights/Source/AudioInsights/Private/Providers/SoundTraceProvider.h
--- a/Engine/Plugins/AudioInsights/Source/AudioInsights/Private/Providers/SoundTraceProvider.h
+++ b/Engine/Plugins/AudioInsights/Source/AudioInsights/Private/Providers/SoundTraceProvider.h
@@ -47,6 +47,7 @@ namespace UE::Audio::Insights
 
 #if !WITH_EDITOR
 		virtual void OnTimingViewTimeMarkerChanged(double TimeMarker) override;
+		void CacheIsPlottingFlagRecursive(const TSharedPtr<IDashboardDataTreeViewEntry>& Entry);
 
 		TUniquePtr<FSoundSessionCachedMessages> SessionCachedMessages;
 #endif // !WITH_EDITOR
@@ -67,6 +68,7 @@ namespace UE::Audio::Insights
 		TMap<uint32, SoundEntryKeys> ActiveSoundToEntryKeysMap;
 		TArray<SoundMessageIDs> EntriesTimingOut;
 		TSet<uint32> SoundsStoppedBeforeStart;
+		TSet<uint32> PlottingSoundEntries;
 
 		double DashboardTimeoutTime = 3.0;
 
