[Audio Insights] Bug fixes and QOL for ploting in standalone Unreal Insights.
- Keep entries that were plotting but have now stopped in the sounds dashboard. These entries are greyed out.
- When moving the timemarker, reduce the data buffers to single values. Fixes data getting stuck on incorrect values when scrubbing back and forth
- Fixed sounds virtualizing and realizing missing the earlier entries in the sounds dashboard

diff --git a/Engine/Plugins/AudioInsights/Source/AudioInsights/Private/Messages/SoundTraceMessages.cpp b/Engine/Plugins/AudioInsights/Source/AudioInsights/Private/Messages/SoundTraceMessages.cpp
--- a/Engine/Plugins/AudioInsights/Source/AudioInsights/Private/Messages/SoundTraceMessages.cpp
+++ b/Engine/Plugins/AudioInsights/Source/AudioInsights/Private/Messages/SoundTraceMessages.cpp
@@ -75,4 +75,38 @@ namespace UE::Audio::Insights
 		PitchDataPoints.SetCapacity(DataPointsCapacity);
 		RelativeRenderCostDataPoints.SetCapacity(DataPointsCapacity);
 	}
+
+#if !WITH_EDITOR
+	void FSoundDashboardEntry::SwitchBuffersToSingleValues()
+	{
+		constexpr uint32 DataPointsCapacity = 1;
+
+		PriorityDataPoints.Empty();
+		PriorityDataPoints.SetCapacity(DataPointsCapacity);
+
+		DistanceDataPoints.Empty();
+		DistanceDataPoints.SetCapacity(DataPointsCapacity);
+
+		DistanceAttenuationDataPoints.Empty();
+		DistanceAttenuationDataPoints.SetCapacity(DataPointsCapacity);
+
+		LPFFreqDataPoints.Empty();
+		LPFFreqDataPoints.SetCapacity(DataPointsCapacity);
+
+		HPFFreqDataPoints.Empty();
+		HPFFreqDataPoints.SetCapacity(DataPointsCapacity);
+
+		AmplitudeDataPoints.Empty();
+		AmplitudeDataPoints.SetCapacity(DataPointsCapacity);
+
+		VolumeDataPoints.Empty();
+		VolumeDataPoints.SetCapacity(DataPointsCapacity);
+
+		PitchDataPoints.Empty();
+		PitchDataPoints.SetCapacity(DataPointsCapacity);
+
+		RelativeRenderCostDataPoints.Empty();
+		RelativeRenderCostDataPoints.SetCapacity(DataPointsCapacity);
+	}
+#endif // !WITH_EDITOR
 }
\ No newline at end of file
diff --git a/Engine/Plugins/AudioInsights/Source/AudioInsights/Private/Providers/SoundTraceProvider.cpp b/Engine/Plugins/AudioInsights/Source/AudioInsights/Private/Providers/SoundTraceProvider.cpp
--- a/Engine/Plugins/AudioInsights/Source/AudioInsights/Private/Providers/SoundTraceProvider.cpp
+++ b/Engine/Plugins/AudioInsights/Source/AudioInsights/Private/Providers/SoundTraceProvider.cpp
@@ -1042,6 +1042,7 @@ namespace UE::Audio::Insights
 			FSoundDashboardEntry& EntryRef = *ActiveSoundEntry->Get();
 			EntryRef.SetName(*StartCachedMessage.Name);
 			EntryRef.EntryType = StartCachedMessage.EntryType;
+			EntryRef.bForceKeepEntryAlive = false;
 
 			if (!EntryRef.bIsPlotActive)
 			{
@@ -1077,6 +1078,7 @@ namespace UE::Audio::Insights
 			ChildEntry->Timestamp = SoundWaveStartCachedMessage.Timestamp;
 			ChildEntry->SetName(*SoundWaveStartCachedMessage.Name);
 			ChildEntry->EntryType = SoundWaveStartCachedMessage.EntryType;
+			ChildEntry->bForceKeepEntryAlive = false;
 
 			if (!ChildEntry->bIsPlotActive)
 			{
@@ -1099,21 +1101,13 @@ namespace UE::Audio::Insights
 
 			if (ActiveSoundEntry && ActiveSoundEntry->IsValid() && (*ActiveSoundEntry)->Timestamp < StopCachedMessage.Timestamp)
 			{
-				RemoveActiveSoundEntry(StopCachedMessage.DeviceId, *ActiveSoundEntry);
-			}
-		}
-
-		// Re-run through start messages after stop messages have been processed
-		// Any entries found with timestamps higher than the current time marker need to be cleared out too.
-		for (const FSoundStartMessage& StartCachedMessage : SessionCachedMessages->StartCachedMessages)
-		{
-			if (StartCachedMessage.Timestamp > TimeMarker)
-			{
-				TSharedPtr<FSoundDashboardEntry>* FoundSoundEntry = GetActiveSoundEntryFromStartMessage(StartCachedMessage);
-
-				if (FoundSoundEntry && FoundSoundEntry->IsValid())
+				if ((*ActiveSoundEntry)->bIsPlotActive)
 				{
-					RemoveActiveSoundEntry(StartCachedMessage.DeviceId, *FoundSoundEntry);
+					(*ActiveSoundEntry)->bForceKeepEntryAlive = true;
+				}
+				else
+				{
+					RemoveActiveSoundEntry(StopCachedMessage.DeviceId, *ActiveSoundEntry);
 				}
 			}
 		}
@@ -1227,6 +1221,10 @@ namespace UE::Audio::Insights
 					continue;
 				}
 
+				// Once we have gathered all of the plot data, we are only interested in the closest value to the time marker to display inside the Sounds Dashboard
+				// Changing our data buffers to single points helps avoid issues with ordering data while moving the timemarker back and forth
+				SoundEntry->SwitchBuffersToSingleValues();
+
 				SoundEntry->PriorityDataPoints.Push({ CachedEntryInfo.PriorityMessage.Timestamp, CachedEntryInfo.PriorityMessage.Priority });
 				SoundEntry->DistanceDataPoints.Push({ CachedEntryInfo.DistanceMessage.Timestamp, CachedEntryInfo.DistanceMessage.Distance });
 				SoundEntry->DistanceAttenuationDataPoints.Push({ CachedEntryInfo.DistanceAttenuationMessage.Timestamp, CachedEntryInfo.DistanceAttenuationMessage.DistanceAttenuation });
diff --git a/Engine/Plugins/AudioInsights/Source/AudioInsights/Private/Views/SoundDashboardViewFactory.cpp b/Engine/Plugins/AudioInsights/Source/AudioInsights/Private/Views/SoundDashboardViewFactory.cpp
--- a/Engine/Plugins/AudioInsights/Source/AudioInsights/Private/Views/SoundDashboardViewFactory.cpp
+++ b/Engine/Plugins/AudioInsights/Source/AudioInsights/Private/Views/SoundDashboardViewFactory.cpp
@@ -102,7 +102,12 @@ namespace UE::Audio::Insights
 		bool IsVisible(const IDashboardDataTreeViewEntry& InEntry, const bool bShowRecentlyStoppedSounds)
 		{
 			const FSoundDashboardEntry& SoundEntry = CastEntry(InEntry);
+
+#if WITH_EDITOR
 			return SoundEntry.bIsVisible && (bShowRecentlyStoppedSounds || SoundEntry.TimeoutTimestamp == INVALID_TIMEOUT);
+#else
+			return SoundEntry.bIsVisible && ((bShowRecentlyStoppedSounds || SoundEntry.TimeoutTimestamp == INVALID_TIMEOUT) || SoundEntry.bForceKeepEntryAlive);
+#endif // WITH_EDITOR
 		}
 
 		FSlateColor GetIconColor(const TSharedPtr<IDashboardDataTreeViewEntry>& InEntry)
@@ -1299,7 +1304,11 @@ namespace UE::Audio::Insights
 			.RenderOpacity_Lambda([Item]
 			{
 				const FSoundDashboardEntry& SoundDashboardEntry = CastEntry(*Item);
+#if WITH_EDITOR
 				return SoundDashboardEntry.TimeoutTimestamp == INVALID_TIMEOUT ? 1.0f : 0.25f;
+#else
+				return SoundDashboardEntry.TimeoutTimestamp == INVALID_TIMEOUT && !SoundDashboardEntry.bForceKeepEntryAlive ? 1.0f : 0.25f;
+#endif // WITH_EDITOR
 			})
 			.Visibility_Lambda([this, Item]()
 			{
@@ -1321,7 +1330,11 @@ namespace UE::Audio::Insights
 				const int32 NumUnpinnedChildren = GetNumChildrenWithoutPinEntryType(*Item, FSoundDashboardEntry::EPinnedEntryType::HiddenOriginalEntry, bShowRecentlyStoppedSounds);
 
 				const bool bRowShouldBeVisible = SoundEntry.bIsVisible 
+#if WITH_EDITOR
 												&& (bShowRecentlyStoppedSounds || SoundEntry.TimeoutTimestamp == INVALID_TIMEOUT)
+#else
+												&& ((bShowRecentlyStoppedSounds || SoundEntry.TimeoutTimestamp == INVALID_TIMEOUT) || SoundEntry.bForceKeepEntryAlive)
+#endif // WITH_EDITOR
 												&& SoundEntry.PinnedEntryType != FSoundDashboardEntry::EPinnedEntryType::HiddenOriginalEntry 
 												&& (!IsCategoryItem(*Item) || (NumUnpinnedChildren > 0));
 
diff --git a/Engine/Plugins/AudioInsights/Source/AudioInsights/Public/Messages/SoundTraceMessages.h b/Engine/Plugins/AudioInsights/Source/AudioInsights/Public/Messages/SoundTraceMessages.h
--- a/Engine/Plugins/AudioInsights/Source/AudioInsights/Public/Messages/SoundTraceMessages.h
+++ b/Engine/Plugins/AudioInsights/Source/AudioInsights/Public/Messages/SoundTraceMessages.h
@@ -202,6 +202,10 @@ namespace UE::Audio::Insights
 		virtual bool HasSetInitExpansion() const override { return bHasSetInitExpansion; }
 		virtual void ResetHasSetInitExpansion() override { bHasSetInitExpansion = true; }
 
+#if !WITH_EDITOR
+		void SwitchBuffersToSingleValues();
+#endif // !WITH_EDITOR
+
 		enum class EPinnedEntryType : uint8
 		{
 			None = 0,
@@ -220,6 +224,9 @@ namespace UE::Audio::Insights
 		bool bHasSetInitExpansion = false;
 		bool bIsVisible = false;
 		bool bIsPlotActive = false;
+#if !WITH_EDITOR
+		bool bForceKeepEntryAlive = false;
+#endif // !WITH_EDITOR
 
 		::Audio::TCircularAudioBuffer<FDataPoint> PriorityDataPoints;
 		::Audio::TCircularAudioBuffer<FDataPoint> DistanceDataPoints;
