[Audio Insights] Persist visible columns for sounds dashboard in the editor settings.

diff --git a/Engine/Plugins/AudioInsights/Source/AudioInsights/Private/Views/SoundDashboardViewFactory.cpp b/Engine/Plugins/AudioInsights/Source/AudioInsights/Private/Views/SoundDashboardViewFactory.cpp
--- a/Engine/Plugins/AudioInsights/Source/AudioInsights/Private/Views/SoundDashboardViewFactory.cpp
+++ b/Engine/Plugins/AudioInsights/Source/AudioInsights/Private/Views/SoundDashboardViewFactory.cpp
@@ -1574,6 +1574,13 @@ namespace UE::Audio::Insights
 	{
 		using namespace FSoundDashboardViewFactoryPrivate;
 
+		FSoundDashboardVisibleColumns VisibleColumns;
+		if (VisibleColumnSettings.IsSet())
+		{
+			// Use the values loaded from the settings:
+			VisibleColumns = VisibleColumnSettings.GetValue();
+		}
+
 		static const TMap<FName, FTraceTreeDashboardViewFactory::FHeaderRowColumnData> HeaderRowColumnData =
 		{
 #if WITH_EDITOR
@@ -1583,7 +1590,7 @@ namespace UE::Audio::Insights
 					.DisplayName      = MuteColumnDisplayName,
 					.IconName         = "AudioInsights.Icon.SoundDashboard.Mute",
 					.bShowDisplayName = false,
-					.bDefaultHidden   = false,
+					.bDefaultHidden   = !VisibleColumns.bMute,
 					.FillWidth        = 0.05f,
 					.Alignment        = EHorizontalAlignment::HAlign_Center
 				}
@@ -1594,7 +1601,7 @@ namespace UE::Audio::Insights
 					.DisplayName      = SoloColumnDisplayName,
 					.IconName         = "AudioInsights.Icon.SoundDashboard.Solo",
 					.bShowDisplayName = false,
-					.bDefaultHidden   = false,
+					.bDefaultHidden   = !VisibleColumns.bSolo,
 					.FillWidth        = 0.05f,
 					.Alignment        = EHorizontalAlignment::HAlign_Center
 				}
@@ -1606,7 +1613,7 @@ namespace UE::Audio::Insights
 					.DisplayName = PlotColumnDisplayName,
 					.IconName = "AudioInsights.Icon.SoundDashboard.Plots",
 					.bShowDisplayName = false,
-					.bDefaultHidden = false,
+					.bDefaultHidden = !VisibleColumns.bPlot,
 					.FillWidth = 0.05f,
 					.Alignment = EHorizontalAlignment::HAlign_Center
 				}
@@ -1617,7 +1624,7 @@ namespace UE::Audio::Insights
 					.DisplayName      = NameColumnDisplayName,
 					.IconName         = NAME_None,
 					.bShowDisplayName = true,
-					.bDefaultHidden   = false,
+					.bDefaultHidden   = !VisibleColumns.bName,
 					.FillWidth        = 0.5f,
 					.Alignment        = EHorizontalAlignment::HAlign_Left
 				}
@@ -1628,7 +1635,7 @@ namespace UE::Audio::Insights
 					.DisplayName	  = PlayOrderColumnDisplayName,
 					.IconName		  = NAME_None,
 					.bShowDisplayName = true,
-					.bDefaultHidden   = true,
+					.bDefaultHidden   = !VisibleColumns.bPlayOrder,
 					.FillWidth		  = 0.1f,
 					.Alignment		  = EHorizontalAlignment::HAlign_Left
 				}
@@ -1639,7 +1646,7 @@ namespace UE::Audio::Insights
 					.DisplayName	  = PriorityColumnDisplayName,
 					.IconName		  = NAME_None,
 					.bShowDisplayName = true,
-					.bDefaultHidden	  = false,
+					.bDefaultHidden	  = !VisibleColumns.bPriority,
 					.FillWidth		  = 0.08f,
 					.Alignment		  = EHorizontalAlignment::HAlign_Left
 				}
@@ -1650,7 +1657,7 @@ namespace UE::Audio::Insights
 					.DisplayName      = DistanceColumnDisplayName,
 					.IconName         = NAME_None,
 					.bShowDisplayName = true,
-					.bDefaultHidden   = false,
+					.bDefaultHidden   = !VisibleColumns.bDistance,
 					.FillWidth        = 0.1f,
 					.Alignment        = EHorizontalAlignment::HAlign_Left
 				}
@@ -1661,7 +1668,7 @@ namespace UE::Audio::Insights
 					.DisplayName	  = DistanceAttenuationColumnDisplayName,
 					.IconName		  = NAME_None,
 					.bShowDisplayName = true,
-					.bDefaultHidden   = true,
+					.bDefaultHidden   = !VisibleColumns.bDistanceAttenuation,
 					.FillWidth		  = 0.1f,
 					.Alignment		  = EHorizontalAlignment::HAlign_Left
 				}
@@ -1672,7 +1679,7 @@ namespace UE::Audio::Insights
 					.DisplayName      = AmplitudeColumnDisplayName,
 					.IconName         = NAME_None,
 					.bShowDisplayName = true,
-					.bDefaultHidden   = false,
+					.bDefaultHidden   = !VisibleColumns.bAmplitude,
 					.FillWidth        = 0.12f,
 					.Alignment        = EHorizontalAlignment::HAlign_Left
 				}
@@ -1683,7 +1690,7 @@ namespace UE::Audio::Insights
 					.DisplayName      = VolumeColumnDisplayName,
 					.IconName         = NAME_None,
 					.bShowDisplayName = true,
-					.bDefaultHidden   = false,
+					.bDefaultHidden   = !VisibleColumns.bVolume,
 					.FillWidth        = 0.1f,
 					.Alignment        = EHorizontalAlignment::HAlign_Left
 				}
@@ -1694,7 +1701,7 @@ namespace UE::Audio::Insights
 					.DisplayName	  = LPFFreqColumnDisplayName,
 					.IconName		  = NAME_None,
 					.bShowDisplayName = true,
-					.bDefaultHidden   = true,
+					.bDefaultHidden   = !VisibleColumns.bLPFFreq,
 					.FillWidth		  = 0.1f,
 					.Alignment		  = EHorizontalAlignment::HAlign_Left
 				}
@@ -1705,7 +1712,7 @@ namespace UE::Audio::Insights
 					.DisplayName	  = HPFFreqColumnDisplayName,
 					.IconName		  = NAME_None,
 					.bShowDisplayName = true,
-					.bDefaultHidden   = true,
+					.bDefaultHidden   = !VisibleColumns.bHPFFreq,
 					.FillWidth		  = 0.1f,
 					.Alignment		  = EHorizontalAlignment::HAlign_Left
 				}
@@ -1716,7 +1723,7 @@ namespace UE::Audio::Insights
 					.DisplayName      = PitchColumnDisplayName,
 					.IconName         = NAME_None,
 					.bShowDisplayName = true,
-					.bDefaultHidden   = false,
+					.bDefaultHidden   = !VisibleColumns.bPitch,
 					.FillWidth        = 0.1f,
 					.Alignment        = EHorizontalAlignment::HAlign_Left
 				}
@@ -1727,7 +1734,7 @@ namespace UE::Audio::Insights
 					.DisplayName = RelativeRenderCostColumnDisplayName,
 					.IconName = NAME_None,
 					.bShowDisplayName = true,
-					.bDefaultHidden = true,
+					.bDefaultHidden = !VisibleColumns.bRelativeRenderCost,
 					.FillWidth = 0.1f,
 					.Alignment = EHorizontalAlignment::HAlign_Left
 				}
@@ -1739,7 +1746,7 @@ namespace UE::Audio::Insights
 					.DisplayName = ActorColumnDisplayName,
 					.IconName = NAME_None,
 					.bShowDisplayName = true,
-					.bDefaultHidden = true,
+					.bDefaultHidden = !VisibleColumns.bActor,
 					.FillWidth = 0.1f,
 					.Alignment = EHorizontalAlignment::HAlign_Left
 				}
@@ -1751,7 +1758,7 @@ namespace UE::Audio::Insights
 					.DisplayName = LOCTEXT("SoundDashboard_CategoryColumnDisplayName", "Category"),
 					.IconName = NAME_None,
 					.bShowDisplayName = true,
-					.bDefaultHidden = true,
+					.bDefaultHidden = !VisibleColumns.bCategory,
 					.FillWidth = 0.1f,
 					.Alignment = EHorizontalAlignment::HAlign_Left
 				}
@@ -1761,6 +1768,22 @@ namespace UE::Audio::Insights
 		return HeaderRowColumnData;
 	}
 
+	void FSoundDashboardViewFactory::OnHiddenColumnsListChanged()
+	{
+		// Update and save settings to reload it on the next Editor sessions
+		if (VisibleColumnSettings && HeaderRowWidget)
+		{
+			for (const SHeaderRow::FColumn& Column : HeaderRowWidget->GetColumns())
+			{
+				VisibleColumnSettings->SetIsVisible(Column.ColumnId, Column.bIsVisible);
+			}
+
+#if WITH_EDITOR
+			FSoundDashboardSettings::OnRequestWriteSettings.Broadcast();
+#endif // WITH_EDITOR
+		}
+	}
+
 	const TMap<FName, FTraceTreeDashboardViewFactory::FColumnData>& FSoundDashboardViewFactory::GetColumns() const
 	{
 		using namespace FSoundDashboardViewFactoryPrivate;
@@ -2558,6 +2581,24 @@ namespace UE::Audio::Insights
 		bShowRecentlyStoppedSounds = InSettings.bShowStoppedSounds;
 		bDisplayAmpPeakInDb = InSettings.AmplitudeDisplayMode == EAudioAmplitudeDisplayMode::Decibels;
 
+		// If HeaderRowWidget is already created, update it with the desired visible columns:
+		if (HeaderRowWidget)
+		{
+			// Null our VisibleColumnSettings TOptional while updating HeaderRowWidget so that we ignore OnHiddenColumnsListChanged callbacks:
+			VisibleColumnSettings.Reset();
+
+			for (const SHeaderRow::FColumn& Column : HeaderRowWidget->GetColumns())
+			{
+				const bool bShowColumn = InSettings.VisibleColumns.GetIsVisible(Column.ColumnId);
+				if (Column.bIsVisible != bShowColumn)
+				{
+					HeaderRowWidget->SetShowGeneratedColumn(Column.ColumnId, bShowColumn);
+				}
+			}
+		}
+
+		VisibleColumnSettings = InSettings.VisibleColumns;
+
 		if (RecentlyStoppedSoundTimeout != InSettings.StoppedSoundTimeoutTime)
 		{
 			RecentlyStoppedSoundTimeout = InSettings.StoppedSoundTimeoutTime;
@@ -2573,6 +2614,11 @@ namespace UE::Audio::Insights
 		OutSettings.bShowStoppedSounds = bShowRecentlyStoppedSounds;
 		OutSettings.AmplitudeDisplayMode = bDisplayAmpPeakInDb ? EAudioAmplitudeDisplayMode::Decibels : EAudioAmplitudeDisplayMode::Linear;
 
+		if (VisibleColumnSettings)
+		{
+			OutSettings.VisibleColumns = VisibleColumnSettings.GetValue();
+		}
+
 		if (OutSettings.StoppedSoundTimeoutTime != RecentlyStoppedSoundTimeout)
 		{
 			OutSettings.StoppedSoundTimeoutTime = RecentlyStoppedSoundTimeout;
diff --git a/Engine/Plugins/AudioInsights/Source/AudioInsights/Private/Views/TreeDashboardViewFactory.cpp b/Engine/Plugins/AudioInsights/Source/AudioInsights/Private/Views/TreeDashboardViewFactory.cpp
--- a/Engine/Plugins/AudioInsights/Source/AudioInsights/Private/Views/TreeDashboardViewFactory.cpp
+++ b/Engine/Plugins/AudioInsights/Source/AudioInsights/Private/Views/TreeDashboardViewFactory.cpp
@@ -171,7 +171,8 @@ namespace UE::Audio::Insights
 			[](const TPair<FName, FHeaderRowColumnData>& ColumnInfo) { return ColumnInfo.Key; });
 
 		SAssignNew(HeaderRowWidget, SHeaderRow)
-		.CanSelectGeneratedColumn(true); // Allows for showing/hiding columns
+		.CanSelectGeneratedColumn(true) // Allows for showing/hiding columns
+		.OnHiddenColumnsListChanged(this, &FTraceTreeDashboardViewFactory::OnHiddenColumnsListChanged);
 
 		// This only works if header row columns are added with slots and not programmatically
 		// check in SHeaderRow::Construct: for ( FColumn* const Column : InArgs.Slots ) for more info
diff --git a/Engine/Plugins/AudioInsights/Source/AudioInsights/Public/Settings/SoundDashboardSettings.h b/Engine/Plugins/AudioInsights/Source/AudioInsights/Public/Settings/SoundDashboardSettings.h
--- a/Engine/Plugins/AudioInsights/Source/AudioInsights/Public/Settings/SoundDashboardSettings.h
+++ b/Engine/Plugins/AudioInsights/Source/AudioInsights/Public/Settings/SoundDashboardSettings.h
@@ -57,6 +57,86 @@ enum class ESoundDashboardAutoExpandOptions : uint8
 	Nothing			UMETA(DisplayName = "Nothing")
 };
 
+USTRUCT()
+struct FSoundDashboardVisibleColumns
+{
+	GENERATED_BODY()
+
+	UPROPERTY(EditAnywhere, config, Category = SoundDashboard)
+	bool bMute = true;
+
+	UPROPERTY(EditAnywhere, config, Category = SoundDashboard)
+	bool bSolo = true;
+
+	UPROPERTY(EditAnywhere, config, Category = SoundDashboard)
+	bool bPlot = true;
+
+	UPROPERTY(EditAnywhere, config, Category = SoundDashboard)
+	bool bName = true;
+
+	UPROPERTY(EditAnywhere, config, Category = SoundDashboard)
+	bool bPlayOrder = false;
+
+	UPROPERTY(EditAnywhere, config, Category = SoundDashboard)
+	bool bPriority = true;
+
+	UPROPERTY(EditAnywhere, config, Category = SoundDashboard)
+	bool bDistance = true;
+
+	UPROPERTY(EditAnywhere, config, Category = SoundDashboard)
+	bool bDistanceAttenuation = false;
+
+	UPROPERTY(EditAnywhere, config, Category = SoundDashboard, meta = (DisplayName = "Amp (Peak)"))
+	bool bAmplitude = true;
+
+	UPROPERTY(EditAnywhere, config, Category = SoundDashboard)
+	bool bVolume = true;
+
+	UPROPERTY(EditAnywhere, config, Category = SoundDashboard, meta = (DisplayName = "LPF Freq"))
+	bool bLPFFreq = false;
+
+	UPROPERTY(EditAnywhere, config, Category = SoundDashboard, meta = (DisplayName = "HPF Freq"))
+	bool bHPFFreq = false;
+
+	UPROPERTY(EditAnywhere, config, Category = SoundDashboard)
+	bool bPitch = true;
+
+	UPROPERTY(EditAnywhere, config, Category = SoundDashboard)
+	bool bRelativeRenderCost = false;
+
+	UPROPERTY(EditAnywhere, config, Category = SoundDashboard)
+	bool bActor = false;
+
+	UPROPERTY(EditAnywhere, config, Category = SoundDashboard)
+	bool bCategory = false;
+
+	bool GetIsVisible(FName ColumnId) const
+	{
+		const FName PropertyName(TEXT("b") + ColumnId.ToString());
+		if (const FProperty* Property = StaticStruct()->FindPropertyByName(PropertyName))
+		{
+			if (const FBoolProperty* BoolProperty = CastField<FBoolProperty>(Property))
+			{
+				return BoolProperty->GetPropertyValue_InContainer(this);
+			}
+		}
+
+		return false;
+	}
+
+	void SetIsVisible(FName ColumnId, bool bIsVisible)
+	{
+		const FName PropertyName(TEXT("b") + ColumnId.ToString());
+		if (const FProperty* Property = StaticStruct()->FindPropertyByName(PropertyName))
+		{
+			if (const FBoolProperty* BoolProperty = CastField<FBoolProperty>(Property))
+			{
+				BoolProperty->SetPropertyValue_InContainer(this, bIsVisible);
+			}
+		}
+	}
+};
+
 USTRUCT()
 struct FSoundDashboardSettings
 {
@@ -77,6 +157,9 @@ struct FSoundDashboardSettings
 	UPROPERTY(EditAnywhere, config, Category = SoundDashboard, meta = (ClampMin = "0", DisplayName = "Recently Stopped Sounds Timeout"))
 	float StoppedSoundTimeoutTime = 3.0f;
 
+	UPROPERTY(EditAnywhere, config, Category = SoundDashboard)
+	FSoundDashboardVisibleColumns VisibleColumns;
+
 #if WITH_EDITOR
 	DECLARE_MULTICAST_DELEGATE_OneParam(FOnReadSoundDashboardSettings, const FSoundDashboardSettings&);
 	static AUDIOINSIGHTS_API FOnReadSoundDashboardSettings OnReadSettings;
diff --git a/Engine/Plugins/AudioInsights/Source/AudioInsights/Public/Views/SoundDashboardViewFactory.h b/Engine/Plugins/AudioInsights/Source/AudioInsights/Public/Views/SoundDashboardViewFactory.h
--- a/Engine/Plugins/AudioInsights/Source/AudioInsights/Public/Views/SoundDashboardViewFactory.h
+++ b/Engine/Plugins/AudioInsights/Source/AudioInsights/Public/Views/SoundDashboardViewFactory.h
@@ -138,6 +138,7 @@ namespace UE::Audio::Insights
 		UE_API virtual FReply OnDataRowKeyInput(const FGeometry& InGeometry, const FKeyEvent& InKeyEvent) const override;
 
 		UE_API virtual const TMap<FName, FTraceTreeDashboardViewFactory::FHeaderRowColumnData>& GetHeaderRowColumns() const override;
+		UE_API virtual void OnHiddenColumnsListChanged() override;
 		UE_API virtual const TMap<FName, FTraceTreeDashboardViewFactory::FColumnData>& GetColumns() const override;
 
 		UE_API virtual void SortTable() override;
@@ -243,6 +244,7 @@ namespace UE::Audio::Insights
 		bool bShowRecentlyStoppedSounds = false;
 		bool bTreeViewModeChanged = false;
 		bool bDisplayAmpPeakInDb = true;
+		TOptional<FSoundDashboardVisibleColumns> VisibleColumnSettings;
 	};
 } // namespace UE::Audio::Insights
 
diff --git a/Engine/Plugins/AudioInsights/Source/AudioInsights/Public/Views/TreeDashboardViewFactory.h b/Engine/Plugins/AudioInsights/Source/AudioInsights/Public/Views/TreeDashboardViewFactory.h
--- a/Engine/Plugins/AudioInsights/Source/AudioInsights/Public/Views/TreeDashboardViewFactory.h
+++ b/Engine/Plugins/AudioInsights/Source/AudioInsights/Public/Views/TreeDashboardViewFactory.h
@@ -100,6 +100,7 @@ namespace UE::Audio::Insights
 		UE_API virtual void RefreshFilteredEntriesListView();
 
 		virtual const TMap<FName, FHeaderRowColumnData>& GetHeaderRowColumns() const = 0;
+		virtual void OnHiddenColumnsListChanged() {};
 		virtual const TMap<FName, FColumnData>& GetColumns() const = 0;
 
 		virtual void ProcessEntries(EProcessReason Reason) = 0;
diff --git a/Engine/Plugins/AudioInsights/Source/AudioInsightsEditor/Private/AudioInsightsEditorSettings.cpp b/Engine/Plugins/AudioInsights/Source/AudioInsightsEditor/Private/AudioInsightsEditorSettings.cpp
--- a/Engine/Plugins/AudioInsights/Source/AudioInsightsEditor/Private/AudioInsightsEditorSettings.cpp
+++ b/Engine/Plugins/AudioInsights/Source/AudioInsightsEditor/Private/AudioInsightsEditorSettings.cpp
@@ -26,7 +26,11 @@ void UAudioInsightsEditorSettings::PostInitProperties()
 	Super::PostInitProperties();
 
 	OnRequestReadSettingsHandle = FSoundDashboardSettings::OnRequestReadSettings.AddLambda([this]() { FSoundDashboardSettings::OnReadSettings.Broadcast(SoundDashboardSettings); });
-	OnRequestWriteSettingsHandle = FSoundDashboardSettings::OnRequestWriteSettings.AddLambda([this]() { FSoundDashboardSettings::OnWriteSettings.Broadcast(SoundDashboardSettings); });
+	OnRequestWriteSettingsHandle = FSoundDashboardSettings::OnRequestWriteSettings.AddLambda([this]()
+		{
+			FSoundDashboardSettings::OnWriteSettings.Broadcast(SoundDashboardSettings);
+			SaveConfig();
+		});
 }
 
 void UAudioInsightsEditorSettings::BeginDestroy()
