[Audio Insights] Add crosshair to the Plots tab

diff --git a/Engine/Plugins/AudioInsights/Source/AudioInsights/Private/Views/SAudioCurveView.cpp b/Engine/Plugins/AudioInsights/Source/AudioInsights/Private/Views/SAudioCurveView.cpp
--- a/Engine/Plugins/AudioInsights/Source/AudioInsights/Private/Views/SAudioCurveView.cpp
+++ b/Engine/Plugins/AudioInsights/Source/AudioInsights/Private/Views/SAudioCurveView.cpp
@@ -48,6 +48,11 @@ FReply SAudioCurveView::OnMouseMove(const FGeometry& InMyGeometry, const FPointe
 	return FReply::Handled();
 }
 
+void SAudioCurveView::OnMouseLeave(const FPointerEvent& MouseEvent)
+{
+	CrosshairPosition.Reset();
+}
+
 TSharedRef<SToolTip> SAudioCurveView::CreateCurveTooltip()
 {
 	return SNew(SToolTip)
@@ -84,72 +89,83 @@ TSharedRef<SToolTip> SAudioCurveView::CreateCurveTooltip()
 
 void SAudioCurveView::UpdateCurveToolTip(const FGeometry& InMyGeometry, const FPointerEvent& InMouseEvent)
 {
-	if (InMyGeometry.IsUnderLocation(InMouseEvent.GetScreenSpacePosition()) && MetadataPerCurve)
+	if (!InMyGeometry.IsUnderLocation(InMouseEvent.GetScreenSpacePosition()) || !MetadataPerCurve.IsValid())
 	{
-		// Mouse position in widget space
-		const FVector2f HitPosition = InMyGeometry.AbsoluteToLocal(InMouseEvent.GetScreenSpacePosition());
+		CrosshairPosition.Reset();
+		return;
+	}
+
+	// Mouse position in widget space
+	const FVector2f HitPosition = InMyGeometry.AbsoluteToLocal(InMouseEvent.GetScreenSpacePosition());
 
-		// Range helper struct
-		const SSimpleTimeSlider::FScrubRangeToScreen RangeToScreen(ViewRange.Get(), InMyGeometry.GetLocalSize());
+	// Range helper struct
+	const SSimpleTimeSlider::FScrubRangeToScreen RangeToScreen(ViewRange.Get(), InMyGeometry.GetLocalSize());
 
-		// Mouse position from widget space to curve input space
-		const double TargetX = RangeToScreen.LocalXToInput(HitPosition.X);
+	// Mouse position from widget space to curve input space
+	const double TargetX = RangeToScreen.LocalXToInput(HitPosition.X);
 
-		// Keep track of closest curve index, closest point, value 
-		int32 ClosestCurveId = INDEX_NONE;
-		FCurvePoint ClosestPoint;
-		float ClosestDistance = TNumericLimits<float>::Max();
+	// Keep track of closest curve index, closest point, value 
+	int32 ClosestCurveId = INDEX_NONE;
+	FCurvePoint ClosestPoint;
+	float ClosestDistance = TNumericLimits<float>::Max();
+	float ThickestPlot = 0.0f;
 
-		for (const auto& [CurveId, CurveMetadata] : *MetadataPerCurve)
+	for (const auto& [CurveId, CurveMetadata] : *MetadataPerCurve)
+	{
+		if (CurveMetadata.CurveColor.A == 0.0f || CurveMetadata.PlotThickness < ThickestPlot)
 		{
-			if (CurveMetadata.CurveColor.A == 0.0f)
-			{
-				continue;
-			}
+			continue;
+		}
 
-			const TArray<FCurvePoint>* CurvePoints = PointDataPerCurve->Find(CurveId);
-			const int32 NumPoints = CurvePoints->Num();
-			if (NumPoints > 0)
+		const TArray<FCurvePoint>* CurvePoints = PointDataPerCurve->Find(CurveId);
+		const int32 NumPoints = CurvePoints->Num();
+		if (NumPoints > 0)
+		{
+			for (int32 i = 1; i < NumPoints; ++i)
 			{
-				for (int32 i = 1; i < NumPoints; ++i)
-				{
-					const FCurvePoint& Point1 = (*CurvePoints)[i - 1];
-					const FCurvePoint& Point2 = (*CurvePoints)[i];
+				const FCurvePoint& Point1 = (*CurvePoints)[i - 1];
+				const FCurvePoint& Point2 = (*CurvePoints)[i];
 
-					// Find points that contain mouse hit-point x
-					if (Point1.Key >= TargetX && TargetX <= Point2.Key)
+				// Find points that contain mouse hit-point x
+				if (Point1.Key >= TargetX && TargetX <= Point2.Key)
+				{
+					// Choose point with the smallest x delta
+					const float Delta1 = abs(TargetX - Point1.Key);
+					const float Delta2 = abs(TargetX - Point2.Key);
+					const FCurvePoint& TargetPoint = Delta1 < Delta2 ? Point1 : Point2;
+
+					// Convert target point Y to widget space 
+					const FVector2f LocalSize = InMyGeometry.GetLocalSize();
+					const float WidgetSpaceY = ValueToLocalY(LocalSize, TargetPoint.Value);
+					
+					// Compare distance in widget space between HitPosition and closest point by x value on this curve
+					const float Distance = FVector2f::Distance(HitPosition, FVector2f(RangeToScreen.InputToLocalX(TargetPoint.Key), WidgetSpaceY));
+
+					if (Distance < ClosestDistance || CurveMetadata.PlotThickness > ThickestPlot)
 					{
-						// Choose point with the smallest x delta
-						const float Delta1 = abs(TargetX - Point1.Key);
-						const float Delta2 = abs(TargetX - Point2.Key);
-						const FCurvePoint& TargetPoint = Delta1 < Delta2 ? Point1 : Point2;
-
-						// Convert target point Y to widget space 
-						const FVector2f LocalSize = InMyGeometry.GetLocalSize();
-						const float WidgetSpaceY = ValueToLocalY(LocalSize, TargetPoint.Value);
-						
-						// Compare distance in widget space between HitPosition and closest point by x value on this curve
-						const float Distance = FVector2f::Distance(HitPosition, FVector2f(RangeToScreen.InputToLocalX(TargetPoint.Key), WidgetSpaceY));
-
-						if (Distance < ClosestDistance)
-						{
-							ClosestDistance = Distance;
-							ClosestCurveId = CurveId;
-							ClosestPoint = TargetPoint;
-						}
-						break;
+						ClosestDistance = Distance;
+						ClosestCurveId = CurveId;
+						ClosestPoint = TargetPoint;
+						ThickestPlot = CurveMetadata.PlotThickness;
 					}
+					break;
 				}
 			}
 		}
+	}
 
-		// Set tooltip text values  
-		if (ClosestCurveId != INDEX_NONE)
-		{
-			CurveToolTipXValueText = FText::Format(LOCTEXT("CurveToolTipValueXFormat", "X: {0}"), FText::AsNumber(ClosestPoint.Key, &XValueFormattingOptions));
-			CurveToolTipYValueText = FText::Format(LOCTEXT("CurveToolTipValueYFormat", "Y: {0}"), FText::AsNumber(ClosestPoint.Value, &YValueFormattingOptions));
-			CurveToolTipDisplayNameText = MetadataPerCurve->Find(ClosestCurveId)->DisplayName;
-		}
+	// Set tooltip text values  
+	if (ClosestCurveId != INDEX_NONE)
+	{
+		CurveToolTipXValueText = FText::Format(LOCTEXT("CurveToolTipValueXFormat", "X: {0}"), FText::AsNumber(ClosestPoint.Key, &XValueFormattingOptions));
+		CurveToolTipYValueText = FText::Format(LOCTEXT("CurveToolTipValueYFormat", "Y: {0}"), FText::AsNumber(ClosestPoint.Value, &YValueFormattingOptions));
+		CurveToolTipDisplayNameText = MetadataPerCurve->Find(ClosestCurveId)->DisplayName;
+
+		CrosshairPosition = ClosestPoint;
+	}
+	else
+	{
+		CrosshairPosition.Reset();
 	}
 }
 
@@ -360,6 +376,56 @@ int32 SAudioCurveView::PaintGridLines(const FGeometry& AllottedGeometry, const F
 	return LayerId;
 }
 
+int32 SAudioCurveView::PaintCrosshair(const FGeometry& MyGeometry, FSlateWindowElementList& OutDrawElements, int32 LayerId, const SSimpleTimeSlider::FScrubRangeToScreen& RangeToScreen) const
+{
+	if (!CrosshairPosition.IsSet())
+	{
+		return LayerId;
+	}
+
+	const FVector2f Size = MyGeometry.GetLocalSize();
+	const int32 CrosshairLayerID = LayerId++;
+
+	const float PointX = RangeToScreen.InputToLocalX(CrosshairPosition.GetValue().Key);
+	const float PointY = ValueToLocalY(Size, CrosshairPosition.GetValue().Value);
+
+	constexpr FLinearColor CrosshairColor(1.0f, 1.0f, 1.0f, 0.75f);
+
+	TArray<FVector2f> HorizontalCrosshairLinePoints
+	{
+		{ 0.0f, PointY },
+		{ Size.X, PointY }
+	};
+
+	const TArray<FVector2f> VerticalCrosshairLinePoints
+	{
+		{ PointX, 0.0f },
+		{ PointX, Size.Y }
+	};
+
+	FSlateDrawElement::MakeLines(
+		OutDrawElements,
+		CrosshairLayerID,
+		MyGeometry.ToPaintGeometry(),
+		HorizontalCrosshairLinePoints,
+		LineDrawEffects,
+		CrosshairColor,
+		true
+	);
+
+	FSlateDrawElement::MakeLines(
+		OutDrawElements,
+		CrosshairLayerID,
+		MyGeometry.ToPaintGeometry(),
+		VerticalCrosshairLinePoints,
+		LineDrawEffects,
+		CrosshairColor,
+		true
+	);
+
+	return LayerId;
+}
+
 int32 SAudioCurveView::PaintYAxisLabels(const FGeometry& AllottedGeometry, FSlateWindowElementList& OutDrawElements, int32 LayerId) const
 {
 	const int32 BackgroundRectangleLayerId = LayerId++;
@@ -444,6 +510,7 @@ int32 SAudioCurveView::PaintCurves(const FGeometry& AllottedGeometry, const FSla
 
 	const SSimpleTimeSlider::FScrubRangeToScreen RangeToScreen(ViewRange.Get(), AllottedGeometry.GetLocalSize());
 	LayerId = PaintGridLines(AllottedGeometry, MyCullingRect, OutDrawElements, LayerId, InWidgetStyle, bParentEnabled, RangeToScreen);
+	LayerId = PaintCrosshair(AllottedGeometry, OutDrawElements, LayerId, RangeToScreen);
 	
 	static constexpr float LargeFrameTime = 0.5f; // ideally, we could check the recorded frame data for actual frame length
 	const FVector2f Size = AllottedGeometry.GetLocalSize();
diff --git a/Engine/Plugins/AudioInsights/Source/AudioInsights/Public/Views/SAudioCurveView.h b/Engine/Plugins/AudioInsights/Source/AudioInsights/Public/Views/SAudioCurveView.h
--- a/Engine/Plugins/AudioInsights/Source/AudioInsights/Public/Views/SAudioCurveView.h
+++ b/Engine/Plugins/AudioInsights/Source/AudioInsights/Public/Views/SAudioCurveView.h
@@ -22,6 +22,7 @@ public:
 	{
 		FLinearColor CurveColor;
 		FText DisplayName;
+		float PlotThickness = 1.0f;
 	};
 
 	SLATE_BEGIN_ARGS(SAudioCurveView)
@@ -79,9 +80,11 @@ protected:
 	virtual FVector2D ComputeDesiredSize(float) const override;
 	virtual int32 OnPaint( const FPaintArgs& Args, const FGeometry& AllottedGeometry, const FSlateRect& MyCullingRect, FSlateWindowElementList& OutDrawElements, int32 LayerId, const FWidgetStyle& InWidgetStyle, bool bParentEnabled ) const override;
 	virtual FReply OnMouseMove(const FGeometry& MyGeometry, const FPointerEvent& MouseEvent) override;
+	virtual void OnMouseLeave(const FPointerEvent& MouseEvent) override;
 	
 	int32 PaintCurves(const FGeometry& AllottedGeometry, const FSlateRect& MyCullingRect, FSlateWindowElementList& OutDrawElements, int32 LayerId, const FWidgetStyle& InWidgetStyle, bool bParentEnabled) const;
 	int32 PaintGridLines(const FGeometry& AllottedGeometry, const FSlateRect& MyCullingRect, FSlateWindowElementList& OutDrawElements, int32 LayerId, const FWidgetStyle& InWidgetStyle, bool bParentEnabled, const SSimpleTimeSlider::FScrubRangeToScreen& RangeToScreen) const;
+	int32 PaintCrosshair(const FGeometry& MyGeometry, FSlateWindowElementList& OutDrawElements, int32 LayerId, const SSimpleTimeSlider::FScrubRangeToScreen& RangeToScreen) const;
 	int32 PaintYAxisLabels(const FGeometry& AllottedGeometry, FSlateWindowElementList& OutDrawElements, int32 LayerId) const;
 	void UpdateCurveToolTip(const FGeometry& MyGeometry, const FPointerEvent& MouseEvent);
 	void UpdateYDataRange();
@@ -112,4 +115,6 @@ protected:
 	ESlateDrawEffect LineDrawEffects;
 	uint32 NumHorizontalGridLines;
 	FSlateFontInfo LabelFont;
+
+	TOptional<FCurvePoint> CrosshairPosition;
 };
