[Audio Insights] Increase the thickness of plots that are selected in the Sounds Tab. Fixed plots displaying with small gaps within the lines, caused by duplicate data being processed.

diff --git a/Engine/Plugins/AudioInsights/Source/AudioInsights/Private/Views/SAudioCurveView.cpp b/Engine/Plugins/AudioInsights/Source/AudioInsights/Private/Views/SAudioCurveView.cpp
--- a/Engine/Plugins/AudioInsights/Source/AudioInsights/Private/Views/SAudioCurveView.cpp
+++ b/Engine/Plugins/AudioInsights/Source/AudioInsights/Private/Views/SAudioCurveView.cpp
@@ -549,7 +549,8 @@ int32 SAudioCurveView::PaintCurves(const FGeometry& AllottedGeometry, const FSla
 					Points,
 					LineDrawEffects,
 					CurveMetadata.CurveColor,
-					true
+					true,
+					CurveMetadata.PlotThickness
 				);
 
 				Points.Reset();
@@ -585,7 +586,8 @@ int32 SAudioCurveView::PaintCurves(const FGeometry& AllottedGeometry, const FSla
 			Points,
 			LineDrawEffects,
 			CurveMetadata.CurveColor,
-			true
+			true,
+			CurveMetadata.PlotThickness
 		);
 	}
 
diff --git a/Engine/Plugins/AudioInsights/Source/AudioInsights/Private/Views/SoundDashboardViewFactory.cpp b/Engine/Plugins/AudioInsights/Source/AudioInsights/Private/Views/SoundDashboardViewFactory.cpp
--- a/Engine/Plugins/AudioInsights/Source/AudioInsights/Private/Views/SoundDashboardViewFactory.cpp
+++ b/Engine/Plugins/AudioInsights/Source/AudioInsights/Private/Views/SoundDashboardViewFactory.cpp
@@ -1395,9 +1395,9 @@ namespace UE::Audio::Insights
 			return bEntryTypePassesFilter;
 		});
 
-		if (InReason == FTraceTreeDashboardViewFactory::EProcessReason::EntriesUpdated)
+		if (InReason == FTraceTreeDashboardViewFactory::EProcessReason::EntriesUpdated && FilteredEntriesListView.IsValid())
 		{
-			OnProcessPlotData.Broadcast(DataViewEntries);
+			OnProcessPlotData.Broadcast(DataViewEntries, FilteredEntriesListView->GetSelectedItems());
 		}
 	}
 
@@ -1565,6 +1565,16 @@ namespace UE::Audio::Insights
 		return MenuBuilder.MakeWidget();
 	}
 
+	void FSoundDashboardViewFactory::OnSelectionChanged(TSharedPtr<IDashboardDataTreeViewEntry> SelectedItem, ESelectInfo::Type SelectInfo)
+	{
+#if !WITH_EDITOR
+		if (FilteredEntriesListView.IsValid())
+		{
+			OnUpdatePlotSelection.Broadcast(FilteredEntriesListView->GetSelectedItems());
+		}
+#endif // !WITH_EDITOR
+	}
+
 	FReply FSoundDashboardViewFactory::OnDataRowKeyInput(const FGeometry& InGeometry, const FKeyEvent& InKeyEvent) const
 	{
 		return (CommandList && CommandList->ProcessCommandBindings(InKeyEvent)) ? FReply::Handled() : FReply::Unhandled();
diff --git a/Engine/Plugins/AudioInsights/Source/AudioInsights/Private/Views/SoundPlotsDashboardViewFactory.cpp b/Engine/Plugins/AudioInsights/Source/AudioInsights/Private/Views/SoundPlotsDashboardViewFactory.cpp
--- a/Engine/Plugins/AudioInsights/Source/AudioInsights/Private/Views/SoundPlotsDashboardViewFactory.cpp
+++ b/Engine/Plugins/AudioInsights/Source/AudioInsights/Private/Views/SoundPlotsDashboardViewFactory.cpp
@@ -57,14 +57,15 @@ namespace UE::Audio::Insights
 
 #if !WITH_EDITOR
 		OnUpdatePlotVisibility = SoundsDashboard->OnUpdatePlotVisibility.AddSP(this, &FSoundPlotsDashboardViewFactory::UpdatePlotVisibility);
+		OnUpdatePlotSelection = SoundsDashboard->OnUpdatePlotSelection.AddSP(this, &FSoundPlotsDashboardViewFactory::UpdatePlotSelection);
 #endif // !WITH_EDITOR
 	}
 
-	void FSoundPlotsDashboardViewFactory::ProcessPlotData(const TArray<TSharedPtr<IDashboardDataTreeViewEntry>>& DataViewEntries)
+	void FSoundPlotsDashboardViewFactory::ProcessPlotData(const TArray<TSharedPtr<IDashboardDataTreeViewEntry>>& DataViewEntries, const TArray<TSharedPtr<IDashboardDataTreeViewEntry>>& SelectedEntries)
 	{
 		if (PlotsWidget.IsValid())
 		{
-			PlotsWidget->ProcessPlotData(DataViewEntries);
+			PlotsWidget->ProcessPlotData(DataViewEntries, SelectedEntries);
 		}
 	}
 
@@ -76,6 +77,14 @@ namespace UE::Audio::Insights
 			PlotsWidget->UpdatePlotVisibility(DataViewEntries);
 		}
 	}
+
+	void FSoundPlotsDashboardViewFactory::UpdatePlotSelection(const TArray<TSharedPtr<IDashboardDataTreeViewEntry>>& SelectedEntries)
+	{
+		if (PlotsWidget.IsValid())
+		{
+			PlotsWidget->UpdatePlotSelection(SelectedEntries);
+		}
+	}
 #endif // !WITH_EDITOR
 } // namespace UE::Audio::Insights
 #undef LOCTEXT_NAMESPACE
diff --git a/Engine/Plugins/AudioInsights/Source/AudioInsights/Private/Views/SoundPlotsWidgetView.cpp b/Engine/Plugins/AudioInsights/Source/AudioInsights/Private/Views/SoundPlotsWidgetView.cpp
--- a/Engine/Plugins/AudioInsights/Source/AudioInsights/Private/Views/SoundPlotsWidgetView.cpp
+++ b/Engine/Plugins/AudioInsights/Source/AudioInsights/Private/Views/SoundPlotsWidgetView.cpp
@@ -17,6 +17,12 @@
 
 namespace UE::Audio::Insights
 {
+	namespace FSoundPlotsWidgetViewPrivate
+	{
+		constexpr float DefaultPlotThickness = 1.0f;
+		constexpr float SelectedPlotThickness = 4.0f;
+	}
+
 	FSoundPlotsWidgetView::FSoundPlotsWidgetView(TMap<FName, FSoundPlotsWidgetView::FPlotColumnInfo>&& InColumnInfo, const TFunction<bool(const IDashboardDataTreeViewEntry& InEntry)> InIsPlotEnabledForEntryFunc)
 		: ColumnInfo(InColumnInfo)
 		, IsPlotEnabledForEntryFunc(InIsPlotEnabledForEntryFunc)
@@ -81,7 +87,7 @@ namespace UE::Audio::Insights
 		return MakePlotsWidget();
 	}
 
-	void FSoundPlotsWidgetView::ProcessPlotData(const TArray<TSharedPtr<IDashboardDataTreeViewEntry>>& DataViewEntries)
+	void FSoundPlotsWidgetView::ProcessPlotData(const TArray<TSharedPtr<IDashboardDataTreeViewEntry>>& DataViewEntries, const TArray<TSharedPtr<IDashboardDataTreeViewEntry>>& SelectedEntries)
 	{
 #if !WITH_EDITOR
 		if (GameState == EGameState::ControlledByTimemarker && !bTimeMarkerMoved)
@@ -110,7 +116,7 @@ namespace UE::Audio::Insights
 		bool bHasNewMetadata = false;
 		for (const TSharedPtr<IDashboardDataTreeViewEntry>& DataEntry : DataViewEntries)
 		{
-			CollectDataRecursive(DataEntry, bHasNewMetadata);
+			CollectDataRecursive(DataEntry, SelectedEntries, bHasNewMetadata);
 		}
 
 		// Set metadata for each widget if updated
@@ -189,6 +195,36 @@ namespace UE::Audio::Insights
 		}
 	}
 
+#if !WITH_EDITOR
+	void FSoundPlotsWidgetView::UpdatePlotSelection(const TArray<TSharedPtr<IDashboardDataTreeViewEntry>>& SelectedEntries)
+	{
+		using namespace FSoundPlotsWidgetViewPrivate;
+
+		if (!PlotWidgetMetadataPerCurve.IsValid())
+		{
+			return;
+		}
+
+		for (auto& [SourceID, Metadata] : *PlotWidgetMetadataPerCurve)
+		{
+			Metadata.PlotThickness = DefaultPlotThickness;
+		}
+
+		for (const TSharedPtr<IDashboardDataTreeViewEntry>& Entry : SelectedEntries)
+		{
+			if (!Entry.IsValid())
+			{
+				continue;
+			}
+
+			if (SAudioCurveView::FCurveMetadata* Metadata = PlotWidgetMetadataPerCurve->Find(Entry->GetEntryID()))
+			{
+				Metadata->PlotThickness = SelectedPlotThickness;
+			}
+		}
+	}
+#endif // !WITH_EDITOR
+
 	void FSoundPlotsWidgetView::OnAnalysisStarting(const double Timestamp)
 	{
 #if WITH_EDITOR
@@ -359,8 +395,10 @@ namespace UE::Audio::Insights
 			];
 	}
 
-	void FSoundPlotsWidgetView::CollectDataRecursive(const TSharedPtr<IDashboardDataTreeViewEntry>& Entry, bool& bOutHasNewMetadata)
+	void FSoundPlotsWidgetView::CollectDataRecursive(const TSharedPtr<IDashboardDataTreeViewEntry>& Entry, const TArray<TSharedPtr<IDashboardDataTreeViewEntry>>& SelectedEntries, bool& bOutHasNewMetadata)
 	{
+		using namespace FSoundPlotsWidgetViewPrivate;
+
 		if (!Entry.IsValid())
 		{
 			return;
@@ -373,7 +411,7 @@ namespace UE::Audio::Insights
 		{
 			for (const TSharedPtr<IDashboardDataTreeViewEntry>& ChildEntry : Entry->Children)
 			{
-				CollectDataRecursive(ChildEntry, bOutHasNewMetadata);
+				CollectDataRecursive(ChildEntry, SelectedEntries, bOutHasNewMetadata);
 			}
 
 			return;
@@ -416,7 +454,7 @@ namespace UE::Audio::Insights
 					double* LatestColumnTimestep = (*LatestTimestampMap)->Find(SourceId);
 					if (LatestColumnTimestep)
 					{
-						if (Timestamp < *LatestColumnTimestep)
+						if (Timestamp <= *LatestColumnTimestep)
 						{
 							continue;
 						}
@@ -443,17 +481,27 @@ namespace UE::Audio::Insights
 		}
 
 		// Create metadata for this curve if necessary
-		if (!PlotWidgetMetadataPerCurve->Contains(SourceId))
+		if (SAudioCurveView::FCurveMetadata* MetaData = PlotWidgetMetadataPerCurve->Find(SourceId))
+		{
+			const bool bPlotWasSelected = MetaData->PlotThickness != DefaultPlotThickness;
+			if (bPlotWasSelected != SelectedEntries.Contains(Entry))
+			{
+				MetaData->PlotThickness = !bPlotWasSelected ? SelectedPlotThickness : DefaultPlotThickness;
+				bOutHasNewMetadata = true;
+			}
+		}
+		else
 		{
 			FPlotCurveMetadata& NewMetadata = PlotWidgetMetadataPerCurve->Add(SourceId);
 			NewMetadata.CurveColor = FLinearColor(FColor::MakeRandomColor());
 			NewMetadata.DisplayName = Entry->GetDisplayName();
+			NewMetadata.PlotThickness = SelectedEntries.Contains(Entry) ? SelectedPlotThickness : DefaultPlotThickness;
 			bOutHasNewMetadata = true;
 		}
 
 		for (const TSharedPtr<IDashboardDataTreeViewEntry>& ChildEntry : Entry->Children)
 		{
-			CollectDataRecursive(ChildEntry, bOutHasNewMetadata);
+			CollectDataRecursive(ChildEntry, SelectedEntries, bOutHasNewMetadata);
 		}
 	}
 
diff --git a/Engine/Plugins/AudioInsights/Source/AudioInsights/Public/Views/SoundDashboardViewFactory.h b/Engine/Plugins/AudioInsights/Source/AudioInsights/Public/Views/SoundDashboardViewFactory.h
--- a/Engine/Plugins/AudioInsights/Source/AudioInsights/Public/Views/SoundDashboardViewFactory.h
+++ b/Engine/Plugins/AudioInsights/Source/AudioInsights/Public/Views/SoundDashboardViewFactory.h
@@ -113,12 +113,15 @@ namespace UE::Audio::Insights
 
 		UE_API TMap<FName, FSoundPlotsWidgetView::FPlotColumnInfo> GetPlotColumnInfo();
 
-		DECLARE_MULTICAST_DELEGATE_OneParam(FOnProcessPlotData, const TArray<TSharedPtr<IDashboardDataTreeViewEntry>>& /*DataViewEntries*/);
+		DECLARE_MULTICAST_DELEGATE_TwoParams(FOnProcessPlotData, const TArray<TSharedPtr<IDashboardDataTreeViewEntry>>& /*DataViewEntries*/, const TArray<TSharedPtr<IDashboardDataTreeViewEntry>>& /*SelectedEntries*/);
 		FOnProcessPlotData OnProcessPlotData;
 
 #if !WITH_EDITOR
 		DECLARE_MULTICAST_DELEGATE_OneParam(FOnUpdatePlotVisibility, const TArray<TSharedPtr<IDashboardDataTreeViewEntry>>& /*DataViewEntries*/);
 		FOnUpdatePlotVisibility OnUpdatePlotVisibility;
+
+		DECLARE_MULTICAST_DELEGATE_OneParam(FOnUpdatePlotSelection, const TArray<TSharedPtr<IDashboardDataTreeViewEntry>>& /*SelectedEntries*/);
+		FOnUpdatePlotSelection OnUpdatePlotSelection;
 #endif // !WITH_EDITOR
 
 	private:
@@ -135,6 +138,7 @@ namespace UE::Audio::Insights
 		UE_API virtual TSharedRef<ITableRow> OnGenerateRow(TSharedPtr<IDashboardDataTreeViewEntry> Item, const TSharedRef<STableViewBase>& OwnerTable) override;
 
 		UE_API virtual TSharedPtr<SWidget> OnConstructContextMenu() override;
+		UE_API virtual void OnSelectionChanged(TSharedPtr<IDashboardDataTreeViewEntry> SelectedItem, ESelectInfo::Type SelectInfo) override;
 		UE_API virtual FReply OnDataRowKeyInput(const FGeometry& InGeometry, const FKeyEvent& InKeyEvent) const override;
 
 		UE_API virtual const TMap<FName, FTraceTreeDashboardViewFactory::FHeaderRowColumnData>& GetHeaderRowColumns() const override;
diff --git a/Engine/Plugins/AudioInsights/Source/AudioInsights/Public/Views/SoundPlotsDashboardViewFactory.h b/Engine/Plugins/AudioInsights/Source/AudioInsights/Public/Views/SoundPlotsDashboardViewFactory.h
--- a/Engine/Plugins/AudioInsights/Source/AudioInsights/Public/Views/SoundPlotsDashboardViewFactory.h
+++ b/Engine/Plugins/AudioInsights/Source/AudioInsights/Public/Views/SoundPlotsDashboardViewFactory.h
@@ -28,12 +28,14 @@ namespace UE::Audio::Insights
 	private:
 		TSharedPtr<FSoundPlotsWidgetView> PlotsWidget;
 
-		void ProcessPlotData(const TArray<TSharedPtr<IDashboardDataTreeViewEntry>>& DataViewEntries);
+		void ProcessPlotData(const TArray<TSharedPtr<IDashboardDataTreeViewEntry>>& DataViewEntries, const TArray<TSharedPtr<IDashboardDataTreeViewEntry>>& SelectedEntries);
 
 #if !WITH_EDITOR
 		void UpdatePlotVisibility(const TArray<TSharedPtr<IDashboardDataTreeViewEntry>>& DataViewEntries);
+		void UpdatePlotSelection(const TArray<TSharedPtr<IDashboardDataTreeViewEntry>>& SelectedEntries);
 
 		FDelegateHandle OnUpdatePlotVisibility;
+		FDelegateHandle OnUpdatePlotSelection;
 #endif // !WITH_EDITOR
 
 		FDelegateHandle OnProcessPlotData;
diff --git a/Engine/Plugins/AudioInsights/Source/AudioInsights/Public/Views/SoundPlotsWidgetView.h b/Engine/Plugins/AudioInsights/Source/AudioInsights/Public/Views/SoundPlotsWidgetView.h
--- a/Engine/Plugins/AudioInsights/Source/AudioInsights/Public/Views/SoundPlotsWidgetView.h
+++ b/Engine/Plugins/AudioInsights/Source/AudioInsights/Public/Views/SoundPlotsWidgetView.h
@@ -30,9 +30,13 @@ namespace UE::Audio::Insights
 		UE_API virtual ~FSoundPlotsWidgetView();
 
 		UE_API TSharedRef<SWidget> MakeWidget();
-		UE_API void ProcessPlotData(const TArray<TSharedPtr<IDashboardDataTreeViewEntry>>& DataViewEntries);
+		UE_API void ProcessPlotData(const TArray<TSharedPtr<IDashboardDataTreeViewEntry>>& DataViewEntries, const TArray<TSharedPtr<IDashboardDataTreeViewEntry>>& SelectedEntries);
 		UE_API void UpdatePlotVisibility(const TArray<TSharedPtr<IDashboardDataTreeViewEntry>>& DataViewEntries);
 
+#if !WITH_EDITOR
+		UE_API void UpdatePlotSelection(const TArray<TSharedPtr<IDashboardDataTreeViewEntry>>& SelectedEntries);
+#endif // !WITH_EDITOR
+
 	private:
 		using FPlotCurvePoint = SAudioCurveView::FCurvePoint;
 		using FPointDataPerCurveMap = TMap<int32, TArray<FPlotCurvePoint>>; // Map of source id to data point array 
@@ -56,7 +60,7 @@ namespace UE::Audio::Insights
 		TSharedRef<SWidget> MakePlotsWidget();
 		void ConstructPlotsWidgetContents();
 
-		void CollectDataRecursive(const TSharedPtr<IDashboardDataTreeViewEntry>& Entry, bool& bOutHasNewMetadata);
+		void CollectDataRecursive(const TSharedPtr<IDashboardDataTreeViewEntry>& Entry, const TArray<TSharedPtr<IDashboardDataTreeViewEntry>>& SelectedEntries, bool& bOutHasNewMetadata);
 		void ResetPlots();
 
 		void UpdatePlotVisibilityRecursive(const TSharedPtr<IDashboardDataTreeViewEntry>& Entry);
