[Audio Insights] Color the plot's icon the same color as the plot it enables in the Sound's Dashboard

diff --git a/Engine/Plugins/AudioInsights/Source/AudioInsights/Private/Views/SoundDashboardViewFactory.cpp b/Engine/Plugins/AudioInsights/Source/AudioInsights/Private/Views/SoundDashboardViewFactory.cpp
--- a/Engine/Plugins/AudioInsights/Source/AudioInsights/Private/Views/SoundDashboardViewFactory.cpp
+++ b/Engine/Plugins/AudioInsights/Source/AudioInsights/Private/Views/SoundDashboardViewFactory.cpp
@@ -2903,6 +2903,7 @@ namespace UE::Audio::Insights
 
 						return FSlateStyle::Get().GetBrush(IconName);
 					})
+					.ColorAndOpacity(InRowData->GetEntryColor())
 				]
 			];
 	}
diff --git a/Engine/Plugins/AudioInsights/Source/AudioInsights/Private/Views/SoundPlotsWidgetView.cpp b/Engine/Plugins/AudioInsights/Source/AudioInsights/Private/Views/SoundPlotsWidgetView.cpp
--- a/Engine/Plugins/AudioInsights/Source/AudioInsights/Private/Views/SoundPlotsWidgetView.cpp
+++ b/Engine/Plugins/AudioInsights/Source/AudioInsights/Private/Views/SoundPlotsWidgetView.cpp
@@ -493,7 +493,7 @@ namespace UE::Audio::Insights
 		else
 		{
 			FPlotCurveMetadata& NewMetadata = PlotWidgetMetadataPerCurve->Add(SourceId);
-			NewMetadata.CurveColor = FLinearColor(FColor::MakeRandomColor());
+			NewMetadata.CurveColor = Entry->GetEntryColor();
 			NewMetadata.DisplayName = Entry->GetDisplayName();
 			NewMetadata.PlotThickness = SelectedEntries.Contains(Entry) ? SelectedPlotThickness : DefaultPlotThickness;
 			bOutHasNewMetadata = true;
diff --git a/Engine/Plugins/AudioInsights/Source/AudioInsights/Public/AudioInsightsDataSource.h b/Engine/Plugins/AudioInsights/Source/AudioInsights/Public/AudioInsightsDataSource.h
--- a/Engine/Plugins/AudioInsights/Source/AudioInsights/Public/AudioInsightsDataSource.h
+++ b/Engine/Plugins/AudioInsights/Source/AudioInsights/Public/AudioInsightsDataSource.h
@@ -24,6 +24,7 @@ namespace UE::Audio::Insights
 		virtual bool IsValid() const = 0;
 		virtual int32 GetEntryID() const = 0;
 		virtual const FText& GetDisplayName() const = 0;
+		virtual const FLinearColor& GetEntryColor() const = 0;
 		virtual bool HasSetInitExpansion() const = 0;
 		virtual void ResetHasSetInitExpansion() = 0;
 
diff --git a/Engine/Plugins/AudioInsights/Source/AudioInsights/Public/Messages/SoundTraceMessages.h b/Engine/Plugins/AudioInsights/Source/AudioInsights/Public/Messages/SoundTraceMessages.h
--- a/Engine/Plugins/AudioInsights/Source/AudioInsights/Public/Messages/SoundTraceMessages.h
+++ b/Engine/Plugins/AudioInsights/Source/AudioInsights/Public/Messages/SoundTraceMessages.h
@@ -196,6 +196,11 @@ namespace UE::Audio::Insights
 			return DisplayNameStr;
 		}
 
+		virtual const FLinearColor& GetEntryColor() const override
+		{
+			return bIsCategory ? FLinearColor::White : EntryColor;
+		}
+
 		virtual bool IsValid() const override { return PlayOrder != INDEX_NONE; }
 		virtual int32 GetEntryID() const override { return PlayOrder; }
 
@@ -218,6 +223,7 @@ namespace UE::Audio::Insights
 		double Timestamp = 0.0;
 		double TimeoutTimestamp = INVALID_TIMEOUT;
 		FText CategoryName;
+		const FLinearColor EntryColor = FColor::MakeRandomColor();
 		ESoundDashboardEntryType EntryType = ESoundDashboardEntryType::None;
 		EPinnedEntryType PinnedEntryType = EPinnedEntryType::None;
 		bool bIsCategory = false;
