[Audio Insights] Remove the CPU channel as one of the required channels for Audio Insights. Will hugely reduce the amount of data we track for snapshots and live traces, with no loss of data inside Audio Insights.

diff --git a/Engine/Plugins/AudioInsights/Source/AudioInsights/Private/AudioInsightsTraceModule.cpp b/Engine/Plugins/AudioInsights/Source/AudioInsights/Private/AudioInsightsTraceModule.cpp
--- a/Engine/Plugins/AudioInsights/Source/AudioInsights/Private/AudioInsightsTraceModule.cpp
+++ b/Engine/Plugins/AudioInsights/Source/AudioInsights/Private/AudioInsightsTraceModule.cpp
@@ -20,7 +20,6 @@ namespace UE::Audio::Insights
 {
 	namespace FTraceModulePrivate
 	{
-		static const TCHAR* CPUChannelName = TEXT("Cpu");
 		static const TCHAR* AudioChannelName = TEXT("Audio");
 		static const TCHAR* AudioMixerChannelName = TEXT("AudioMixer");
 		static const TCHAR* BookmarkChannelName = TEXT("Bookmark");
@@ -89,12 +88,11 @@ namespace UE::Audio::Insights
 		using namespace FTraceModulePrivate;
 
 #if WITH_EDITOR
- 		const bool bCPUOn = UE::Trace::ToggleChannel(CPUChannelName, true);
  		const bool bAudioOn = UE::Trace::ToggleChannel(AudioChannelName, true);
  		const bool bAudioMixerOn = UE::Trace::ToggleChannel(AudioMixerChannelName, true);
 		const bool bBookmarkChannelOn = UE::Trace::ToggleChannel(BookmarkChannelName, true);
 
- 		return bCPUOn && bAudioOn && bAudioMixerOn && bBookmarkChannelOn;
+ 		return bAudioOn && bAudioMixerOn && bBookmarkChannelOn;
 #else
 		if (!InstanceID.IsValid())
 		{
@@ -137,7 +135,6 @@ namespace UE::Audio::Insights
 			}
 		};
 
-		DisableIfNotMarkedToRestore(CPUChannelName);
 		DisableIfNotMarkedToRestore(AudioChannelName);
 		DisableIfNotMarkedToRestore(AudioMixerChannelName);
 		DisableIfNotMarkedToRestore(BookmarkChannelName);
