[Audio Insights] Add tooltips to column headers in the Sounds Dashboard

diff --git a/Engine/Plugins/AudioInsights/Source/AudioInsights/Private/Views/SoundDashboardViewFactory.cpp b/Engine/Plugins/AudioInsights/Source/AudioInsights/Private/Views/SoundDashboardViewFactory.cpp
--- a/Engine/Plugins/AudioInsights/Source/AudioInsights/Private/Views/SoundDashboardViewFactory.cpp
+++ b/Engine/Plugins/AudioInsights/Source/AudioInsights/Private/Views/SoundDashboardViewFactory.cpp
@@ -469,51 +469,67 @@ namespace UE::Audio::Insights
 
 		const FName MuteColumnName = "Mute";
 		const FText MuteColumnDisplayName = LOCTEXT("SoundDashboard_MuteColumnDisplayName", "Mute");
+		const FText MuteColumnTooltip = LOCTEXT("SoundDashboard_MuteColumnTooltip", "Mute a sound/category.\nMute will apply to all instances of an asset.");
 
 		const FName SoloColumnName = "Solo";
 		const FText SoloColumnDisplayName = LOCTEXT("SoundDashboard_SoloColumnDisplayName", "Solo");
+		const FText SoloColumnTooltip = LOCTEXT("SoundDashboard_SoloColumnTooltip", "Solo a sound/category.\nMute will apply to all instances of an asset.");
 
 		const FName PlotColumnName = "Plot";
 		const FText PlotColumnDisplayName = LOCTEXT("SoundDashboard_PlotColumnDisplayName", "Plot");
+		const FText PlotColumnTooltip = LOCTEXT("SoundDashboard_PlotColumnTooltip", "Plot property values for a category/sound instance/aggregate data for a parent active sound instance.\nPlots will appear in the Plots tab in Audio Insights.");
 
 		const FName NameColumnName = "Name";
 		const FText NameColumnDisplayName = LOCTEXT("SoundDashboard_NameColumnDisplayName", "Name");
+		const FText NameColumnTooltip = LOCTEXT("SoundDashboard_NameColumnTooltip", "The name of the sound/category.");
 
 		const FName PlayOrderColumnName = "PlayOrder";
 		const FText PlayOrderColumnDisplayName = LOCTEXT("SoundDashboard_PlayOrderColumnDisplayName", "Play Order");
+		const FText PlayOrderColumnTooltip = LOCTEXT("SoundDashboard_PlayOrderColumnTooltip", "The order in which a sound was played.");
 
 		const FName PriorityColumnName = "Priority";
 		const FText PriorityColumnDisplayName = LOCTEXT("SoundDashboard_PriorityColumnDisplayName", "Priority");
+		const FText PriorityColumnTooltip = LOCTEXT("SoundDashboard_PriorityColumnTooltip", "The priority of a sound.\nUsed to determine whether sound can play or remain active if channel limit is met, where higher value is higher priority.");
 
 		const FName DistanceColumnName = "Distance";
 		const FText DistanceColumnDisplayName = LOCTEXT("SoundDashboard_DistanceColumnDisplayName", "Distance");
+		const FText DistanceColumnTooltip = LOCTEXT("SoundDashboard_DistanceColumnTooltip", "The Distance from the sound source to the listener.");
 
 		const FName DistanceAttenuationColumnName = "DistanceAttenuation";
-		const FText DistanceAttenuationColumnDisplayName = LOCTEXT("SoundDashboard_DistanceAttenuationColumnDisplayName", "Distance Attenuation");
+		const FText DistanceAttenuationColumnDisplayName = LOCTEXT("SoundDashboard_DistanceAttenuationColumnDisplayName", "Distance/Occlusion Attenuation");
+		const FText DistanceAttenuationColumnTooltip = LOCTEXT("SoundDashboard_DistanceAttenuationColumnTooltip", "The amount of attenuation applied by the combined effect of distance and occlusion.");
 
 		const FName AmplitudeColumnName = "Amplitude";
 		const FText AmplitudeColumnDisplayName = LOCTEXT("SoundDashboard_AmplitudeColumnDisplayName", "Amp (Peak)");
+		const FText AmplitudeColumnTooltip = LOCTEXT("SoundDashboard_AmplitudeColumnTooltip", "The measured peak amplitude of the sound.\nMeasured in either decibels or linear volume.\nThe units of measurement can be changed in the Sound tab settings.");
 
 		const FName VolumeColumnName = "Volume";
 		const FText VolumeColumnDisplayName = LOCTEXT("SoundDashboard_VolumeColumnDisplayName", "Volume");
+		const FText VolumeColumnTooltip = LOCTEXT("SoundDashboard_VolumeColumnTooltip", "The volume multiplier applied to a sound.");
 
 		const FName LPFFreqColumnName = "LPFFreq";
 		const FText LPFFreqColumnDisplayName = LOCTEXT("SoundDashboard_LPFFreqColumnDisplayName", "LPF Freq (Hz)");
+		const FText LPFFreqColumnTooltip = LOCTEXT("SoundDashboard_LPFFreqColumnTooltip", "The frequency of the Low Pass Filter applied to the sound, measured in Hertz.\nFrequencies above the value will be filtered.");
 
 		const FName HPFFreqColumnName = "HPFFreq";
 		const FText HPFFreqColumnDisplayName = LOCTEXT("SoundDashboard_HPFFreqColumnDisplayName", "HPF Freq (Hz)");
+		const FText HPFFreqColumnTooltip = LOCTEXT("SoundDashboard_HPFColumnTooltip", "The frequency of the High Pass Filter applied to the sound, measured in Hertz.\nFrequencies below the value will be filtered.");
 
 		const FName PitchColumnName = "Pitch";
 		const FText PitchColumnDisplayName = LOCTEXT("SoundDashboard_PitchColumnDisplayName", "Pitch");
+		const FText PitchColumnTooltip = LOCTEXT("SoundDashboard_PitchColumnTooltip", "The pitch multiplier applied to a sound.");
 
 		const FName RelativeRenderCostColumnName = "RelativeRenderCost";
 		const FText RelativeRenderCostColumnDisplayName = LOCTEXT("SoundDashboard_RelativeRenderCostColumnDisplayName", "Rel. Render Cost");
+		const FText RelativeRenderCostColumnTooltip = LOCTEXT("SoundDashboard_RelativeRenderCostColumnTooltip", "The estimated relative render cost of the sound.");
 
 		const FName ActorColumnName = "Actor";
 		const FText ActorColumnDisplayName = LOCTEXT("SoundDashboard_ActorColumnDisplayName", "Actor");
+		const FText ActorColumnTooltip = LOCTEXT("SoundDashboard_ActorColumnTooltip", "The name of the Actor this sound is playing on (if applicable).");
 
 		const FName CategoryColumnName = "Category";
 		const FText CategoryColumnDisplayName = LOCTEXT("SoundDashboard_CategoryColumnDisplayName", "Category");
+		const FText CategoryColumnTooltip = LOCTEXT("SoundDashboard_CategoryColumnTooltip", "The type of sound asset this sound belongs to.");
 	} // namespace FSoundDashboardViewFactoryPrivate
 
 	/////////////////////////////////////////////////////////////////////////////////////////
@@ -1597,8 +1613,9 @@ namespace UE::Audio::Insights
 			{
 				MuteColumnName,
 				{
-					.DisplayName      = MuteColumnDisplayName,
-					.IconName         = "AudioInsights.Icon.SoundDashboard.Mute",
+					.DisplayName	  = MuteColumnDisplayName,
+					.IconName		  = "AudioInsights.Icon.SoundDashboard.Mute",
+					.TooltipText	  = MuteColumnTooltip,
 					.bShowDisplayName = false,
 					.bDefaultHidden   = !VisibleColumns.bMute,
 					.FillWidth        = 0.05f,
@@ -1610,6 +1627,7 @@ namespace UE::Audio::Insights
 				{
 					.DisplayName      = SoloColumnDisplayName,
 					.IconName         = "AudioInsights.Icon.SoundDashboard.Solo",
+					.TooltipText	  = SoloColumnTooltip,
 					.bShowDisplayName = false,
 					.bDefaultHidden   = !VisibleColumns.bSolo,
 					.FillWidth        = 0.05f,
@@ -1620,12 +1638,13 @@ namespace UE::Audio::Insights
 			{
 				PlotColumnName,
 				{
-					.DisplayName = PlotColumnDisplayName,
-					.IconName = "AudioInsights.Icon.SoundDashboard.Plots",
+					.DisplayName	  = PlotColumnDisplayName,
+					.IconName		  = "AudioInsights.Icon.SoundDashboard.Plots",
+					.TooltipText	  = PlotColumnTooltip,
 					.bShowDisplayName = false,
-					.bDefaultHidden = !VisibleColumns.bPlot,
-					.FillWidth = 0.05f,
-					.Alignment = EHorizontalAlignment::HAlign_Center
+					.bDefaultHidden   = !VisibleColumns.bPlot,
+					.FillWidth		  = 0.05f,
+					.Alignment		  = EHorizontalAlignment::HAlign_Center
 				}
 			},
 			{
@@ -1633,6 +1652,7 @@ namespace UE::Audio::Insights
 				{
 					.DisplayName      = NameColumnDisplayName,
 					.IconName         = NAME_None,
+					.TooltipText	  = NameColumnTooltip,
 					.bShowDisplayName = true,
 					.bDefaultHidden   = !VisibleColumns.bName,
 					.FillWidth        = 0.5f,
@@ -1644,6 +1664,7 @@ namespace UE::Audio::Insights
 				{
 					.DisplayName	  = PlayOrderColumnDisplayName,
 					.IconName		  = NAME_None,
+					.TooltipText	  = PlayOrderColumnTooltip,
 					.bShowDisplayName = true,
 					.bDefaultHidden   = !VisibleColumns.bPlayOrder,
 					.FillWidth		  = 0.1f,
@@ -1655,6 +1676,7 @@ namespace UE::Audio::Insights
 				{
 					.DisplayName	  = PriorityColumnDisplayName,
 					.IconName		  = NAME_None,
+					.TooltipText	  = PriorityColumnTooltip,
 					.bShowDisplayName = true,
 					.bDefaultHidden	  = !VisibleColumns.bPriority,
 					.FillWidth		  = 0.08f,
@@ -1666,6 +1688,7 @@ namespace UE::Audio::Insights
 				{
 					.DisplayName      = DistanceColumnDisplayName,
 					.IconName         = NAME_None,
+					.TooltipText	  = DistanceColumnTooltip,
 					.bShowDisplayName = true,
 					.bDefaultHidden   = !VisibleColumns.bDistance,
 					.FillWidth        = 0.1f,
@@ -1677,6 +1700,7 @@ namespace UE::Audio::Insights
 				{
 					.DisplayName	  = DistanceAttenuationColumnDisplayName,
 					.IconName		  = NAME_None,
+					.TooltipText	  = DistanceAttenuationColumnTooltip,
 					.bShowDisplayName = true,
 					.bDefaultHidden   = !VisibleColumns.bDistanceAttenuation,
 					.FillWidth		  = 0.1f,
@@ -1688,6 +1712,7 @@ namespace UE::Audio::Insights
 				{
 					.DisplayName      = AmplitudeColumnDisplayName,
 					.IconName         = NAME_None,
+					.TooltipText	  = AmplitudeColumnTooltip,
 					.bShowDisplayName = true,
 					.bDefaultHidden   = !VisibleColumns.bAmplitude,
 					.FillWidth        = 0.12f,
@@ -1699,6 +1724,7 @@ namespace UE::Audio::Insights
 				{
 					.DisplayName      = VolumeColumnDisplayName,
 					.IconName         = NAME_None,
+					.TooltipText	  = VolumeColumnTooltip,
 					.bShowDisplayName = true,
 					.bDefaultHidden   = !VisibleColumns.bVolume,
 					.FillWidth        = 0.1f,
@@ -1710,6 +1736,7 @@ namespace UE::Audio::Insights
 				{
 					.DisplayName	  = LPFFreqColumnDisplayName,
 					.IconName		  = NAME_None,
+					.TooltipText	  = LPFFreqColumnTooltip,
 					.bShowDisplayName = true,
 					.bDefaultHidden   = !VisibleColumns.bLPFFreq,
 					.FillWidth		  = 0.1f,
@@ -1721,6 +1748,7 @@ namespace UE::Audio::Insights
 				{
 					.DisplayName	  = HPFFreqColumnDisplayName,
 					.IconName		  = NAME_None,
+					.TooltipText	  = HPFFreqColumnTooltip,
 					.bShowDisplayName = true,
 					.bDefaultHidden   = !VisibleColumns.bHPFFreq,
 					.FillWidth		  = 0.1f,
@@ -1732,6 +1760,7 @@ namespace UE::Audio::Insights
 				{
 					.DisplayName      = PitchColumnDisplayName,
 					.IconName         = NAME_None,
+					.TooltipText	  = PitchColumnTooltip,
 					.bShowDisplayName = true,
 					.bDefaultHidden   = !VisibleColumns.bPitch,
 					.FillWidth        = 0.1f,
@@ -1741,36 +1770,39 @@ namespace UE::Audio::Insights
 			{
 				RelativeRenderCostColumnName,
 				{
-					.DisplayName = RelativeRenderCostColumnDisplayName,
-					.IconName = NAME_None,
+					.DisplayName	  = RelativeRenderCostColumnDisplayName,
+					.IconName		  = NAME_None,
+					.TooltipText	  = RelativeRenderCostColumnTooltip,
 					.bShowDisplayName = true,
-					.bDefaultHidden = !VisibleColumns.bRelativeRenderCost,
-					.FillWidth = 0.1f,
-					.Alignment = EHorizontalAlignment::HAlign_Left
+					.bDefaultHidden   = !VisibleColumns.bRelativeRenderCost,
+					.FillWidth		  = 0.1f,
+					.Alignment		  = EHorizontalAlignment::HAlign_Left
 				}
 			},
 #if WITH_EDITOR
 			{
 				ActorColumnName,
 				{
-					.DisplayName = ActorColumnDisplayName,
-					.IconName = NAME_None,
+					.DisplayName	  = ActorColumnDisplayName,
+					.IconName		  = NAME_None,
+					.TooltipText	  = ActorColumnTooltip,
 					.bShowDisplayName = true,
-					.bDefaultHidden = !VisibleColumns.bActor,
-					.FillWidth = 0.1f,
-					.Alignment = EHorizontalAlignment::HAlign_Left
+					.bDefaultHidden   = !VisibleColumns.bActor,
+					.FillWidth		  = 0.1f,
+					.Alignment		  = EHorizontalAlignment::HAlign_Left
 				}
 			},
 #endif // WITH_EDITOR
 			{
 				CategoryColumnName,
 				{
-					.DisplayName = LOCTEXT("SoundDashboard_CategoryColumnDisplayName", "Category"),
-					.IconName = NAME_None,
+					.DisplayName	  = CategoryColumnDisplayName,
+					.IconName		  = NAME_None,
+					.TooltipText	  = CategoryColumnTooltip,
 					.bShowDisplayName = true,
-					.bDefaultHidden = !VisibleColumns.bCategory,
-					.FillWidth = 0.1f,
-					.Alignment = EHorizontalAlignment::HAlign_Left
+					.bDefaultHidden   = !VisibleColumns.bCategory,
+					.FillWidth		  = 0.1f,
+					.Alignment		  = EHorizontalAlignment::HAlign_Left
 				}
 			}
 		};
diff --git a/Engine/Plugins/AudioInsights/Source/AudioInsights/Private/Views/TreeDashboardViewFactory.cpp b/Engine/Plugins/AudioInsights/Source/AudioInsights/Private/Views/TreeDashboardViewFactory.cpp
--- a/Engine/Plugins/AudioInsights/Source/AudioInsights/Private/Views/TreeDashboardViewFactory.cpp
+++ b/Engine/Plugins/AudioInsights/Source/AudioInsights/Private/Views/TreeDashboardViewFactory.cpp
@@ -196,7 +196,7 @@ namespace UE::Audio::Insights
 					.Padding(ColumnData.IconName != NAME_None ? 4.0f : 0.0f, 3.0f)
 					[
 						ColumnData.IconName != NAME_None
-							? SNew(SImage).Image(FSlateStyle::Get().GetBrush(ColumnData.IconName)) 
+							? SNew(SImage).Image(FSlateStyle::Get().GetBrush(ColumnData.IconName)).ToolTipText(ColumnData.TooltipText)
 							: SNullWidget::NullWidget
 					]
 					// Text (optional)
@@ -205,7 +205,7 @@ namespace UE::Audio::Insights
 					.Padding(0.0f, 3.0f, 0.0f, 3.0f)
 					[
 						ColumnData.bShowDisplayName 
-							? SNew(STextBlock).Text(ColumnData.DisplayName) 
+							? SNew(STextBlock).Text(ColumnData.DisplayName).ToolTipText(ColumnData.TooltipText)
 							: SNullWidget::NullWidget
 					]
 				];
diff --git a/Engine/Plugins/AudioInsights/Source/AudioInsights/Public/Views/TreeDashboardViewFactory.h b/Engine/Plugins/AudioInsights/Source/AudioInsights/Public/Views/TreeDashboardViewFactory.h
--- a/Engine/Plugins/AudioInsights/Source/AudioInsights/Public/Views/TreeDashboardViewFactory.h
+++ b/Engine/Plugins/AudioInsights/Source/AudioInsights/Public/Views/TreeDashboardViewFactory.h
@@ -68,6 +68,7 @@ namespace UE::Audio::Insights
 		{
 			const FText DisplayName;
 			const FName IconName;
+			const FText TooltipText = FText::GetEmpty();
 			const bool bShowDisplayName = true;
 			const bool bDefaultHidden = false;
 			const float FillWidth = 1.0f;
