[Audio Insights] Change names of tracing modes to Monitor and Trace

diff --git a/Engine/Plugins/AudioInsights/Source/AudioInsightsEditor/Private/AudioInsightsEditorDashboardFactory.cpp b/Engine/Plugins/AudioInsights/Source/AudioInsightsEditor/Private/AudioInsightsEditorDashboardFactory.cpp
--- a/Engine/Plugins/AudioInsights/Source/AudioInsightsEditor/Private/AudioInsightsEditorDashboardFactory.cpp
+++ b/Engine/Plugins/AudioInsights/Source/AudioInsightsEditor/Private/AudioInsightsEditorDashboardFactory.cpp
@@ -46,23 +46,23 @@ namespace UE::Audio::Insights
 		static const FText OnlyTraceAudioChannelsText = LOCTEXT("AudioDashboard_OnlyTraceAudioChannels_Text", "Only trace audio channels:");
 		static const FText OnlyTraceAudioChannelsTooltipText = LOCTEXT("AudioDashboard_OnlyTraceAudioChannels_TooltipText", "Disable all non-command line trace channels apart from Audio, Audio Mixer, Bookmark and CPU when starting a live monitoring / live trace. This will reduce the file size of snapshots / trace sessions while using Audio Insights.");
 
-		static const FText StartDirectTraceText = LOCTEXT("AudioDashboard_StartDirectTrace_Text", "Start Live Monitoring");
-		static const FText StartDirectTraceTooltipText = LOCTEXT("AudioDashboard_StartDirectTrace_TooltipText", "Starts a live monitoring, no .utrace file is recorded.");
+		static const FText StartDirectTraceText = LOCTEXT("AudioDashboard_StartDirectTrace_Text", "Start Monitoring");
+		static const FText StartDirectTraceTooltipText = LOCTEXT("AudioDashboard_StartDirectTrace_TooltipText", "Starts monitoring for audio trace data without saving the data to a .utrace file.");
 
-		static const FText StopDirectTraceText = LOCTEXT("AudioDashboard_StopDirectTrace_Text", "Stop Live Monitoring");
-		static const FText StopDirectTraceTooltipText = LOCTEXT("AudioDashboard_StopDirectTrace_TooltipText", "Stops the currently active Live Monitoring.");
+		static const FText StopDirectTraceText = LOCTEXT("AudioDashboard_StopDirectTrace_Text", "Stop Monitoring");
+		static const FText StopDirectTraceTooltipText = LOCTEXT("AudioDashboard_StopDirectTrace_TooltipText", "Stops the active monitoring session.");
 
-		static const FText StopLiveTraceText = LOCTEXT("AudioDashboard_StopLiveTrace_Text", "Stop Live Trace");
-		static const FText StopLiveTraceTooltipText = LOCTEXT("AudioDashboard_StopLiveTrace_TooltipText", "Stops the currently active Live Trace recording.");
+		static const FText StopLiveTraceText = LOCTEXT("AudioDashboard_StopLiveTrace_Text", "Stop Trace");
+		static const FText StopLiveTraceTooltipText = LOCTEXT("AudioDashboard_StopLiveTrace_TooltipText", "Stops the currently active trace recording.");
 
 		static const FText SaveSnapshotText = LOCTEXT("AudioDashboard_SnapshotButton_Text", "Save Snapshot");
-		static const FText SaveSnapshotTooltipText = LOCTEXT("AudioDashboard_SnapshotButton_TooltipText", "Saves a snapshot of the current trace, the snapshots are limited to 32MB");
+		static const FText SaveSnapshotTooltipText = LOCTEXT("AudioDashboard_SnapshotButton_TooltipText", "Saves a snapshot of the most recent trace data to a .utrace file. Snapshots are limited to 32MB and can be accessed inside Unreal Insights.");
 
 		static const FText SaveSnapshotNotificationText = LOCTEXT("AudioDashboard_SnapshotButton_NotificationText", "Audio Insights Snapshot saved.");
 		static const FText SaveSnapshotNotificationSubText = LOCTEXT("AudioDashboard_SnapshotButton_NotificationSubText", "A snapshot .utrace with the most recent events has been saved to your trace server.");
 
 		static const FText SaveBookmarkText = LOCTEXT("AudioDashboard_BookmarkButton_Text", "Bookmark");
-		static const FText SaveBookmarkTooltipText = LOCTEXT("AudioDashboard_BookmarkButton_TooltipText", "Ctrl+M - Creates an 'Audio Insights' bookmark in the .utrace file generated by Save Snapshot or Live Trace at the timestamp the button was pressed.");
+		static const FText SaveBookmarkTooltipText = LOCTEXT("AudioDashboard_BookmarkButton_TooltipText", "Ctrl+M - Creates an 'Audio Insights' bookmark in the .utrace file generated by Save Snapshot or Trace at the timestamp the button was pressed.");
 
 		static const FText BookmarkLabelText = LOCTEXT("AudioDashboard_BookmarkLabelTitle_Text", "Audio Insights Bookmark");
 		static const FText SaveBookmarkNotificationText = LOCTEXT("AudioDashboard_BookmarkButton_NotificationText", "Audio Insights Bookmark created.");
