[Audio Insights] Add mid-playback sound name updates

diff --git a/Engine/Plugins/AudioInsights/Source/AudioInsights/Private/Providers/SoundTraceProvider.cpp b/Engine/Plugins/AudioInsights/Source/AudioInsights/Private/Providers/SoundTraceProvider.cpp
--- a/Engine/Plugins/AudioInsights/Source/AudioInsights/Private/Providers/SoundTraceProvider.cpp
+++ b/Engine/Plugins/AudioInsights/Source/AudioInsights/Private/Providers/SoundTraceProvider.cpp
@@ -203,17 +203,19 @@ namespace UE::Audio::Insights
 
 				UE::Trace::IAnalyzer::FInterfaceBuilder& Builder = Context.InterfaceBuilder;
 
-				Builder.RouteEvent(RouteId_ActiveSoundStart,	"Audio", "SoundStart");
-				Builder.RouteEvent(RouteId_SoundWaveStart,		"Audio", "SoundWaveStart");
-				Builder.RouteEvent(RouteId_Stop,				"Audio", "SoundStop");
-				Builder.RouteEvent(RouteId_Priority,			"Audio", "SoundPriority");
-				Builder.RouteEvent(RouteId_Distance,			"Audio", "SoundDistance");
-				Builder.RouteEvent(RouteId_DistanceAttenuation,	"Audio", "SoundDistanceAttenuation");
-				Builder.RouteEvent(RouteId_Filters,				"Audio", "SoundSourceFilters");
-				Builder.RouteEvent(RouteId_Amplitude,			"Audio", "SoundSourceEnvelope");
-				Builder.RouteEvent(RouteId_Volume,				"Audio", "SoundVolume");
-				Builder.RouteEvent(RouteId_Pitch,				"Audio", "SoundPitch");
-				Builder.RouteEvent(RouteId_RelativeRenderCost,	"Audio", "SoundRelativeRenderCost");
+				Builder.RouteEvent(RouteId_ActiveSoundStart,		"Audio", "SoundStart");
+				Builder.RouteEvent(RouteId_SoundWaveStart,			"Audio", "SoundWaveStart");
+				Builder.RouteEvent(RouteId_SoundIsAlivePing,		"Audio", "SoundIsAlivePing");
+				Builder.RouteEvent(RouteId_SoundWaveIsAlivePing,	"Audio", "SoundWaveIsAlivePing");
+				Builder.RouteEvent(RouteId_Stop,					"Audio", "SoundStop");
+				Builder.RouteEvent(RouteId_Priority,				"Audio", "SoundPriority");
+				Builder.RouteEvent(RouteId_Distance,				"Audio", "SoundDistance");
+				Builder.RouteEvent(RouteId_DistanceAttenuation,		"Audio", "SoundDistanceAttenuation");
+				Builder.RouteEvent(RouteId_Filters,					"Audio", "SoundSourceFilters");
+				Builder.RouteEvent(RouteId_Amplitude,				"Audio", "SoundSourceEnvelope");
+				Builder.RouteEvent(RouteId_Volume,					"Audio", "SoundVolume");
+				Builder.RouteEvent(RouteId_Pitch,					"Audio", "SoundPitch");
+				Builder.RouteEvent(RouteId_RelativeRenderCost,		"Audio", "SoundRelativeRenderCost");
 			}
 
 			virtual bool OnEvent(uint16 RouteId, EStyle Style, const FOnEventContext& Context) override
@@ -236,6 +238,18 @@ namespace UE::Audio::Insights
 						break;
 					}
 
+					case RouteId_SoundIsAlivePing:
+					{
+						Messages.ActiveSoundIsAlivePingMessages.Enqueue(FSoundIsAlivePingMessage(Context));
+						break;
+					}
+
+					case RouteId_SoundWaveIsAlivePing:
+					{
+						Messages.SoundWaveIsAlivePingMessages.Enqueue(FSoundWaveIsAlivePingMessage(Context));
+						break;
+					}
+
 					case RouteId_Stop:
 					{
 						Messages.StopMessages.Enqueue(FSoundStopMessage{ Context });
@@ -312,6 +326,8 @@ namespace UE::Audio::Insights
 			{
 				RouteId_ActiveSoundStart,
 				RouteId_SoundWaveStart,
+				RouteId_SoundIsAlivePing,
+				RouteId_SoundWaveIsAlivePing,
 				RouteId_Stop,
 				RouteId_Priority,
 				RouteId_Distance,
@@ -353,12 +369,22 @@ namespace UE::Audio::Insights
 			return &ActiveSoundEntryReturn;
 		};
 
+		TSharedPtr<FSoundDashboardEntry> AlivePingSoundEntryReturn = nullptr;
+		auto CreateAndGetActiveSoundEntryFromPingIfNotFound = [this, &AlivePingSoundEntryReturn](const FSoundIsAlivePingMessage& Msg) -> TSharedPtr<FSoundDashboardEntry>*
+		{
+			AlivePingSoundEntryReturn = nullptr;
+
+			CreateAndGetActiveSoundEntryIfNotFound(Msg, AlivePingSoundEntryReturn);
+
+			return &AlivePingSoundEntryReturn;
+		};
+
 		TSharedPtr<FSoundDashboardEntry> SoundWaveActiveEntryReturn = nullptr;
 		auto GetActiveSoundEntryFromSoundWaveStartMessage = [this, &SoundWaveActiveEntryReturn](const FSoundWaveStartMessage& Msg) -> TSharedPtr<FSoundDashboardEntry>*
 		{
 			SoundWaveActiveEntryReturn = nullptr;
 
-			GetActiveSoundEntryFromIDs(Msg.ActiveSoundPlayOrder, Msg.DeviceId, SoundWaveActiveEntryReturn);
+			GetActiveSoundEntryFromIDs(Msg.ActiveSoundPlayOrder, SoundWaveActiveEntryReturn);
 			return &SoundWaveActiveEntryReturn;
 		};
 
@@ -367,7 +393,7 @@ namespace UE::Audio::Insights
 		{
 			ParamReturn = nullptr;
 
-			GetActiveSoundEntryFromIDs(Msg.ActiveSoundPlayOrder, Msg.DeviceId, ParamReturn);
+			GetActiveSoundEntryFromIDs(Msg.ActiveSoundPlayOrder, ParamReturn);
 			return &ParamReturn;
 		};
 
@@ -376,26 +402,17 @@ namespace UE::Audio::Insights
 		{
 			SoundEntryReturn = nullptr;
 
-			GetActiveSoundEntryFromIDs(Msg.PlayOrder, Msg.DeviceId, SoundEntryReturn);
+			GetActiveSoundEntryFromIDs(Msg.PlayOrder, SoundEntryReturn);
 			return &SoundEntryReturn;
 		};
 
-		// Process messages
-		ProcessMessageQueue<FSoundStartMessage>(TraceMessages.ActiveSoundStartMessages, GetActiveSoundEntryFromActiveStartMessage,
-		[this](const FSoundStartMessage& Msg, const TSharedPtr<FSoundDashboardEntry>* OutActiveSoundEntry)
+		auto ProcessStartMessage = [this](const FSoundStartMessage& Msg, const TSharedPtr<FSoundDashboardEntry>* OutActiveSoundEntry)
 		{
 			if (OutActiveSoundEntry == nullptr || !OutActiveSoundEntry->IsValid())
 			{
 				return;
 			}
 
-#if !WITH_EDITOR
-			if (SessionCachedMessages.IsValid())
-			{
-				SessionCachedMessages->StartCachedMessages.EmplaceBack(Msg);
-			}
-#endif // !WITH_EDITOR
-
 			FSoundDashboardEntry& ActiveSoundEntryRef = *OutActiveSoundEntry->Get();
 			ActiveSoundEntryRef.SetName(*Msg.Name);
 			ActiveSoundEntryRef.EntryType = Msg.EntryType;
@@ -409,10 +426,10 @@ namespace UE::Audio::Insights
 #if WITH_EDITOR
 			ActiveSoundEntryRef.OwnerActorName = FText::FromString(*Msg.OwnerActorName);
 #endif // WITH_EDITOR
-		});
+		};
 
-		ProcessMessageQueue<FSoundWaveStartMessage>(TraceMessages.SoundWaveStartMessages, GetActiveSoundEntryFromSoundWaveStartMessage,
-		[this](const FSoundWaveStartMessage& Msg, const TSharedPtr<FSoundDashboardEntry>* OutActiveSoundEntry)
+
+		auto ProcessSoundWaveStartMessage = [this](const FSoundWaveStartMessage& Msg, const TSharedPtr<FSoundDashboardEntry>* OutActiveSoundEntry)
 		{
 			if (OutActiveSoundEntry == nullptr || !OutActiveSoundEntry->IsValid())
 			{
@@ -424,12 +441,16 @@ namespace UE::Audio::Insights
 				return;
 			}
 
-#if !WITH_EDITOR
-			if (SessionCachedMessages.IsValid())
+			// Avoid creating duplicate entries
+			const bool bEntryExists = (*OutActiveSoundEntry)->Children.ContainsByPredicate([&Msg](const TSharedPtr<IDashboardDataTreeViewEntry>& ChildEntry)
 			{
-				SessionCachedMessages->SoundWaveStartCachedMessages.EmplaceBack(Msg);
+				return ChildEntry.IsValid() && Msg.PlayOrder == ChildEntry->GetEntryID();
+			});
+
+			if (bEntryExists)
+			{
+				return;
 			}
-#endif // !WITH_EDITOR
 			
 			TSharedPtr<FSoundDashboardEntry> SoundWaveEntry = MakeShared<FSoundDashboardEntry>();
 			SoundWaveEntry->DeviceId = Msg.DeviceId;
@@ -445,6 +466,59 @@ namespace UE::Audio::Insights
 #endif // WITH_EDITOR
 
 			(*OutActiveSoundEntry)->Children.Add(SoundWaveEntry);
+		};
+
+		// Process messages
+		ProcessMessageQueue<FSoundStartMessage>(TraceMessages.ActiveSoundStartMessages, GetActiveSoundEntryFromActiveStartMessage,
+		[this, &ProcessStartMessage](const FSoundStartMessage& Msg, const TSharedPtr<FSoundDashboardEntry>* OutActiveSoundEntry)
+		{
+			ProcessStartMessage(Msg, OutActiveSoundEntry);
+
+#if !WITH_EDITOR
+			if (SessionCachedMessages.IsValid())
+			{
+				SessionCachedMessages->StartCachedMessages.EmplaceBack(Msg);
+			}
+#endif // !WITH_EDITOR
+		});
+
+		ProcessMessageQueue<FSoundWaveStartMessage>(TraceMessages.SoundWaveStartMessages, GetActiveSoundEntryFromSoundWaveStartMessage,
+		[this, &ProcessSoundWaveStartMessage](const FSoundWaveStartMessage& Msg, const TSharedPtr<FSoundDashboardEntry>* OutActiveSoundEntry)
+		{
+			ProcessSoundWaveStartMessage(Msg, OutActiveSoundEntry);
+
+#if !WITH_EDITOR
+			if (SessionCachedMessages.IsValid())
+			{
+				SessionCachedMessages->SoundWaveStartCachedMessages.EmplaceBack(Msg);
+			}
+#endif // !WITH_EDITOR
+		});
+
+		ProcessMessageQueue<FSoundIsAlivePingMessage>(TraceMessages.ActiveSoundIsAlivePingMessages, CreateAndGetActiveSoundEntryFromPingIfNotFound,
+		[this, &ProcessStartMessage](const FSoundIsAlivePingMessage& Msg, const TSharedPtr<FSoundDashboardEntry>* OutActiveSoundEntry)
+		{
+			ProcessStartMessage(Msg, OutActiveSoundEntry);
+
+#if !WITH_EDITOR
+			if (SessionCachedMessages.IsValid())
+			{
+				SessionCachedMessages->SoundIsAlivePingCachedMessages.EmplaceBack(Msg);
+			}
+#endif // !WITH_EDITOR
+		});
+
+		ProcessMessageQueue<FSoundWaveIsAlivePingMessage>(TraceMessages.SoundWaveIsAlivePingMessages, GetActiveSoundEntryFromSoundWaveStartMessage,
+		[this, &ProcessSoundWaveStartMessage](const FSoundWaveIsAlivePingMessage& Msg, const TSharedPtr<FSoundDashboardEntry>* OutActiveSoundEntry)
+		{
+			ProcessSoundWaveStartMessage(Msg, OutActiveSoundEntry);
+			
+#if !WITH_EDITOR
+			if (SessionCachedMessages.IsValid())
+			{
+				SessionCachedMessages->SoundWaveIsAlivePingCachedMessages.EmplaceBack(Msg);
+			}
+#endif //! WITH_EDITOR
 		});
 
 		TArray<int32, TInlineAllocator<64>> EntriesWithPoppedDataPoints;
@@ -687,7 +761,7 @@ namespace UE::Audio::Insights
 		for (const SoundMessageIDs SoundMessageIDs : EntriesTimingOut)
 		{
 			TSharedPtr<FSoundDashboardEntry> ActiveSoundEntry = nullptr;
-			GetActiveSoundEntryFromIDs(SoundMessageIDs.PlayOrder, SoundMessageIDs.DeviceId, ActiveSoundEntry);
+			GetActiveSoundEntryFromIDs(SoundMessageIDs.PlayOrder, ActiveSoundEntry);
 
 			if (!ActiveSoundEntry.IsValid())
 			{
@@ -697,7 +771,7 @@ namespace UE::Audio::Insights
 
 			if (ActiveSoundEntry->TimeoutTimestamp <= CurrentTimeStamp)
 			{
-				RemoveActiveSoundEntry(SoundMessageIDs.DeviceId, ActiveSoundEntry);
+				RemoveActiveSoundEntry(ActiveSoundEntry);
 				NumToTrim++;
 			}
 			else
@@ -752,7 +826,7 @@ namespace UE::Audio::Insights
 
 				CategoryEntry->Children.Add(NewEntry);
 
-				ActiveSoundToEntryKeysMap.Add(Msg.PlayOrder, SoundEntryKeys{ Msg.EntryType });
+				ActiveSoundToEntryKeysMap.Add(Msg.PlayOrder, SoundEntryKeys{ Msg.EntryType, Msg.DeviceId });
 
 				OutReturnedSoundEntry = NewEntry;
 			}
@@ -787,7 +861,22 @@ namespace UE::Audio::Insights
 		});
 	}
 
-	void FSoundTraceProvider::GetActiveSoundEntryFromIDs(const uint32 ActiveSoundPlayOrder, const ::Audio::FDeviceId DeviceId, TSharedPtr<FSoundDashboardEntry>& OutActiveSoundEntry)
+	void FSoundTraceProvider::CreateAndGetActiveSoundEntryIfNotFound(const FSoundStartMessage& Msg, TSharedPtr<FSoundDashboardEntry>& OutReturnedSoundEntry)
+	{
+		OutReturnedSoundEntry = nullptr;
+
+		GetActiveSoundEntryFromIDs(Msg.PlayOrder, OutReturnedSoundEntry);
+		if (!OutReturnedSoundEntry.IsValid())
+		{
+			GetOrCreateActiveSoundEntry(Msg, OutReturnedSoundEntry);
+		}
+		else
+		{
+			OutReturnedSoundEntry.Reset();
+		}
+	}
+
+	void FSoundTraceProvider::GetActiveSoundEntryFromIDs(const uint32 ActiveSoundPlayOrder, TSharedPtr<FSoundDashboardEntry>& OutActiveSoundEntry)
 	{
 		const SoundEntryKeys* SoundEntryKeys = ActiveSoundToEntryKeysMap.Find(ActiveSoundPlayOrder);
 		if (SoundEntryKeys == nullptr)
@@ -795,12 +884,12 @@ namespace UE::Audio::Insights
 			return;
 		}
 
-		TSharedPtr<FSoundDashboardEntry>* CategoryEntry = FindDeviceEntry(DeviceId, SoundEntryKeys->EntryType);
+		TSharedPtr<FSoundDashboardEntry>* CategoryEntry = FindDeviceEntry(SoundEntryKeys->DeviceId, SoundEntryKeys->EntryType);
 
 		FSoundTraceProviderPrivate::FindActiveSoundEntryInCategory(CategoryEntry, ActiveSoundPlayOrder, OutActiveSoundEntry);
 	}
 
-	void FSoundTraceProvider::RemoveActiveSoundEntry(const ::Audio::FDeviceId DeviceId, TSharedPtr<FSoundDashboardEntry> OutActiveSoundEntry)
+	void FSoundTraceProvider::RemoveActiveSoundEntry(TSharedPtr<FSoundDashboardEntry> OutActiveSoundEntry)
 	{
 		if (!OutActiveSoundEntry.IsValid())
 		{
@@ -815,7 +904,7 @@ namespace UE::Audio::Insights
 			return;
 		}
 
-		TSharedPtr<FSoundDashboardEntry>* CategoryEntry = FindDeviceEntry(DeviceId, SoundEntryKeys->EntryType);
+		TSharedPtr<FSoundDashboardEntry>* CategoryEntry = FindDeviceEntry(SoundEntryKeys->DeviceId, SoundEntryKeys->EntryType);
 		if (CategoryEntry == nullptr || !CategoryEntry->IsValid())
 		{
 			return;
@@ -825,7 +914,7 @@ namespace UE::Audio::Insights
 
 		if ((*CategoryEntry)->Children.Num() == 0)
 		{
-			RemoveDeviceEntry(DeviceId, OutActiveSoundEntry->EntryType);
+			RemoveDeviceEntry(SoundEntryKeys->DeviceId, OutActiveSoundEntry->EntryType);
 		}
 
 		ActiveSoundToEntryKeysMap.Remove(ActiveSoundPlayOrder);
@@ -981,49 +1070,45 @@ namespace UE::Audio::Insights
 			return;
 		}
 
-		for (const auto& [AudioDeviceID, DeviceData] : DeviceDataMap)
+		auto InitActiveSoundEntry = [this](const FSoundStartMessage& Msg, FSoundDashboardEntry& OutActiveSoundEntry)
 		{
-			for (const auto& [EntryType, SoundDashboardEntry] : DeviceData)
+			OutActiveSoundEntry.SetName(*Msg.Name);
+			OutActiveSoundEntry.EntryType = Msg.EntryType;
+			OutActiveSoundEntry.bForceKeepEntryAlive = false;
+
+			if (!OutActiveSoundEntry.bIsPlotActive)
 			{
-				CacheIsPlottingFlagRecursive(SoundDashboardEntry);
+				OutActiveSoundEntry.bIsPlotActive = PlottingSoundEntries.Contains(OutActiveSoundEntry.PlayOrder);
 			}
-		}
-
-		TSharedPtr<FSoundDashboardEntry> ActiveSoundEntryReturn = nullptr;
-		auto GetOrCreateActiveSoundEntryFromActiveStartMessage = [this, &ActiveSoundEntryReturn](const FSoundStartMessage& Msg) -> TSharedPtr<FSoundDashboardEntry>*
-		{
-			ActiveSoundEntryReturn = nullptr;
-
-			GetOrCreateActiveSoundEntry(Msg, ActiveSoundEntryReturn);
-			return &ActiveSoundEntryReturn;
 		};
 
-		TSharedPtr<FSoundDashboardEntry> StartMessageReturn = nullptr;
-		auto GetActiveSoundEntryFromStartMessage = [this, &StartMessageReturn](const FSoundStartMessage& Msg) -> TSharedPtr<FSoundDashboardEntry>*
+		auto CreateChildEntry = [this](const FSoundWaveStartMessage& Msg, const TSharedPtr<FSoundDashboardEntry>& OutParentSoundEntry)
 		{
-			StartMessageReturn = nullptr;
-
-			GetActiveSoundEntryFromIDs(Msg.PlayOrder, Msg.DeviceId, StartMessageReturn);
-			return &StartMessageReturn;
-		};
+			TSharedPtr<FSoundDashboardEntry> ChildEntry = MakeShared<FSoundDashboardEntry>();
+			ChildEntry->DeviceId = Msg.DeviceId;
+			ChildEntry->PlayOrder = Msg.PlayOrder;
+			ChildEntry->Timestamp = Msg.Timestamp;
+			ChildEntry->SetName(*Msg.Name);
+			ChildEntry->EntryType = Msg.EntryType;
+			ChildEntry->bForceKeepEntryAlive = false;
 
-		TSharedPtr<FSoundDashboardEntry> SoundEntryReturn = nullptr;
-		auto GetActiveSoundEntryFromStopMessage = [this, &SoundEntryReturn](const FSoundStopMessage& Msg) -> TSharedPtr<FSoundDashboardEntry>*
-		{
-			SoundEntryReturn = nullptr;
+			if (!ChildEntry->bIsPlotActive)
+			{
+				ChildEntry->bIsPlotActive = PlottingSoundEntries.Contains(ChildEntry->PlayOrder);
+			}
 
-			GetActiveSoundEntryFromIDs(Msg.PlayOrder, Msg.DeviceId, SoundEntryReturn);
-			return &SoundEntryReturn;
+			OutParentSoundEntry->Children.Add(ChildEntry);
 		};
 
-		TSharedPtr<FSoundDashboardEntry> ParamReturn = nullptr;
-		auto GetActiveSoundEntryFromParameterMessage = [this, &ParamReturn](const FSoundParameterMessage& Msg) -> TSharedPtr<FSoundDashboardEntry>*
+		for (const auto& [AudioDeviceID, DeviceData] : DeviceDataMap)
 		{
-			ParamReturn = nullptr;
+			for (const auto& [EntryType, SoundDashboardEntry] : DeviceData)
+			{
+				CacheIsPlottingFlagRecursive(SoundDashboardEntry);
+			}
+		}
 
-			GetActiveSoundEntryFromIDs(Msg.ActiveSoundPlayOrder, Msg.DeviceId, ParamReturn);
-			return &ParamReturn;
-		};
+		TSet<uint32> ActiveSoundIDs;
 
 		// Collect all the start messages registered until this point in time 
 		for (const FSoundStartMessage& StartCachedMessage : SessionCachedMessages->StartCachedMessages)
@@ -1033,21 +1118,38 @@ namespace UE::Audio::Insights
 				break;
 			}
 
-			TSharedPtr<FSoundDashboardEntry>* ActiveSoundEntry = GetOrCreateActiveSoundEntryFromActiveStartMessage(StartCachedMessage);
-			if (ActiveSoundEntry == nullptr || !ActiveSoundEntry->IsValid())
+			TSharedPtr<FSoundDashboardEntry> ActiveSoundEntry = nullptr;
+			GetOrCreateActiveSoundEntry(StartCachedMessage, ActiveSoundEntry);
+
+			if (!ActiveSoundEntry.IsValid())
 			{
 				continue;
 			}
 
-			FSoundDashboardEntry& EntryRef = *ActiveSoundEntry->Get();
-			EntryRef.SetName(*StartCachedMessage.Name);
-			EntryRef.EntryType = StartCachedMessage.EntryType;
-			EntryRef.bForceKeepEntryAlive = false;
+			ActiveSoundIDs.Add(StartCachedMessage.PlayOrder);
+
+			InitActiveSoundEntry(StartCachedMessage, *ActiveSoundEntry);
+		}
 
-			if (!EntryRef.bIsPlotActive)
+		// Collect all the is alive ping messages registered until this point in time 
+		for (const FSoundIsAlivePingMessage& IsAlivePingCachedMessage : SessionCachedMessages->SoundIsAlivePingCachedMessages)
+		{
+			if (IsAlivePingCachedMessage.Timestamp > TimeMarker)
 			{
-				EntryRef.bIsPlotActive = PlottingSoundEntries.Contains(EntryRef.PlayOrder);
+				break;
 			}
+
+			ActiveSoundIDs.Add(IsAlivePingCachedMessage.PlayOrder);
+
+			TSharedPtr<FSoundDashboardEntry> ActiveSoundEntry = nullptr;
+			CreateAndGetActiveSoundEntryIfNotFound(IsAlivePingCachedMessage, ActiveSoundEntry);
+
+			if (!ActiveSoundEntry.IsValid())
+			{
+				continue;
+			}
+
+			InitActiveSoundEntry(IsAlivePingCachedMessage, *ActiveSoundEntry);
 		}
 
 		////////////////////////////////////////
@@ -1060,32 +1162,59 @@ namespace UE::Audio::Insights
 				break;
 			}
 
-			if (!FSoundTraceProviderPrivate::EntryTypeHasSoundWaveEntries(SoundWaveStartCachedMessage.EntryType))
+			GetActiveSoundEntryFromIDs(SoundWaveStartCachedMessage.ActiveSoundPlayOrder, ParentSoundEntryReturn);
+
+			if (!ParentSoundEntryReturn->IsValid() || !FSoundTraceProviderPrivate::EntryTypeHasSoundWaveEntries(ParentSoundEntryReturn->EntryType))
 			{
 				continue;
 			}
 
-			GetActiveSoundEntryFromIDs(SoundWaveStartCachedMessage.ActiveSoundPlayOrder, SoundWaveStartCachedMessage.DeviceId, ParentSoundEntryReturn);
+			ActiveSoundIDs.Add(SoundWaveStartCachedMessage.PlayOrder);
+
+			const bool bEntryExists = ParentSoundEntryReturn->Children.ContainsByPredicate([&SoundWaveStartCachedMessage](const TSharedPtr<IDashboardDataTreeViewEntry>& ChildEntry)
+			{
+				return ChildEntry.IsValid() && SoundWaveStartCachedMessage.PlayOrder == ChildEntry->GetEntryID();
+			});
 
-			if (!ParentSoundEntryReturn->IsValid())
+			if (bEntryExists)
 			{
 				continue;
 			}
 
-			TSharedPtr<FSoundDashboardEntry> ChildEntry = MakeShared<FSoundDashboardEntry>();
-			ChildEntry->DeviceId = SoundWaveStartCachedMessage.DeviceId;
-			ChildEntry->PlayOrder = SoundWaveStartCachedMessage.PlayOrder;
-			ChildEntry->Timestamp = SoundWaveStartCachedMessage.Timestamp;
-			ChildEntry->SetName(*SoundWaveStartCachedMessage.Name);
-			ChildEntry->EntryType = SoundWaveStartCachedMessage.EntryType;
-			ChildEntry->bForceKeepEntryAlive = false;
+			CreateChildEntry(SoundWaveStartCachedMessage, ParentSoundEntryReturn);
+		}
 
-			if (!ChildEntry->bIsPlotActive)
+		////////////////////////////////////////
+		// Add all soundwave ping messages registered until this point in time 
+		TSharedPtr<FSoundDashboardEntry> ParentSoundIsAlivePingEntryReturn = nullptr;
+		for (const FSoundWaveIsAlivePingMessage& SoundWaveIsAlivePingCachedMessage : SessionCachedMessages->SoundWaveIsAlivePingCachedMessages)
+		{
+			if (SoundWaveIsAlivePingCachedMessage.Timestamp > TimeMarker)
 			{
-				ChildEntry->bIsPlotActive = PlottingSoundEntries.Contains(ChildEntry->PlayOrder);
+				break;
+			}
+
+			GetActiveSoundEntryFromIDs(SoundWaveIsAlivePingCachedMessage.ActiveSoundPlayOrder, ParentSoundIsAlivePingEntryReturn);
+
+			if (!ParentSoundIsAlivePingEntryReturn->IsValid() || !FSoundTraceProviderPrivate::EntryTypeHasSoundWaveEntries(ParentSoundIsAlivePingEntryReturn->EntryType))
+			{
+				continue;
+			}
+
+			ActiveSoundIDs.Add(SoundWaveIsAlivePingCachedMessage.PlayOrder);
+
+			// Avoid creating duplicate entries when receiving the IsAlivePing message
+			const bool bEntryExists = ParentSoundIsAlivePingEntryReturn->Children.ContainsByPredicate([&SoundWaveIsAlivePingCachedMessage](const TSharedPtr<IDashboardDataTreeViewEntry>& ChildEntry)
+			{
+				return ChildEntry.IsValid() && SoundWaveIsAlivePingCachedMessage.PlayOrder == ChildEntry->GetEntryID();
+			});
+
+			if (bEntryExists)
+			{
+				continue;
 			}
 
-			ParentSoundEntryReturn->Children.Add(ChildEntry);
+			CreateChildEntry(SoundWaveIsAlivePingCachedMessage, ParentSoundIsAlivePingEntryReturn);
 		}
 
 		// Selectively remove start messages collected in the step above by knowing which sounds were stopped.
@@ -1097,17 +1226,38 @@ namespace UE::Audio::Insights
 				break;
 			}
 
-			TSharedPtr<FSoundDashboardEntry>* ActiveSoundEntry = GetActiveSoundEntryFromStopMessage(StopCachedMessage);
+			TSharedPtr<FSoundDashboardEntry> ActiveSoundEntry = nullptr;
+			GetActiveSoundEntryFromIDs(StopCachedMessage.PlayOrder, ActiveSoundEntry);
 
-			if (ActiveSoundEntry && ActiveSoundEntry->IsValid() && (*ActiveSoundEntry)->Timestamp < StopCachedMessage.Timestamp)
+			if (ActiveSoundEntry.IsValid() && ActiveSoundEntry->Timestamp < StopCachedMessage.Timestamp)
 			{
-				if ((*ActiveSoundEntry)->bIsPlotActive)
+				if (ActiveSoundEntry->bIsPlotActive)
 				{
-					(*ActiveSoundEntry)->bForceKeepEntryAlive = true;
+					ActiveSoundEntry->bForceKeepEntryAlive = true;
 				}
 				else
 				{
-					RemoveActiveSoundEntry(StopCachedMessage.DeviceId, *ActiveSoundEntry);
+					RemoveActiveSoundEntry(ActiveSoundEntry);
+				}
+			}
+		}
+
+		TArray<uint32> PlayOrderArray;
+		ActiveSoundToEntryKeysMap.GetKeys(PlayOrderArray);
+
+		// Iterate through the remaining active sounds we have in the dashboard - check whether we detected a start or IsAlive ping message for that entry
+		// If not, we assume this entry is no longer playing and remove it from the dashboard.
+		for (int ReverseIndex = PlayOrderArray.Num() - 1; ReverseIndex >= 0; --ReverseIndex)
+		{
+			const uint32 PlayOrder = PlayOrderArray[ReverseIndex];
+			if (!ActiveSoundIDs.Contains(PlayOrder))
+			{
+				TSharedPtr<FSoundDashboardEntry> ActiveSoundEntry;
+				GetActiveSoundEntryFromIDs(PlayOrder, ActiveSoundEntry);
+
+				if (ActiveSoundEntry.IsValid())
+				{
+					RemoveActiveSoundEntry(ActiveSoundEntry);
 				}
 			}
 		}
@@ -1129,7 +1279,6 @@ namespace UE::Audio::Insights
 				FSoundRelativeRenderCostMessage RelativeRenderCostMessage;
 			};
 
-			TArray<uint32> PlayOrderArray;
 			ActiveSoundToEntryKeysMap.GetKeys(PlayOrderArray);
 
 			TArray<CachedEntryInfo> CachedEntryInfos;
@@ -1208,14 +1357,15 @@ namespace UE::Audio::Insights
 			// Update the device entries with the collected info
 			for (const CachedEntryInfo& CachedEntryInfo : CachedEntryInfos)
 			{
-				TSharedPtr<FSoundDashboardEntry>* ActiveSoundEntry = GetActiveSoundEntryFromParameterMessage(CachedEntryInfo.PriorityMessage);
+				TSharedPtr<FSoundDashboardEntry> ActiveSoundEntry = nullptr;
+				GetActiveSoundEntryFromIDs(CachedEntryInfo.PriorityMessage.ActiveSoundPlayOrder, ActiveSoundEntry);
 				
-				if (ActiveSoundEntry == nullptr || !ActiveSoundEntry->IsValid())
+				if (!ActiveSoundEntry.IsValid())
 				{
 					continue;
 				}
 
-				FSoundDashboardEntry* SoundEntry = FSoundTraceProviderPrivate::FindSoundEntryForActiveSound(CachedEntryInfo.PriorityMessage.PlayOrder, ActiveSoundEntry);
+				FSoundDashboardEntry* SoundEntry = FSoundTraceProviderPrivate::FindSoundEntryForActiveSound(CachedEntryInfo.PriorityMessage.PlayOrder, &ActiveSoundEntry);
 				if (SoundEntry == nullptr)
 				{
 					continue;
diff --git a/Engine/Plugins/AudioInsights/Source/AudioInsights/Private/Providers/SoundTraceProvider.h b/Engine/Plugins/AudioInsights/Source/AudioInsights/Private/Providers/SoundTraceProvider.h
--- a/Engine/Plugins/AudioInsights/Source/AudioInsights/Private/Providers/SoundTraceProvider.h
+++ b/Engine/Plugins/AudioInsights/Source/AudioInsights/Private/Providers/SoundTraceProvider.h
@@ -39,8 +39,9 @@ namespace UE::Audio::Insights
 
 	private:
 		void GetOrCreateActiveSoundEntry(const FSoundStartMessage& Msg, TSharedPtr<FSoundDashboardEntry>& OutReturnedSoundEntry);
-		void GetActiveSoundEntryFromIDs(const uint32 PlayOrderID, const ::Audio::FDeviceId DeviceId, TSharedPtr<FSoundDashboardEntry>& OutSoundEntry);
-		void RemoveActiveSoundEntry(const ::Audio::FDeviceId DeviceId, TSharedPtr<FSoundDashboardEntry> OutEntry);
+		void CreateAndGetActiveSoundEntryIfNotFound(const FSoundStartMessage& Msg, TSharedPtr<FSoundDashboardEntry>& OutReturnedSoundEntry);
+		void GetActiveSoundEntryFromIDs(const uint32 PlayOrderID, TSharedPtr<FSoundDashboardEntry>& OutSoundEntry);
+		void RemoveActiveSoundEntry(TSharedPtr<FSoundDashboardEntry> OutEntry);
 
 		void UpdateAggregateActiveSoundData();
 		void CollectAggregateData(FSoundDashboardEntry& ActiveSoundEntry);
@@ -57,6 +58,7 @@ namespace UE::Audio::Insights
 		struct SoundEntryKeys
 		{
 			const ESoundDashboardEntryType EntryType;
+			const ::Audio::FDeviceId DeviceId;
 		};
 
 		struct SoundMessageIDs
diff --git a/Engine/Plugins/AudioInsights/Source/AudioInsights/Public/Messages/SoundTraceMessages.h b/Engine/Plugins/AudioInsights/Source/AudioInsights/Public/Messages/SoundTraceMessages.h
--- a/Engine/Plugins/AudioInsights/Source/AudioInsights/Public/Messages/SoundTraceMessages.h
+++ b/Engine/Plugins/AudioInsights/Source/AudioInsights/Public/Messages/SoundTraceMessages.h
@@ -78,6 +78,8 @@ namespace UE::Audio::Insights
 		uint32 ActiveSoundPlayOrder;
 	};
 
+	using FSoundIsAlivePingMessage = FSoundStartMessage;
+	using FSoundWaveIsAlivePingMessage = FSoundWaveStartMessage;
 	using FSoundStopMessage = FSoundMessageBase;
 
 #define DEFINE_SOUND_PARAM_MESSAGE(ClassName, SuperClassName, ParamName, Type, Default)	\
@@ -110,6 +112,8 @@ namespace UE::Audio::Insights
 	public:
 		TAnalyzerMessageQueue<FSoundStartMessage> ActiveSoundStartMessages { 0.1 };
 		TAnalyzerMessageQueue<FSoundWaveStartMessage> SoundWaveStartMessages{ 0.1 };
+		TAnalyzerMessageQueue<FSoundIsAlivePingMessage> ActiveSoundIsAlivePingMessages{ 0.1 };
+		TAnalyzerMessageQueue<FSoundWaveIsAlivePingMessage> SoundWaveIsAlivePingMessages{ 0.1 };
 		TAnalyzerMessageQueue<FSoundPriorityMessage> PriorityMessages{ 0.1 };
 		TAnalyzerMessageQueue<FSoundDistanceMessage> DistanceMessages { 0.1 };
 		TAnalyzerMessageQueue<FSoundDistanceAttenuationMessage> DistanceAttenuationMessages { 0.1 };
@@ -128,6 +132,8 @@ namespace UE::Audio::Insights
 		FSoundSessionCachedMessages(TraceServices::IAnalysisSession& InSession)
 			: StartCachedMessages(InSession.GetLinearAllocator(), 16384)
 			, SoundWaveStartCachedMessages(InSession.GetLinearAllocator(), 16384)
+			, SoundIsAlivePingCachedMessages(InSession.GetLinearAllocator(), 16384)
+			, SoundWaveIsAlivePingCachedMessages(InSession.GetLinearAllocator(), 16384)
 			, PriorityCachedMessages(InSession.GetLinearAllocator(), 16384)
 			, DistanceCachedMessages(InSession.GetLinearAllocator(), 16384)
 			, DistanceAttenuationCachedMessages(InSession.GetLinearAllocator(), 16384)
@@ -144,6 +150,8 @@ namespace UE::Audio::Insights
 
 		TraceServices::TPagedArray<FSoundStartMessage> StartCachedMessages;
 		TraceServices::TPagedArray<FSoundWaveStartMessage> SoundWaveStartCachedMessages;
+		TraceServices::TPagedArray<FSoundIsAlivePingMessage> SoundIsAlivePingCachedMessages;
+		TraceServices::TPagedArray<FSoundWaveIsAlivePingMessage> SoundWaveIsAlivePingCachedMessages;
 		TraceServices::TPagedArray<FSoundPriorityMessage> PriorityCachedMessages;
 		TraceServices::TPagedArray<FSoundDistanceMessage> DistanceCachedMessages;
 		TraceServices::TPagedArray<FSoundDistanceAttenuationMessage> DistanceAttenuationCachedMessages;
diff --git a/Engine/Source/Runtime/Engine/Private/ActiveSound.cpp b/Engine/Source/Runtime/Engine/Private/ActiveSound.cpp
--- a/Engine/Source/Runtime/Engine/Private/ActiveSound.cpp
+++ b/Engine/Source/Runtime/Engine/Private/ActiveSound.cpp
@@ -152,6 +152,7 @@ FActiveSound::FActiveSound()
 	, bIsOccluded(false)
 	, bAsyncOcclusionPending(false)
 	, PlaybackTime(0.0f)
+	, PlaybackTimeUnscaled(0.0f)
 	, PlaybackTimeNonVirtualized(0.0f)
 	, MinCurrentPitch(1.0f)
 	, RequestedStartTime(0.0f)
diff --git a/Engine/Source/Runtime/Engine/Private/AudioDevice.cpp b/Engine/Source/Runtime/Engine/Private/AudioDevice.cpp
--- a/Engine/Source/Runtime/Engine/Private/AudioDevice.cpp
+++ b/Engine/Source/Runtime/Engine/Private/AudioDevice.cpp
@@ -248,6 +248,15 @@ UE_TRACE_EVENT_BEGIN(Audio, SoundStart)
 	UE_TRACE_EVENT_FIELD(UE::Trace::WideString, OwnerActorName)
 UE_TRACE_EVENT_END()
 
+UE_TRACE_EVENT_BEGIN(Audio, SoundIsAlivePing)
+	UE_TRACE_EVENT_FIELD(uint32, DeviceId)
+	UE_TRACE_EVENT_FIELD(uint64, Timestamp)
+	UE_TRACE_EVENT_FIELD(uint32, PlayOrder)
+	UE_TRACE_EVENT_FIELD(UE::Trace::WideString, Name)
+	UE_TRACE_EVENT_FIELD(UE::Trace::WideString, SoundClassName)
+	UE_TRACE_EVENT_FIELD(UE::Trace::WideString, OwnerActorName)
+UE_TRACE_EVENT_END()
+
 UE_TRACE_EVENT_BEGIN(Audio, SoundStop)
 	UE_TRACE_EVENT_FIELD(uint32, DeviceId)
 	UE_TRACE_EVENT_FIELD(uint64, Timestamp)
@@ -2244,7 +2253,7 @@ void FAudioDevice::OnTraceStarted(FTraceAuxiliary::EConnectionType TraceType, co
 			continue;
 		}
 
-		SendSoundStartedTraceMessage(*ActiveSound);
+		SendSoundIsAlivePingTraceMessage(*ActiveSound);
 	}
 
 	for (auto& [ActiveSound, VirtualLoop] : VirtualLoops)
@@ -4323,25 +4332,60 @@ int32 FAudioDevice::GetSortedActiveWaveInstances(TArray<FWaveInstance*>& WaveIns
 
 void FAudioDevice::UpdateActiveSoundPlaybackTime(bool bIsGameTicking)
 {
+#if UE_AUDIO_PROFILERTRACE_ENABLED
+	auto ShouldSendIsAliveTraceMessage = [](const FActiveSound& ActiveSound, const float DeltaTime) -> bool
+	{
+		if (ActiveSound.IsPlayingAudio())
+		{
+			// Send a ping to Audio Insights every second to confirm the sound is still playing.
+			// This is to ensure long duration sounds (such as ambiences) will still appear in snapshot files if their start action was before the snapshot start
+			return FMath::Floor(ActiveSound.PlaybackTimeUnscaled + DeltaTime) - FMath::Floor(ActiveSound.PlaybackTimeUnscaled) > 0.0f;
+		}
+
+		return false;
+	};
+#endif // UE_AUDIO_PROFILERTRACE_ENABLED
+
 	if (bIsGameTicking)
 	{
+		const float DeltaTime = GetDeviceDeltaTime();
+
 		for (FActiveSound* ActiveSound : ActiveSounds)
 		{
+#if UE_AUDIO_PROFILERTRACE_ENABLED
+			if (ShouldSendIsAliveTraceMessage(*ActiveSound, DeltaTime))
+			{
+				SendSoundIsAlivePingTraceMessage(*ActiveSound);
+			}
+#endif // UE_AUDIO_PROFILERTRACE_ENABLED
+
 			// Scale the playback time with the device delta time and the current "min pitch" of the sounds which would play on it.
-			const float DeltaTimePitchCorrected = GetDeviceDeltaTime() * ActiveSound->MinCurrentPitch;
+			const float DeltaTimePitchCorrected = DeltaTime * ActiveSound->MinCurrentPitch;
+
 			ActiveSound->PlaybackTime += DeltaTimePitchCorrected;
+			ActiveSound->PlaybackTimeUnscaled += DeltaTime;
 			ActiveSound->PlaybackTimeNonVirtualized += DeltaTimePitchCorrected;
 		}
 	}
 	else if (GIsEditor)
 	{
+		const float DeltaTime = GetDeviceDeltaTime();
+
 		for (FActiveSound* ActiveSound : ActiveSounds)
 		{
 			if (ActiveSound->bIsPreviewSound)
 			{
+#if UE_AUDIO_PROFILERTRACE_ENABLED
+				if (ShouldSendIsAliveTraceMessage(*ActiveSound, DeviceDeltaTime))
+				{
+					SendSoundIsAlivePingTraceMessage(*ActiveSound);
+				}
+#endif // UE_AUDIO_PROFILERTRACE_ENABLED
+
 				// Scale the playback time with the device delta time and the current "min pitch" of the sounds which would play on it.
-				const float DeltaTimePitchCorrected = GetDeviceDeltaTime() * ActiveSound->MinCurrentPitch;
+				const float DeltaTimePitchCorrected = DeltaTime * ActiveSound->MinCurrentPitch;
 				ActiveSound->PlaybackTime += DeltaTimePitchCorrected;
+				ActiveSound->PlaybackTimeUnscaled += DeltaTime;
 				ActiveSound->PlaybackTimeNonVirtualized += DeltaTimePitchCorrected;
 			}
 		}
@@ -5066,33 +5110,42 @@ namespace AudioProfilerTrace
 
 		return SoundClassName;
 	}
+
+	FString GetOwnerActorLabel(const FActiveSound& InActiveSound)
+	{
+#if WITH_EDITOR
+		if (const TObjectPtr<UAudioComponent> AudioComponent = UAudioComponent::GetAudioComponentFromID(InActiveSound.GetAudioComponentID()))
+		{
+			if (const AActor* OwnerActor = AudioComponent->GetOwner())
+			{
+				return OwnerActor->GetActorLabel(false);
+			}
+		}
+#endif // WITH_EDITOR
+
+		return FString();
+	}
+
+	TObjectPtr<UObject> GetSoundObjectPointer(const FActiveSound & InActiveSound)
+	{
+		return InActiveSound.GetSound() != nullptr ? FSoftObjectPath(InActiveSound.GetSound()->GetPathName()).ResolveObject() : nullptr;
+	}
 }
 
 void FAudioDevice::SendSoundStartedTraceMessage(const FActiveSound& InActiveSound) const
 {
 	const bool bChannelEnabled = UE_TRACE_CHANNELEXPR_IS_ENABLED(AudioMixerChannel);
-	TObjectPtr<USoundBase> SoundBasePtr = InActiveSound.GetSound();
 
-	if (bChannelEnabled && SoundBasePtr)
+	if (bChannelEnabled)
 	{
-		TObjectPtr<UObject> SoundObjectPtr = FSoftObjectPath(InActiveSound.GetSound()->GetPathName()).ResolveObject();
+		const TObjectPtr<UObject> SoundObjectPtr = AudioProfilerTrace::GetSoundObjectPointer(InActiveSound);
 		if (SoundObjectPtr == nullptr)
 		{
 			return;
 		}
 
 		const FString SoundBaseAssetName = AudioProfilerTrace::GetSoundBaseAssetName(SoundObjectPtr->GetClass());
-
-		FString OwnerActorName;
-#if WITH_EDITOR
-		if (const TObjectPtr<UAudioComponent> AudioComponent = UAudioComponent::GetAudioComponentFromID(InActiveSound.GetAudioComponentID()))
-		{
-			if (const AActor* OwnerActor = AudioComponent->GetOwner())
-			{
-				OwnerActorName = OwnerActor->GetActorLabel(false);
-			}
-		}
-#endif // WITH_EDITOR
+		const FString OwnerActorLabel = AudioProfilerTrace::GetOwnerActorLabel(InActiveSound);
 
 		UE_TRACE_LOG(Audio, SoundStart, AudioMixerChannel)
 			<< SoundStart.DeviceId(InActiveSound.AudioDevice->DeviceID)
@@ -5100,7 +5153,32 @@ void FAudioDevice::SendSoundStartedTraceMessage(const FActiveSound& InActiveSoun
 			<< SoundStart.PlayOrder(InActiveSound.PlayOrder)
 			<< SoundStart.Name(*SoundObjectPtr->GetPathName())
 			<< SoundStart.SoundClassName(*SoundBaseAssetName)
-			<< SoundStart.OwnerActorName(*OwnerActorName);
+			<< SoundStart.OwnerActorName(*OwnerActorLabel);
+	}
+}
+
+void FAudioDevice::SendSoundIsAlivePingTraceMessage(const FActiveSound& InActiveSound) const
+{
+	const bool bChannelEnabled = UE_TRACE_CHANNELEXPR_IS_ENABLED(AudioMixerChannel);
+
+	if (bChannelEnabled)
+	{
+		const TObjectPtr<UObject> SoundObjectPtr = AudioProfilerTrace::GetSoundObjectPointer(InActiveSound);
+		if (SoundObjectPtr == nullptr)
+		{
+			return;
+		}
+
+		const FString SoundBaseAssetName = AudioProfilerTrace::GetSoundBaseAssetName(SoundObjectPtr->GetClass());
+		const FString OwnerActorLabel = AudioProfilerTrace::GetOwnerActorLabel(InActiveSound);
+
+		UE_TRACE_LOG(Audio, SoundIsAlivePing, AudioMixerChannel)
+			<< SoundIsAlivePing.DeviceId(InActiveSound.AudioDevice->DeviceID)
+			<< SoundIsAlivePing.Timestamp(FPlatformTime::Cycles64())
+			<< SoundIsAlivePing.PlayOrder(InActiveSound.PlayOrder)
+			<< SoundIsAlivePing.Name(*SoundObjectPtr->GetPathName())
+			<< SoundIsAlivePing.SoundClassName(*SoundBaseAssetName)
+			<< SoundIsAlivePing.OwnerActorName(*OwnerActorLabel);
 	}
 }
 #endif // UE_AUDIO_PROFILERTRACE_ENABLED
diff --git a/Engine/Source/Runtime/Engine/Private/SoundWave.cpp b/Engine/Source/Runtime/Engine/Private/SoundWave.cpp
--- a/Engine/Source/Runtime/Engine/Private/SoundWave.cpp
+++ b/Engine/Source/Runtime/Engine/Private/SoundWave.cpp
@@ -121,6 +121,16 @@ UE_TRACE_EVENT_BEGIN(Audio, SoundWaveStart)
 	UE_TRACE_EVENT_FIELD(UE::Trace::WideString, OwnerActorName)
 UE_TRACE_EVENT_END()
 
+UE_TRACE_EVENT_BEGIN(Audio, SoundWaveIsAlivePing)
+	UE_TRACE_EVENT_FIELD(uint32, DeviceId)
+	UE_TRACE_EVENT_FIELD(uint64, Timestamp)
+	UE_TRACE_EVENT_FIELD(uint32, PlayOrder)
+	UE_TRACE_EVENT_FIELD(uint32, ActiveSoundPlayOrder)
+	UE_TRACE_EVENT_FIELD(UE::Trace::WideString, Name)
+	UE_TRACE_EVENT_FIELD(UE::Trace::WideString, SoundClassName)
+	UE_TRACE_EVENT_FIELD(UE::Trace::WideString, OwnerActorName)
+UE_TRACE_EVENT_END()
+
 UE_TRACE_EVENT_BEGIN(Audio, SoundPriority)
 	UE_TRACE_EVENT_FIELD(uint32, DeviceId)
 	UE_TRACE_EVENT_FIELD(uint64, Timestamp)
@@ -3915,25 +3925,15 @@ void USoundWave::Parse(FAudioDevice* AudioDevice, const UPTRINT NodeWaveInstance
 	const bool bChannelEnabled = UE_TRACE_CHANNELEXPR_IS_ENABLED(AudioMixerChannel);
 	const bool bWaveIsActive = bIsWavePlaying || !ActiveSound.bFinished;
 
-	if (bChannelEnabled && bWaveIsActive)
+	if (bChannelEnabled && bWaveIsActive && WaveInstance->SoundClass != nullptr)
 	{
-		if (bIsNewWave && WaveInstance->SoundClass != nullptr)
+		if (bIsNewWave)
 		{
 			const TObjectPtr<UObject> SoundWaveObjectPtr = FSoftObjectPath(WaveInstance->WaveData->GetPathName()).ResolveObject();
 			if (SoundWaveObjectPtr)
 			{
 				const FString SoundWaveBaseAssetName = AudioProfilerTrace::GetSoundBaseAssetName(SoundWaveObjectPtr->GetClass());
-
-				FString OwnerActorName;
-#if WITH_EDITOR
-				if (const TObjectPtr<UAudioComponent> AudioComponent = UAudioComponent::GetAudioComponentFromID(ActiveSound.GetAudioComponentID()))
-				{
-					if (const AActor* OwnerActor = AudioComponent->GetOwner())
-					{
-						OwnerActorName = OwnerActor->GetActorLabel(false);
-					}
-				}
-#endif // WITH_EDITOR
+				const FString OwnerActorLabel = AudioProfilerTrace::GetOwnerActorLabel(ActiveSound);
 
 				UE_TRACE_LOG(Audio, SoundWaveStart, AudioMixerChannel)
 					<< SoundWaveStart.DeviceId(AudioDevice->DeviceID)
@@ -3942,7 +3942,27 @@ void USoundWave::Parse(FAudioDevice* AudioDevice, const UPTRINT NodeWaveInstance
 					<< SoundWaveStart.ActiveSoundPlayOrder(ActiveSound.GetPlayOrder())
 					<< SoundWaveStart.Name(*SoundWaveObjectPtr->GetPathName())
 					<< SoundWaveStart.SoundClassName(*SoundWaveBaseAssetName)
-					<< SoundWaveStart.OwnerActorName(*OwnerActorName);
+					<< SoundWaveStart.OwnerActorName(*OwnerActorLabel);
+			}
+		}
+		// Occasionally send a ping to Audio Insights to confirm the sound is still playing.
+		// This is to ensure long duration sounds (such as ambiences) will still appear in snapshot files if their start action was before the snapshot start
+		else if (FMath::Floor(ActiveSound.PlaybackTimeUnscaled + AudioDevice->GetDeviceDeltaTime()) - FMath::Floor(ActiveSound.PlaybackTimeUnscaled) > 0.0f)
+		{
+			const TObjectPtr<UObject> SoundWaveObjectPtr = FSoftObjectPath(WaveInstance->WaveData->GetPathName()).ResolveObject();
+			if (SoundWaveObjectPtr)
+			{
+				const FString SoundWaveBaseAssetName = AudioProfilerTrace::GetSoundBaseAssetName(SoundWaveObjectPtr->GetClass());
+				const FString OwnerActorLabel = AudioProfilerTrace::GetOwnerActorLabel(ActiveSound);
+
+				UE_TRACE_LOG(Audio, SoundWaveIsAlivePing, AudioMixerChannel)
+					<< SoundWaveIsAlivePing.DeviceId(AudioDevice->DeviceID)
+					<< SoundWaveIsAlivePing.Timestamp(FPlatformTime::Cycles64())
+					<< SoundWaveIsAlivePing.PlayOrder(WaveInstance->GetPlayOrder())
+					<< SoundWaveIsAlivePing.ActiveSoundPlayOrder(ActiveSound.GetPlayOrder())
+					<< SoundWaveIsAlivePing.Name(*SoundWaveObjectPtr->GetPathName())
+					<< SoundWaveIsAlivePing.SoundClassName(*SoundWaveBaseAssetName)
+					<< SoundWaveIsAlivePing.OwnerActorName(*OwnerActorLabel);
 			}
 		}
 
diff --git a/Engine/Source/Runtime/Engine/Public/ActiveSound.h b/Engine/Source/Runtime/Engine/Public/ActiveSound.h
--- a/Engine/Source/Runtime/Engine/Public/ActiveSound.h
+++ b/Engine/Source/Runtime/Engine/Public/ActiveSound.h
@@ -579,6 +579,9 @@ public:
 	/** Duration between now and when the sound has been started. */
 	float PlaybackTime;
 
+	/** Unscaled duration between now and when the sound has been started. This value is not scaled by the playback rate of the sound */
+	float PlaybackTimeUnscaled;
+
 	/** If virtualized, duration between last time virtualized and now. */
 	float PlaybackTimeNonVirtualized;
 
diff --git a/Engine/Source/Runtime/Engine/Public/AudioDevice.h b/Engine/Source/Runtime/Engine/Public/AudioDevice.h
--- a/Engine/Source/Runtime/Engine/Public/AudioDevice.h
+++ b/Engine/Source/Runtime/Engine/Public/AudioDevice.h
@@ -153,6 +153,7 @@ namespace ESortedActiveWaveGetType
 namespace AudioProfilerTrace
 {
 	FString GetSoundBaseAssetName(const TObjectPtr<UClass> SoundClass);
+	FString GetOwnerActorLabel(const FActiveSound& InActiveSound);
 }
 #endif // #if UE_AUDIO_PROFILERTRACE_ENABLED
 
@@ -1710,6 +1711,9 @@ private:
 #if UE_AUDIO_PROFILERTRACE_ENABLED
 	/** Notify Audio Insights that an active sound has started playing. */
 	void SendSoundStartedTraceMessage(const FActiveSound& InActiveSound) const;
+
+	/** Notify Audio Insights that an active sound is still alive and playing. */
+	void SendSoundIsAlivePingTraceMessage(const FActiveSound& InActiveSound) const;
 #endif // UE_AUDIO_PROFILERTRACE_ENABLED
 
 public:
