[Audio Insights] AudioInsightsEditorDashboardFactory: removed CPU trace channel from OnlyTraceAudioChannels checkbox tooltip description + aded tooltip to associated checkbox label

diff --git a/Engine/Plugins/AudioInsights/Source/AudioInsightsEditor/Private/AudioInsightsEditorDashboardFactory.cpp b/Engine/Plugins/AudioInsights/Source/AudioInsightsEditor/Private/AudioInsightsEditorDashboardFactory.cpp
--- a/Engine/Plugins/AudioInsights/Source/AudioInsightsEditor/Private/AudioInsightsEditorDashboardFactory.cpp
+++ b/Engine/Plugins/AudioInsights/Source/AudioInsightsEditor/Private/AudioInsightsEditorDashboardFactory.cpp
@@ -44,7 +44,7 @@ namespace UE::Audio::Insights
 		static const FText ViewResetLayoutText = LOCTEXT("AudioDashboard_View_ResetLayout_Text", "Reset Layout");
 
 		static const FText OnlyTraceAudioChannelsText = LOCTEXT("AudioDashboard_OnlyTraceAudioChannels_Text", "Only trace audio channels:");
-		static const FText OnlyTraceAudioChannelsTooltipText = LOCTEXT("AudioDashboard_OnlyTraceAudioChannels_TooltipText", "Disable all non-command line trace channels apart from Audio, Audio Mixer, Bookmark and CPU when starting a live monitoring / live trace. This will reduce the file size of snapshots / trace sessions while using Audio Insights.");
+		static const FText OnlyTraceAudioChannelsTooltipText = LOCTEXT("AudioDashboard_OnlyTraceAudioChannels_TooltipText", "Disable all non-command line trace channels apart from Audio, Audio Mixer and Bookmark when starting a live monitoring / live trace. This will reduce the file size of snapshots / trace sessions while using Audio Insights.");
 
 		static const FText StartDirectTraceText = LOCTEXT("AudioDashboard_StartDirectTrace_Text", "Start Monitoring");
 		static const FText StartDirectTraceTooltipText = LOCTEXT("AudioDashboard_StartDirectTrace_TooltipText", "Starts monitoring for audio trace data without saving the data to a .utrace file.");
@@ -555,6 +555,7 @@ namespace UE::Audio::Insights
 				[
 					SNew(STextBlock)
 					.Text(OnlyTraceAudioChannelsText)
+					.ToolTipText(OnlyTraceAudioChannelsTooltipText)
 					.Font(IPropertyTypeCustomizationUtils::GetRegularFont())
 				]
 				// Only trace audio channels checkbox
