[Audio Insights] First pass for the Event Log tab
- Split out trace utils and event log calls out of AudioDevice and into AudioTraceUtil.
- Custom events relating to an Active Sound can be sent to the Event log via Audio::Trace::EventLog::SendSoundAssetEventLog
- Created dashboard and provider for Event Log. The tab records events in order of when they happened, and can be used to control the time marker in the standalone version of Unreal Insights
- Updated names in Sound Dashboard trace messages to be consistent with the Event Log - this will break some backwards compatibility with old .urace files
- Exposed more functionality for dashboards inheriting from TableDashboardViewFactory

diff --git a/Engine/Plugins/AudioInsights/Source/AudioInsights/Private/AudioInsightsModule.cpp b/Engine/Plugins/AudioInsights/Source/AudioInsights/Private/AudioInsightsModule.cpp
--- a/Engine/Plugins/AudioInsights/Source/AudioInsights/Private/AudioInsightsModule.cpp
+++ b/Engine/Plugins/AudioInsights/Source/AudioInsights/Private/AudioInsightsModule.cpp
@@ -15,6 +15,7 @@
 
 #if !WITH_EDITOR
 #include "AudioInsightsComponent.h"
+#include "Views/AudioEventLogDashboardViewFactory.h"
 #include "Views/SoundDashboardViewFactory.h"
 #include "Views/SoundPlotsDashboardViewFactory.h"
 #include "Views/VirtualLoopDashboardViewFactory.h"
@@ -52,6 +53,7 @@ namespace UE::Audio::Insights
 			TSharedRef<FSoundDashboardViewFactory> SoundsDashboard = MakeShared<FSoundDashboardViewFactory>();
 			TSharedRef<FSoundPlotsDashboardViewFactory> PlotsDashboard = MakeShared<FSoundPlotsDashboardViewFactory>();
 
+			DashboardFactory->RegisterViewFactory(MakeShared<FAudioEventLogDashboardViewFactory>());
 			DashboardFactory->RegisterViewFactory(SoundsDashboard);
 			DashboardFactory->RegisterViewFactory(MakeShared<FVirtualLoopDashboardViewFactory>());
 			DashboardFactory->RegisterViewFactory(PlotsDashboard);
diff --git a/Engine/Plugins/AudioInsights/Source/AudioInsights/Private/AudioInsightsTimingViewExtender.cpp b/Engine/Plugins/AudioInsights/Source/AudioInsights/Private/AudioInsightsTimingViewExtender.cpp
--- a/Engine/Plugins/AudioInsights/Source/AudioInsights/Private/AudioInsightsTimingViewExtender.cpp
+++ b/Engine/Plugins/AudioInsights/Source/AudioInsights/Private/AudioInsightsTimingViewExtender.cpp
@@ -8,12 +8,19 @@ namespace UE::Audio::Insights
 {
 	void FAudioInsightsTimingViewExtender::OnBeginSession(UE::Insights::Timing::ITimingViewSession& InSession)
 	{
+		if (TimingView == nullptr)
+		{
+			TimingView = &InSession;
+		}
+
 		InSession.OnTimeMarkerChanged().AddRaw(this, &FAudioInsightsTimingViewExtender::OnTimeMarkerChanged);
 	}
 
 	void FAudioInsightsTimingViewExtender::OnEndSession(UE::Insights::Timing::ITimingViewSession& InSession)
 	{
 		InSession.OnTimeMarkerChanged().RemoveAll(this);
+
+		TimingView = nullptr;
 	}
 
 	void FAudioInsightsTimingViewExtender::Tick(UE::Insights::Timing::ITimingViewSession& InSession, const TraceServices::IAnalysisSession& InAnalysisSession)
@@ -31,6 +38,16 @@ namespace UE::Audio::Insights
 		return TraceDurationSeconds;
 	}
 
+	void FAudioInsightsTimingViewExtender::SetTimeMarker(const double Timestamp)
+	{
+		if (TimingView == nullptr)
+		{
+			return;
+		}
+
+		TimingView->SetAndCenterOnTimeMarker(Timestamp);
+	}
+
 	void FAudioInsightsTimingViewExtender::OnTimeMarkerChanged(UE::Insights::Timing::ETimeChangedFlags InFlags, double InTimeMarker)
 	{
 		if (AnalysisSession == nullptr)
diff --git a/Engine/Plugins/AudioInsights/Source/AudioInsights/Private/AudioInsightsTimingViewExtender.h b/Engine/Plugins/AudioInsights/Source/AudioInsights/Private/AudioInsightsTimingViewExtender.h
--- a/Engine/Plugins/AudioInsights/Source/AudioInsights/Private/AudioInsightsTimingViewExtender.h
+++ b/Engine/Plugins/AudioInsights/Source/AudioInsights/Private/AudioInsightsTimingViewExtender.h
@@ -24,6 +24,8 @@ namespace UE::Audio::Insights
 
 		double GetCurrentDurationSeconds() const;
 
+		void SetTimeMarker(const double Timestamp);
+
 		DECLARE_MULTICAST_DELEGATE_OneParam(FOnTimingViewTimeMarkerChanged, double /*TimeMarker*/);
 		FOnTimingViewTimeMarkerChanged OnTimingViewTimeMarkerChanged;
 
@@ -31,6 +33,7 @@ namespace UE::Audio::Insights
 		void OnTimeMarkerChanged(UE::Insights::Timing::ETimeChangedFlags InFlags, double InTimeMarker);
 
 		const TraceServices::IAnalysisSession* AnalysisSession = nullptr;
+		UE::Insights::Timing::ITimingViewSession* TimingView = nullptr;
 
 		double TraceDurationSeconds = 0.0;
 	};
diff --git a/Engine/Plugins/AudioInsights/Source/AudioInsights/Private/Messages/AudioEventLogTraceMessages.cpp b/Engine/Plugins/AudioInsights/Source/AudioInsights/Private/Messages/AudioEventLogTraceMessages.cpp
new file mode 100644
--- /dev/null
+++ b/Engine/Plugins/AudioInsights/Source/AudioInsights/Private/Messages/AudioEventLogTraceMessages.cpp
@@ -0,0 +1,48 @@
+// Copyright Epic Games, Inc. All Rights Reserved.
+#include "Messages/AudioEventLogTraceMessages.h"
+
+namespace UE::Audio::Insights
+{
+    FAudioEventLogMessage::FAudioEventLogMessage(const Trace::IAnalyzer::FOnEventContext& InContext)
+    {
+		const Trace::IAnalyzer::FEventData& EventData = InContext.EventData;
+
+		static uint32 MessageIDTracker = 0u;
+		MessageID = MessageIDTracker++;
+
+		DeviceId = static_cast<::Audio::FDeviceId>(EventData.GetValue<uint32>("DeviceId"));
+		Timestamp = InContext.EventTime.AsSeconds(EventData.GetValue<uint64>("Timestamp"));
+		PlayOrder = EventData.GetValue<uint32>("PlayOrder");
+
+		EventData.GetString("EventLogName", EventName);
+		EventData.GetString("AssetPath", AssetPath);
+		EventData.GetString("ActorName", ActorName);
+		EventData.GetString("ActorIconName", ActorIconName);
+		EventData.GetString("SoundClassName", CategoryName);
+
+		if (CategoryName == SoundClassNames::SoundCue)
+		{
+			CategoryType = EAudioEventLogSoundCategory::SoundCue;
+		}
+		else if (CategoryName == SoundClassNames::SoundWave)
+		{
+			CategoryType = EAudioEventLogSoundCategory::SoundWave;
+		}
+		else if (CategoryName == SoundClassNames::MetaSoundSource)
+		{
+			CategoryType = EAudioEventLogSoundCategory::MetaSound;
+		}
+		else if (CategoryName == SoundClassNames::SoundWaveProcedural)
+		{
+			CategoryType = EAudioEventLogSoundCategory::ProceduralSource;
+		}
+		else if (CategoryName == SoundClassNames::SoundCueTemplate)
+		{
+			CategoryType = EAudioEventLogSoundCategory::SoundCueTemplate;
+		}
+		else
+		{
+			CategoryType = EAudioEventLogSoundCategory::None;
+		}
+    }
+}
\ No newline at end of file
diff --git a/Engine/Plugins/AudioInsights/Source/AudioInsights/Private/Messages/SoundTraceMessages.cpp b/Engine/Plugins/AudioInsights/Source/AudioInsights/Private/Messages/SoundTraceMessages.cpp
--- a/Engine/Plugins/AudioInsights/Source/AudioInsights/Private/Messages/SoundTraceMessages.cpp
+++ b/Engine/Plugins/AudioInsights/Source/AudioInsights/Private/Messages/SoundTraceMessages.cpp
@@ -3,15 +3,6 @@
 
 namespace UE::Audio::Insights
 {
-	namespace SoundClassNames
-	{
-		const FString MetaSoundSource = TEXT("MetaSoundSource");
-		const FString SoundWaveProcedural = TEXT("SoundWaveProcedural");
-		const FString SoundCue = TEXT("SoundCue");
-		const FString SoundWave = TEXT("SoundWave");
-		const FString SoundCueTemplate = TEXT("SoundCueTemplate");
-	}
-
 	FSoundMessageBase::FSoundMessageBase(const Trace::IAnalyzer::FOnEventContext& InContext)
 	{
 		const Trace::IAnalyzer::FEventData& EventData = InContext.EventData;
@@ -25,7 +16,7 @@ namespace UE::Audio::Insights
 		: FSoundMessageBase(InContext)
 	{
 		const Trace::IAnalyzer::FEventData& EventData = InContext.EventData;
-		EventData.GetString("Name", Name);
+		EventData.GetString("AssetPath", Name);
 
 		FString SoundClassName;
 
@@ -57,7 +48,7 @@ namespace UE::Audio::Insights
 		}
 
 #if WITH_EDITOR
-		EventData.GetString("OwnerActorName", OwnerActorName);
+		EventData.GetString("ActorLabel", ActorLabel);
 #endif // WITH_EDITOR
 	}
 
diff --git a/Engine/Plugins/AudioInsights/Source/AudioInsights/Private/Providers/AudioEventLogTraceProvider.cpp b/Engine/Plugins/AudioInsights/Source/AudioInsights/Private/Providers/AudioEventLogTraceProvider.cpp
new file mode 100644
--- /dev/null
+++ b/Engine/Plugins/AudioInsights/Source/AudioInsights/Private/Providers/AudioEventLogTraceProvider.cpp
@@ -0,0 +1,176 @@
+// Copyright Epic Games, Inc. All Rights Reserved.
+#include "Providers/AudioEventLogTraceProvider.h"
+
+#include "Trace/Analyzer.h"
+#include "TraceServices/Model/AnalysisSession.h"
+#include "TraceServices/ModuleService.h"
+
+#if WITH_EDITOR
+#include "AudioInsightsLog.h"
+#else
+#include "Common/PagedArray.h"
+#include "Async/ParallelFor.h"
+#endif // WITH_EDITOR
+
+#define LOCTEXT_NAMESPACE "AudioInsights"
+
+namespace UE::Audio::Insights
+{
+	UE::Trace::IAnalyzer* FAudioEventLogTraceProvider::ConstructAnalyzer(TraceServices::IAnalysisSession& InSession)
+	{
+		class FAudioEventLogTraceAnalyzer : public FTraceAnalyzerBase
+		{
+		public:
+			FAudioEventLogTraceAnalyzer(TSharedRef<FAudioEventLogTraceProvider> InProvider, TraceServices::IAnalysisSession& InSession)
+				: FTraceAnalyzerBase(InProvider)
+				, Session(InSession)
+			{
+			}
+
+			virtual void OnAnalysisBegin(const FOnAnalysisContext& Context) override
+			{
+				FTraceAnalyzerBase::OnAnalysisBegin(Context);
+
+				UE::Trace::IAnalyzer::FInterfaceBuilder& Builder = Context.InterfaceBuilder;
+				Builder.RouteEvent(RouteId_EventLog, "Audio", "EventLog");
+				Builder.RouteEvent(RouteId_SoundStart, "Audio", "SoundStart");
+				Builder.RouteEvent(RouteId_SoundStop, "Audio", "SoundStop");
+			}
+
+			virtual bool OnEvent(uint16 RouteId, EStyle Style, const FOnEventContext& Context) override
+			{
+				LLM_SCOPE_BYNAME(TEXT("Insights/FAudioEventLogTraceAnalyzer"));
+
+				FAudioEventLogMessages& Messages = GetProvider<FAudioEventLogTraceProvider>().TraceMessages;
+				switch (RouteId)
+				{
+					case RouteId_EventLog:
+					case RouteId_SoundStart:
+					case RouteId_SoundStop:
+					{
+						Messages.AudioEventLogMessages.Enqueue(FAudioEventLogMessage(Context));
+						break;
+					}
+
+					default:
+					{
+						return OnEventFailure(RouteId, Style, Context);
+					}
+				}
+
+				const double Timestamp = Context.EventTime.AsSeconds(Context.EventData.GetValue<uint64>("Timestamp"));
+
+				{
+					TraceServices::FAnalysisSessionEditScope SessionEditScope(Session);
+					Session.UpdateDurationSeconds(Timestamp);
+				}
+
+				return OnEventSuccess(RouteId, Style, Context);
+			}
+
+		private:
+			enum : uint16
+			{
+				RouteId_EventLog,
+				RouteId_SoundStart,
+				RouteId_SoundStop
+			};
+
+			TraceServices::IAnalysisSession& Session;
+		};
+
+		PlayOrderToActorInfoMap.Empty();
+
+		return new FAudioEventLogTraceAnalyzer(AsShared(), InSession);
+	}
+
+	FName FAudioEventLogTraceProvider::GetName_Static()
+	{
+		return "AudioEventLogProvider";
+	}
+
+	bool FAudioEventLogTraceProvider::ProcessMessages()
+	{
+		auto AddEntryFunc = [this](const FAudioEventLogMessage& Msg)
+		{
+			TSharedPtr<FAudioEventLogDashboardEntry>* ToReturn = nullptr;
+			UpdateDeviceEntry(Msg.DeviceId, Msg.MessageID, [&ToReturn, &Msg, this](TSharedPtr<FAudioEventLogDashboardEntry>& Entry)
+			{
+				if (!Entry.IsValid())
+				{
+					Entry = CreateNewEntry(Msg);
+				}
+				else
+				{
+#if WITH_EDITOR
+					UE_LOG(LogAudioInsights, Warning, TEXT("Duplicate entry processed inside FAudioEventLogTraceProvider - each entry should be unique"));
+#endif // WITH_EDITOR
+					Entry->Timestamp = Msg.Timestamp;
+				}
+
+				ToReturn = &Entry;
+			});
+
+			return ToReturn;
+		};
+
+		ProcessMessageQueue<FAudioEventLogMessage>(TraceMessages.AudioEventLogMessages, AddEntryFunc,
+		[](const FAudioEventLogMessage& Msg, TSharedPtr<FAudioEventLogDashboardEntry>* OutEntry)
+		{
+			// Do nothing - this provider does not clear or alter entries once they are in the log
+		});
+
+		return true;
+	}
+
+	TSharedPtr<FAudioEventLogDashboardEntry> FAudioEventLogTraceProvider::CreateNewEntry(const FAudioEventLogMessage& Msg)
+	{
+		TSharedPtr<FAudioEventLogDashboardEntry> Entry = MakeShared<FAudioEventLogDashboardEntry>();
+
+		Entry->MessageID = Msg.MessageID;
+		Entry->DeviceId = Msg.DeviceId;
+		Entry->Timestamp = Msg.Timestamp;
+		Entry->PlayOrder = Msg.PlayOrder;
+		Entry->CategoryName = Msg.CategoryName;
+		Entry->Category = Msg.CategoryType;
+		Entry->EventName = Msg.EventName;
+		Entry->Name = Msg.AssetPath;
+
+		// If this message did not provide an actor name, the object may have been destroying when the message was sent
+		// Try to locate the actor name using the play order ID
+		if (Msg.ActorName.IsEmpty())
+		{
+			if (const FActorInformation* FoundActorInfo = PlayOrderToActorInfoMap.Find(Entry->PlayOrder))
+			{
+				Entry->ActorName = *FoundActorInfo->ActorName;
+#if WITH_EDITOR
+				Entry->ActorIconName = FoundActorInfo->ActorIconName;
+#else
+				Entry->ActorIconName = NAME_None;
+#endif // WITH_EDITOR
+			}
+			else
+			{
+				Entry->ActorName = LOCTEXT("AudioEventLog_None", "None").ToString();
+				Entry->ActorIconName = NAME_None;
+			}
+		}
+		else
+		{
+			Entry->ActorName = Msg.ActorName;
+
+#if WITH_EDITOR
+			Entry->ActorIconName = FName(Msg.ActorIconName);
+#else
+			Entry->ActorIconName = NAME_None;
+#endif // WITH_EDITOR
+			
+			// Cache the actor information associated with this PlayOrder ID so we can look it up later
+			PlayOrderToActorInfoMap.Add(Entry->PlayOrder, FActorInformation{ Msg.ActorName, FName(Msg.ActorIconName) });
+		}
+
+		return Entry;
+	}
+} // namespace UE::Audio::Insights
+
+#undef LOCTEXT_NAMESPACE
\ No newline at end of file
diff --git a/Engine/Plugins/AudioInsights/Source/AudioInsights/Private/Providers/AudioEventLogTraceProvider.h b/Engine/Plugins/AudioInsights/Source/AudioInsights/Private/Providers/AudioEventLogTraceProvider.h
new file mode 100644
--- /dev/null
+++ b/Engine/Plugins/AudioInsights/Source/AudioInsights/Private/Providers/AudioEventLogTraceProvider.h
@@ -0,0 +1,39 @@
+// Copyright Epic Games, Inc. All Rights Reserved.
+#pragma once
+
+#include "AudioInsightsTraceProviderBase.h"
+#include "Messages/AudioEventLogTraceMessages.h"
+
+namespace UE::Audio::Insights
+{
+	class FAudioEventLogTraceProvider
+		: public TDeviceDataMapTraceProvider<uint32, TSharedPtr<FAudioEventLogDashboardEntry>>
+		, public TSharedFromThis<FAudioEventLogTraceProvider>
+	{
+	public:
+		FAudioEventLogTraceProvider()
+			: TDeviceDataMapTraceProvider<uint32, TSharedPtr<FAudioEventLogDashboardEntry>>(GetName_Static())
+		{
+		}
+
+		virtual ~FAudioEventLogTraceProvider() = default;
+		virtual UE::Trace::IAnalyzer* ConstructAnalyzer(TraceServices::IAnalysisSession& InSession) override;
+
+		static FName GetName_Static();
+
+	private:
+		virtual bool ProcessMessages() override;
+
+		TSharedPtr<FAudioEventLogDashboardEntry> CreateNewEntry(const FAudioEventLogMessage& Msg);
+
+		FAudioEventLogMessages TraceMessages;
+
+		struct FActorInformation
+		{
+			const FString ActorName;
+			const FName ActorIconName;
+		};
+
+		TMap<uint32, FActorInformation> PlayOrderToActorInfoMap;
+	};
+} // namespace UE::Audio::Insights
diff --git a/Engine/Plugins/AudioInsights/Source/AudioInsights/Private/Providers/SoundTraceProvider.cpp b/Engine/Plugins/AudioInsights/Source/AudioInsights/Private/Providers/SoundTraceProvider.cpp
--- a/Engine/Plugins/AudioInsights/Source/AudioInsights/Private/Providers/SoundTraceProvider.cpp
+++ b/Engine/Plugins/AudioInsights/Source/AudioInsights/Private/Providers/SoundTraceProvider.cpp
@@ -424,7 +424,7 @@ namespace UE::Audio::Insights
 			}
 
 #if WITH_EDITOR
-			ActiveSoundEntryRef.OwnerActorName = FText::FromString(*Msg.OwnerActorName);
+			ActiveSoundEntryRef.ActorLabel = FText::FromString(*Msg.ActorLabel);
 #endif // WITH_EDITOR
 		};
 
@@ -462,7 +462,7 @@ namespace UE::Audio::Insights
 			SoundWaveEntry->PinnedEntryType = (*OutActiveSoundEntry)->PinnedEntryType;
 
 #if WITH_EDITOR
-			SoundWaveEntry->OwnerActorName = FText::FromString(*Msg.OwnerActorName);
+			SoundWaveEntry->ActorLabel = FText::FromString(*Msg.ActorLabel);
 #endif // WITH_EDITOR
 
 			(*OutActiveSoundEntry)->Children.Add(SoundWaveEntry);
@@ -1164,7 +1164,7 @@ namespace UE::Audio::Insights
 
 			GetActiveSoundEntryFromIDs(SoundWaveStartCachedMessage.ActiveSoundPlayOrder, ParentSoundEntryReturn);
 
-			if (!ParentSoundEntryReturn->IsValid() || !FSoundTraceProviderPrivate::EntryTypeHasSoundWaveEntries(ParentSoundEntryReturn->EntryType))
+			if (ParentSoundEntryReturn == nullptr || !ParentSoundEntryReturn->IsValid() || !FSoundTraceProviderPrivate::EntryTypeHasSoundWaveEntries(ParentSoundEntryReturn->EntryType))
 			{
 				continue;
 			}
@@ -1196,7 +1196,7 @@ namespace UE::Audio::Insights
 
 			GetActiveSoundEntryFromIDs(SoundWaveIsAlivePingCachedMessage.ActiveSoundPlayOrder, ParentSoundIsAlivePingEntryReturn);
 
-			if (!ParentSoundIsAlivePingEntryReturn->IsValid() || !FSoundTraceProviderPrivate::EntryTypeHasSoundWaveEntries(ParentSoundIsAlivePingEntryReturn->EntryType))
+			if (ParentSoundEntryReturn == nullptr || !ParentSoundIsAlivePingEntryReturn->IsValid() || !FSoundTraceProviderPrivate::EntryTypeHasSoundWaveEntries(ParentSoundIsAlivePingEntryReturn->EntryType))
 			{
 				continue;
 			}
diff --git a/Engine/Plugins/AudioInsights/Source/AudioInsights/Private/Views/AudioEventLogDashboardViewFactory.cpp b/Engine/Plugins/AudioInsights/Source/AudioInsights/Private/Views/AudioEventLogDashboardViewFactory.cpp
new file mode 100644
--- /dev/null
+++ b/Engine/Plugins/AudioInsights/Source/AudioInsights/Private/Views/AudioEventLogDashboardViewFactory.cpp
@@ -0,0 +1,387 @@
+// Copyright Epic Games, Inc. All Rights Reserved.
+#include "Views/AudioEventLogDashboardViewFactory.h"
+
+#include "AudioInsightsModule.h"
+#include "AudioInsightsStyle.h"
+#include "Messages/AudioEventLogTraceMessages.h"
+#include "Providers/AudioEventLogTraceProvider.h"
+
+#define LOCTEXT_NAMESPACE "AudioInsights"
+
+namespace UE::Audio::Insights
+{
+	/////////////////////////////////////////////////////////////////////////////////////////
+	// AudioEventLogPrivate
+	namespace AudioEventLogPrivate
+	{
+		const FAudioEventLogDashboardEntry& CastEntry(const IDashboardDataViewEntry& InData)
+		{
+			return static_cast<const FAudioEventLogDashboardEntry&>(InData);
+		};
+
+		const FName MessageIDColumnID = "MessageID";
+		const FName TimestampColumnID = "Timestamp";
+		const FName PlayOrderColumnID = "PlayOrder";
+		const FName EventColumnID = "Event";
+		const FName AssetNameColumnID = "Asset";
+		const FName ActorColumnID = "Actor";
+		const FName CategoryColumnID = "Category";
+	} // namespace AudioEventLogPrivate
+
+	/////////////////////////////////////////////////////////////////////////////////////////
+	// FAudioEventLogDashboardViewFactory
+	FAudioEventLogDashboardViewFactory::FAudioEventLogDashboardViewFactory()
+	{
+		using namespace AudioEventLogPrivate;
+
+		FTraceModule& AudioInsightsTraceModule = static_cast<FTraceModule&>(FAudioInsightsModule::GetChecked().GetTraceModule());
+
+		const TSharedPtr<FAudioEventLogTraceProvider> AudioEventLogTraceProvider = MakeShared<FAudioEventLogTraceProvider>();
+
+		AudioInsightsTraceModule.AddTraceProvider(AudioEventLogTraceProvider);
+
+		Providers = TArray<TSharedPtr<FTraceProviderBase>>
+		{
+			AudioEventLogTraceProvider
+		};
+
+		SortByColumn = TimestampColumnID;
+		SortMode = EColumnSortMode::Ascending;
+	}
+
+	FName FAudioEventLogDashboardViewFactory::GetName() const
+	{
+		return "EventLog";
+	}
+
+	FText FAudioEventLogDashboardViewFactory::GetDisplayName() const
+	{
+		return LOCTEXT("AudioDashboard_AudioEventLog_Name", "Event Log");
+	}
+
+	EDefaultDashboardTabStack FAudioEventLogDashboardViewFactory::GetDefaultTabStack() const
+	{
+		return EDefaultDashboardTabStack::Log;
+	}
+
+	FSlateIcon FAudioEventLogDashboardViewFactory::GetIcon() const
+	{
+		return FSlateStyle::Get().CreateIcon("AudioInsights.Icon.Log");
+	}
+
+	TSharedRef<SWidget> FAudioEventLogDashboardViewFactory::MakeWidget(TSharedRef<SDockTab> OwnerTab, const FSpawnTabArgs& SpawnTabArgs)
+	{
+#if !WITH_EDITOR
+		FAudioInsightsModule& AudioInsightsModule = FAudioInsightsModule::GetChecked();
+		AudioInsightsModule.GetTimingViewExtender().OnTimingViewTimeMarkerChanged.AddSP(this, &FAudioEventLogDashboardViewFactory::OnTimingViewTimeMarkerChanged);
+#endif // !WITH_EDITOR
+
+		if (!DashboardWidget.IsValid())
+		{
+			DashboardWidget = FTraceTableDashboardViewFactory::MakeWidget(OwnerTab, SpawnTabArgs);
+
+			if (FilteredEntriesListView.IsValid())
+			{
+				FilteredEntriesListView->SetSelectionMode(ESelectionMode::SingleToggle);
+			}
+		}
+
+		return DashboardWidget->AsShared();
+	}
+
+	void FAudioEventLogDashboardViewFactory::RefreshFilteredEntriesListView()
+	{
+		FTraceObjectTableDashboardViewFactory::RefreshFilteredEntriesListView();
+
+		if (FilteredEntriesListView.IsValid())
+		{
+			if (FilteredEntriesListView->GetNumItemsSelected() == 0)
+			{
+				if (SortMode == EColumnSortMode::Ascending)
+				{
+					FilteredEntriesListView->ScrollToBottom();
+				}
+				else if (SortMode == EColumnSortMode::Descending)
+				{
+					FilteredEntriesListView->ScrollToTop();
+				}
+			}
+		}
+	}
+
+	void FAudioEventLogDashboardViewFactory::ProcessEntries(FTraceTableDashboardViewFactory::EProcessReason Reason)
+	{
+		const FString FilterString = GetSearchFilterText().ToString();
+		FTraceTableDashboardViewFactory::FilterEntries<FAudioEventLogTraceProvider>([&FilterString](const IDashboardDataViewEntry& Entry)
+		{
+			const FAudioEventLogDashboardEntry& EventLogEntry = static_cast<const FAudioEventLogDashboardEntry&>(Entry);
+			if (EventLogEntry.GetDisplayName().ToString().Contains(FilterString) || EventLogEntry.ActorName.Contains(FilterString))
+			{
+				return false;
+			}
+
+			return true;
+		});
+	}
+
+	const TMap<FName, FTraceTableDashboardViewFactory::FColumnData>& FAudioEventLogDashboardViewFactory::GetColumns() const
+	{
+		using namespace AudioEventLogPrivate;
+
+		auto CreateColumnData = [this]()
+		{
+			return TMap<FName, FTraceTableDashboardViewFactory::FColumnData>
+			{
+				{
+					MessageIDColumnID,
+					{
+						LOCTEXT("AudioEventLog_MessageIDColumnDisplayName", "ID"),
+						[](const IDashboardDataViewEntry& InData) { return FText::AsNumber(AudioEventLogPrivate::CastEntry(InData).MessageID); },
+						nullptr /* GetIconName */,
+						true /* bDefaultHidden */,
+						0.1f /* FillWidth */
+					}
+				},
+				{
+					TimestampColumnID,
+					{
+						LOCTEXT("AudioEventLog_TimestampColumnDisplayName", "Timestamp"),
+						[](const IDashboardDataViewEntry& InData) { return FText::AsNumber(AudioEventLogPrivate::CastEntry(InData).Timestamp, FSlateStyle::Get().GetTimeFormat()); },
+						nullptr /* GetIconName */,
+						false /* bDefaultHidden */,
+						0.15f /* FillWidth */
+					}
+				},
+				{
+					PlayOrderColumnID,
+					{
+						LOCTEXT("AudioEventLog_PlayOrderColumnDisplayName", "Play Order"),
+						[](const IDashboardDataViewEntry& InData) { return FText::AsNumber(AudioEventLogPrivate::CastEntry(InData).PlayOrder); },
+						nullptr /* GetIconName */,
+						false /* bDefaultHidden */,
+						0.15f /* FillWidth */
+					}
+				},
+				{
+					EventColumnID,
+					{
+						LOCTEXT("AudioEventLog_EventColumnDisplayName", "Event"),
+						[](const IDashboardDataViewEntry& InData) { return FText::FromString(AudioEventLogPrivate::CastEntry(InData).EventName); },
+						nullptr /* GetIconName */,
+						false /* bDefaultHidden */,
+						0.2f  /* FillWidth */
+					}
+				},
+				{
+					AssetNameColumnID,
+					{
+						LOCTEXT("AudioEventLog_AssetColumnDisplayName", "Asset"),
+						[](const IDashboardDataViewEntry& InData) { return AudioEventLogPrivate::CastEntry(InData).GetDisplayName(); },
+						[this](const IDashboardDataViewEntry& InData) -> FName
+						{
+							const FAudioEventLogDashboardEntry& Entry = CastEntry(InData);
+
+							switch (Entry.Category)
+							{
+								case EAudioEventLogSoundCategory::MetaSound:
+									return FName("AudioInsights.Icon.SoundDashboard.MetaSound");
+								case EAudioEventLogSoundCategory::SoundCue:
+									return FName("AudioInsights.Icon.SoundDashboard.SoundCue");
+								case EAudioEventLogSoundCategory::ProceduralSource:
+									return FName("AudioInsights.Icon.SoundDashboard.ProceduralSource");
+								case EAudioEventLogSoundCategory::SoundWave:
+									return FName("AudioInsights.Icon.SoundDashboard.SoundWave");
+								case EAudioEventLogSoundCategory::SoundCueTemplate:
+									return FName("AudioInsights.Icon.SoundDashboard.SoundCue");
+								case EAudioEventLogSoundCategory::None:
+								default:
+									break;
+							}
+
+							return NAME_None;
+						},
+						false /* bDefaultHidden */,
+						0.3f  /* FillWidth */
+					}
+				},
+				{
+					ActorColumnID,
+					{
+						LOCTEXT("AudioEventLog_ActorColumnDisplayName", "Actor"),
+						[](const IDashboardDataViewEntry& InData) { return FText::FromString(AudioEventLogPrivate::CastEntry(InData).ActorName); },
+						[this](const IDashboardDataViewEntry& InData) -> FName
+						{
+							return CastEntry(InData).ActorIconName;
+						},
+						false /* bDefaultHidden */,
+						0.3f  /* FillWidth */
+					}
+				},
+				{
+					CategoryColumnID,
+					{
+						LOCTEXT("AudioEventLog_CategoryColumnDisplayName", "Category"),
+						[](const IDashboardDataViewEntry& InData) { return FText::FromString(AudioEventLogPrivate::CastEntry(InData).CategoryName); },
+						nullptr /* GetIconName */,
+						false /* bDefaultHidden */,
+						0.3f  /* FillWidth */
+					}
+				}
+			};
+		};
+
+		static const TMap<FName, FTraceTableDashboardViewFactory::FColumnData> ColumnData = CreateColumnData();
+		return ColumnData;
+	}
+
+	void FAudioEventLogDashboardViewFactory::SortTable()
+	{
+		using namespace AudioEventLogPrivate;
+
+		auto SortByMessageID = [](const FAudioEventLogDashboardEntry& First, const FAudioEventLogDashboardEntry& Second)
+		{
+			return First.MessageID < Second.MessageID;
+		};
+
+		auto SortByTimestamp = [&SortByMessageID](const FAudioEventLogDashboardEntry& First, const FAudioEventLogDashboardEntry& Second)
+		{
+			const double ComparisonDiff = First.Timestamp - Second.Timestamp;
+
+			if (FMath::IsNearlyZero(ComparisonDiff, UE_KINDA_SMALL_NUMBER))
+			{
+				return SortByMessageID(First, Second);
+			}
+
+			return ComparisonDiff < 0.0;
+		};
+
+		SortByPredicate(SortByTimestamp);
+	}
+
+	bool FAudioEventLogDashboardViewFactory::IsColumnSortable(const FName& ColumnId) const
+	{
+		using namespace AudioEventLogPrivate;
+
+		return ColumnId == TimestampColumnID;
+	}
+
+	void FAudioEventLogDashboardViewFactory::OnSelectionChanged(TSharedPtr<IDashboardDataViewEntry> SelectedItem, ESelectInfo::Type SelectInfo)
+	{
+		using namespace AudioEventLogPrivate;
+
+		if (!FilteredEntriesListView.IsValid())
+		{
+			return;
+		}
+
+		if (!SelectedItem.IsValid())
+		{
+			FilteredEntriesListView->ClearHighlightedItems();
+			return;
+		}
+
+		const FAudioEventLogDashboardEntry& SelectedEventLogEntry = CastEntry(*SelectedItem);
+
+		for (const TSharedPtr<IDashboardDataViewEntry>& Entry : DataViewEntries)
+		{
+			if (!Entry.IsValid())
+			{
+				continue;
+			}
+
+			FilteredEntriesListView->SetItemHighlighted(Entry, CastEntry(*Entry).PlayOrder == SelectedEventLogEntry.PlayOrder);
+		}
+
+		if (SelectInfo != ESelectInfo::Type::Direct)
+		{
+#if WITH_EDITOR
+			// TODO UE-301957: Set timemarker via delegate
+#else
+			FAudioInsightsModule& AudioInsightsModule = FAudioInsightsModule::GetChecked();
+			AudioInsightsModule.GetTimingViewExtender().SetTimeMarker(SelectedEventLogEntry.Timestamp);
+#endif // WITH_EDITOR
+		}
+	}
+
+	const FTableRowStyle* FAudioEventLogDashboardViewFactory::GetRowStyle() const
+	{
+		return &FSlateStyle::Get().GetWidgetStyle<FTableRowStyle>("TreeDashboard.TableViewRow");
+	}
+
+	void FAudioEventLogDashboardViewFactory::SortByPredicate(TFunctionRef<bool(const FAudioEventLogDashboardEntry&, const FAudioEventLogDashboardEntry&)> Predicate)
+	{
+		using namespace AudioEventLogPrivate;
+
+		auto SortDashboardEntries = [this](const TSharedPtr<IDashboardDataViewEntry>& First, const TSharedPtr<IDashboardDataViewEntry>& Second, TFunctionRef<bool(const FAudioEventLogDashboardEntry&, const FAudioEventLogDashboardEntry&)> Predicate)
+		{
+			return Predicate(CastEntry(*First), CastEntry(*Second));
+		};
+
+		if (SortMode == EColumnSortMode::Ascending)
+		{
+			DataViewEntries.Sort([&SortDashboardEntries, &Predicate](const TSharedPtr<IDashboardDataViewEntry>& A, const TSharedPtr<IDashboardDataViewEntry>& B)
+			{
+				return SortDashboardEntries(A, B, Predicate);
+			});
+		}
+		else if (SortMode == EColumnSortMode::Descending)
+		{
+			DataViewEntries.Sort([&SortDashboardEntries, &Predicate](const TSharedPtr<IDashboardDataViewEntry>& A, const TSharedPtr<IDashboardDataViewEntry>& B)
+			{
+				return SortDashboardEntries(B, A, Predicate);
+			});
+		}
+	}
+
+	#if !WITH_EDITOR
+	void FAudioEventLogDashboardViewFactory::OnTimingViewTimeMarkerChanged(double InTimeMarker)
+	{
+		using namespace AudioEventLogPrivate;
+
+		if (!FilteredEntriesListView.IsValid())
+		{
+			return;
+		}
+
+		// We assume the data entries are in timestamp order (earliest to latest)
+		for (int32 Index = 0; Index < DataViewEntries.Num(); ++Index)
+		{
+			const TSharedPtr<IDashboardDataViewEntry>& Entry = DataViewEntries[Index];
+			if (!Entry.IsValid())
+			{
+				continue;
+			}
+
+			const FAudioEventLogDashboardEntry& EventLogEntry = CastEntry(*Entry);
+			if (EventLogEntry.Timestamp == InTimeMarker)
+			{
+				FilteredEntriesListView->SetSelection(Entry);
+				break;
+			}
+			else if (EventLogEntry.Timestamp > InTimeMarker)
+			{
+				const int32 SelectedIndex = Index - 1;
+				if (SelectedIndex < 0)
+				{
+					FilteredEntriesListView->ClearSelection();
+				}
+				else
+				{
+					FilteredEntriesListView->SetSelection(DataViewEntries[SelectedIndex]);
+				}
+
+				break;
+			}
+			// If we've reached the last entry, the timemarker must be beyond the final entry
+			// Select the last entry in the list in this case
+			else if (Index == DataViewEntries.Num() - 1)
+			{
+				FilteredEntriesListView->SetSelection(Entry);
+			}
+		}
+	}
+	#endif // !WITH_EDITOR
+
+} // namespace UE::Audio::Insights
+
+#undef LOCTEXT_NAMESPACE
diff --git a/Engine/Plugins/AudioInsights/Source/AudioInsights/Private/Views/SoundDashboardViewFactory.cpp b/Engine/Plugins/AudioInsights/Source/AudioInsights/Private/Views/SoundDashboardViewFactory.cpp
--- a/Engine/Plugins/AudioInsights/Source/AudioInsights/Private/Views/SoundDashboardViewFactory.cpp
+++ b/Engine/Plugins/AudioInsights/Source/AudioInsights/Private/Views/SoundDashboardViewFactory.cpp
@@ -523,9 +523,9 @@ namespace UE::Audio::Insights
 		const FText RelativeRenderCostColumnDisplayName = LOCTEXT("SoundDashboard_RelativeRenderCostColumnDisplayName", "Rel. Render Cost");
 		const FText RelativeRenderCostColumnTooltip = LOCTEXT("SoundDashboard_RelativeRenderCostColumnTooltip", "The estimated relative render cost of the sound.");
 
-		const FName ActorColumnName = "Actor";
-		const FText ActorColumnDisplayName = LOCTEXT("SoundDashboard_ActorColumnDisplayName", "Actor");
-		const FText ActorColumnTooltip = LOCTEXT("SoundDashboard_ActorColumnTooltip", "The name of the Actor this sound is playing on (if applicable).");
+		const FName ActorLabelColumnName = "ActorLabel";
+		const FText ActorLabelColumnDisplayName = LOCTEXT("SoundDashboard_ActorLabelColumnDisplayName", "Actor Label");
+		const FText ActorLabelColumnTooltip = LOCTEXT("SoundDashboard_ActorLabelColumnTooltip", "The editor label of the Actor this sound is playing on (if applicable).");
 
 		const FName CategoryColumnName = "Category";
 		const FText CategoryColumnDisplayName = LOCTEXT("SoundDashboard_CategoryColumnDisplayName", "Category");
@@ -1781,11 +1781,11 @@ namespace UE::Audio::Insights
 			},
 #if WITH_EDITOR
 			{
-				ActorColumnName,
+				ActorLabelColumnName,
 				{
-					.DisplayName	  = ActorColumnDisplayName,
+					.DisplayName	  = ActorLabelColumnDisplayName,
 					.IconName		  = NAME_None,
-					.TooltipText	  = ActorColumnTooltip,
+					.TooltipText	  = ActorLabelColumnTooltip,
 					.bShowDisplayName = true,
 					.bDefaultHidden   = !VisibleColumns.bActor,
 					.FillWidth		  = 0.1f,
@@ -2060,13 +2060,13 @@ namespace UE::Audio::Insights
 			},
 #if WITH_EDITOR
 			{
-				ActorColumnName,
+				ActorLabelColumnName,
 				{
 					.GetDisplayValue = [](const IDashboardDataTreeViewEntry& InData)
 					{
 						const FSoundDashboardEntry& SoundDashboardEntry = CastEntry(InData);
 
-						return SoundDashboardEntry.OwnerActorName;
+						return SoundDashboardEntry.ActorLabel;
 					}
 				}
 			},
@@ -2213,7 +2213,7 @@ namespace UE::Audio::Insights
 #if WITH_EDITOR
 		auto SortByActor = [&SortByPlayOrder](const FSoundDashboardEntry& First, const FSoundDashboardEntry& Second)
 		{
-			const int32 Comparison = First.OwnerActorName.CompareToCaseIgnored(Second.OwnerActorName);
+			const int32 Comparison = First.ActorLabel.CompareToCaseIgnored(Second.ActorLabel);
 			if (Comparison == 0)
 			{
 				return SortByPlayOrder(First, Second);
@@ -2279,7 +2279,7 @@ namespace UE::Audio::Insights
 			SortByPredicate(SortByRelativeRenderCost);
 		}
 #if WITH_EDITOR
-		else if (SortByColumn == ActorColumnName)
+		else if (SortByColumn == ActorLabelColumnName)
 		{
 			SortByPredicate(SortByActor);
 		}
diff --git a/Engine/Plugins/AudioInsights/Source/AudioInsights/Private/Views/TableDashboardViewFactory.cpp b/Engine/Plugins/AudioInsights/Source/AudioInsights/Private/Views/TableDashboardViewFactory.cpp
--- a/Engine/Plugins/AudioInsights/Source/AudioInsights/Private/Views/TableDashboardViewFactory.cpp
+++ b/Engine/Plugins/AudioInsights/Source/AudioInsights/Private/Views/TableDashboardViewFactory.cpp
@@ -2,6 +2,7 @@
 #include "Views/TableDashboardViewFactory.h"
 
 #include "Algo/Transform.h"
+#include "AudioInsightsStyle.h"
 #include "AudioInsightsTraceProviderBase.h"
 #include "Containers/Array.h"
 #include "Framework/Commands/UICommandList.h"
@@ -25,7 +26,10 @@ namespace UE::Audio::Insights
 	{
 		Data = InData;
 		Factory = InFactory;
-		SMultiColumnTableRow<TSharedPtr<IDashboardDataViewEntry>>::Construct(FSuperRowType::FArguments(), InOwnerTable);
+
+		FSuperRowType::FArguments Args = InArgs._Style != nullptr ? FSuperRowType::FArguments().Style(InArgs._Style) : FSuperRowType::FArguments();
+
+		SMultiColumnTableRow<TSharedPtr<IDashboardDataViewEntry>>::Construct(Args, InOwnerTable);
 	}
 
 	TSharedRef<SWidget> FTraceTableDashboardViewFactory::SRowWidget::GenerateWidgetForColumn(const FName& Column)
@@ -36,17 +40,36 @@ namespace UE::Audio::Insights
 	TSharedRef<SWidget> FTraceTableDashboardViewFactory::GenerateWidgetForColumn(TSharedRef<IDashboardDataViewEntry> InRowData, const FName& Column)
 	{
 		const FColumnData& ColumnData = GetColumns()[Column];
-		const FText ValueText = ColumnData.GetDisplayValue(InRowData.Get());
+		const FText ValueText = ColumnData.GetDisplayValue.IsSet() ? ColumnData.GetDisplayValue(InRowData.Get()) : FText::GetEmpty();
+		const FName IconName = ColumnData.GetIconName.IsSet() ? ColumnData.GetIconName(InRowData.Get()) : NAME_None;
 
-		if (ValueText.IsEmpty())
+		if (ValueText.IsEmpty() && IconName == NAME_None)
 		{
 			return SNullWidget::NullWidget;
 		}
 
 		return SNew(SHorizontalBox)
-		+ SHorizontalBox::Slot().Padding(2)
+
+		// Icon
+		+ SHorizontalBox::Slot()
+		.Padding(0.0f, 2.0f)
+		.AutoWidth()
 		[
-			SNew(STextBlock)
+			IconName != NAME_None
+
+			? SNew(SImage)
+			.Image(FSlateStyle::Get().GetBrush(IconName))
+
+			: SNullWidget::NullWidget
+		]
+
+		+ SHorizontalBox::Slot()
+		.Padding(IconName != NAME_None ? 10.0f : 0.0f, 2.0f)
+		.FillWidth(0.9f)
+		[
+			ValueText.IsEmpty()
+
+			? SNew(STextBlock)
 			.Text_Lambda([this, InRowData, Column]()
 			{
 				const FColumnData& ColumnData = GetColumns()[Column];
@@ -54,6 +77,8 @@ namespace UE::Audio::Insights
 
 				return ValueText;
 			})
+
+			: SNullWidget::NullWidget
 		];
 	}
 
@@ -79,6 +104,12 @@ namespace UE::Audio::Insights
 
 	void FTraceTableDashboardViewFactory::OnColumnSortModeChanged(const EColumnSortPriority::Type InSortPriority, const FName& InColumnId, const EColumnSortMode::Type InSortMode)
 	{
+		// Sorting can be disabled on specified columns
+		if (!IsColumnSortable(InColumnId))
+		{
+			return;
+		}
+
 		SortByColumn = InColumnId;
 		SortMode = InSortMode;
 		RequestSort();
@@ -142,9 +173,11 @@ namespace UE::Audio::Insights
 			.ListItemsSource(&DataViewEntries)
 			.OnContextMenuOpening(this, &FTraceTableDashboardViewFactory::OnConstructContextMenu)
 			.OnSelectionChanged(this, &FTraceTableDashboardViewFactory::OnSelectionChanged)
+			.ClearSelectionOnClick(ClearSelectionOnClick())
 			.OnGenerateRow_Lambda([this](TSharedPtr<IDashboardDataViewEntry> Item, const TSharedRef<STableViewBase>& OwnerTable)
 			{
-				return SNew(SRowWidget, OwnerTable, Item, AsShared());
+				return SNew(SRowWidget, OwnerTable, Item, AsShared())
+					.Style(GetRowStyle());
 			})
 			.HeaderRow
 			(
@@ -304,15 +337,32 @@ namespace UE::Audio::Insights
 	TSharedRef<SWidget> FTraceObjectTableDashboardViewFactory::GenerateWidgetForColumn(TSharedRef<IDashboardDataViewEntry> InRowData, const FName& Column)
 	{
 		const FColumnData& ColumnData = GetColumns()[Column];
-		const FText ValueText = ColumnData.GetDisplayValue(InRowData.Get());
+		const FText ValueText = ColumnData.GetDisplayValue.IsSet() ? ColumnData.GetDisplayValue(InRowData.Get()) : FText::GetEmpty();
+		const FName IconName = ColumnData.GetIconName.IsSet() ? ColumnData.GetIconName(InRowData.Get()) : NAME_None;
 
-		if (ValueText.IsEmpty())
+		if (ValueText.IsEmpty() && IconName == NAME_None)
 		{
 			return SNullWidget::NullWidget;
 		}
 
 		return SNew(SHorizontalBox)
-		+ SHorizontalBox::Slot().Padding(2)
+
+		// Icon
+		+ SHorizontalBox::Slot()
+		.Padding(0.0f, 2.0f)
+		.AutoWidth()
+		[
+			IconName != NAME_None
+
+			? SNew(SImage)
+			.Image(FSlateStyle::Get().GetBrush(IconName))
+
+			: SNullWidget::NullWidget
+		]
+
+		+ SHorizontalBox::Slot()
+		.Padding(IconName != NAME_None ? 10.0f : 0.0f, 2.0f)
+		.FillWidth(0.9f)
 		[
 			SNew(STextBlock)
 			.Text_Lambda([this, InRowData, Column]()
diff --git a/Engine/Plugins/AudioInsights/Source/AudioInsights/Private/Views/VirtualLoopDashboardViewFactory.cpp b/Engine/Plugins/AudioInsights/Source/AudioInsights/Private/Views/VirtualLoopDashboardViewFactory.cpp
--- a/Engine/Plugins/AudioInsights/Source/AudioInsights/Private/Views/VirtualLoopDashboardViewFactory.cpp
+++ b/Engine/Plugins/AudioInsights/Source/AudioInsights/Private/Views/VirtualLoopDashboardViewFactory.cpp
@@ -86,6 +86,7 @@ namespace UE::Audio::Insights
 					{
 						LOCTEXT("VirtualLoop_PlayOrderColumnDisplayName", "Play Order"),
 						[](const IDashboardDataViewEntry& InData) { return FText::AsNumber(static_cast<const FVirtualLoopDashboardEntry&>(InData).PlayOrder); },
+						nullptr /*GetIconName*/,
 						true /* bDefaultHidden */,
 						0.1f /* FillWidth */
 					}
@@ -95,6 +96,7 @@ namespace UE::Audio::Insights
 					{
 						LOCTEXT("VirtualLoop_NameColumnDisplayName", "Name"),
 						[](const IDashboardDataViewEntry& InData) { return VirtualLoopPrivate::CastEntry(InData).GetDisplayName(); },
+						nullptr /*GetIconName*/,
 						false /* bDefaultHidden */,
 						0.6f  /* FillWidth */
 					}
@@ -104,6 +106,7 @@ namespace UE::Audio::Insights
 					{
 						LOCTEXT("VirtualLoop_VirtualizedTimeColumnDisplayName", "Time (Virtualized)"),
 						[](const IDashboardDataViewEntry& InData) { return FSlateStyle::Get().FormatSecondsAsTime(VirtualLoopPrivate::CastEntry(InData).TimeVirtualized); },
+						nullptr /*GetIconName*/,
 						false /* bDefaultHidden */,
 						0.15f /* FillWidth */
 					}
@@ -113,6 +116,7 @@ namespace UE::Audio::Insights
 					{
 						LOCTEXT("VirtualLoop_TotalTimeColumnDisplayName", "Time (Total)"),
 						[](const IDashboardDataViewEntry& InData) { return FSlateStyle::Get().FormatSecondsAsTime(VirtualLoopPrivate::CastEntry(InData).PlaybackTime); },
+						nullptr /*GetIconName*/,
 						false /* bDefaultHidden */,
 						0.12f /* FillWidth */
 					}
@@ -122,6 +126,7 @@ namespace UE::Audio::Insights
 					{
 						LOCTEXT("VirtualLoop_UpdateIntervalColumnDisplayName", "Update Interval"),
 						[](const IDashboardDataViewEntry& InData) { return FSlateStyle::Get().FormatSecondsAsTime(VirtualLoopPrivate::CastEntry(InData).UpdateInterval); },
+						nullptr /*GetIconName*/,
 						false /* bDefaultHidden */,
 						0.13f /* FillWidth */
 					}
diff --git a/Engine/Plugins/AudioInsights/Source/AudioInsights/Public/Messages/AudioEventLogTraceMessages.h b/Engine/Plugins/AudioInsights/Source/AudioInsights/Public/Messages/AudioEventLogTraceMessages.h
new file mode 100644
--- /dev/null
+++ b/Engine/Plugins/AudioInsights/Source/AudioInsights/Public/Messages/AudioEventLogTraceMessages.h
@@ -0,0 +1,74 @@
+// Copyright Epic Games, Inc. All Rights Reserved.
+#pragma once
+
+#include "AudioDefines.h"
+#include "Messages/AnalyzerMessageQueue.h"
+#include "SoundTraceMessages.h"
+#include "Trace/Analyzer.h"
+#include "Views/TableDashboardViewFactory.h"
+
+#if !WITH_EDITOR
+#include "Common/PagedArray.h"
+#include "TraceServices/Model/AnalysisSession.h"
+#endif // !WITH_EDITOR
+
+namespace UE::Audio::Insights
+{
+	enum class EAudioEventLogSoundCategory : uint8
+	{
+		None,
+		MetaSound,
+		SoundCue,
+		ProceduralSource,
+		SoundWave,
+		SoundCueTemplate
+	};
+
+	struct FAudioEventLogMessage
+	{
+		FAudioEventLogMessage() = default;
+		FAudioEventLogMessage(const Trace::IAnalyzer::FOnEventContext& InContext);
+
+		::Audio::FDeviceId DeviceId = INDEX_NONE;
+		uint32 MessageID = INDEX_NONE;
+		uint32 PlayOrder = INDEX_NONE;
+
+		FString EventName;
+		FString AssetPath;
+		FString ActorName;
+		FString ActorIconName;
+		FString CategoryName;
+
+		EAudioEventLogSoundCategory CategoryType;
+
+		double Timestamp = 0.0;
+	};
+
+	class FAudioEventLogDashboardEntry : public FSoundAssetDashboardEntry
+	{
+	public:
+		FAudioEventLogDashboardEntry() = default;
+		virtual ~FAudioEventLogDashboardEntry() = default;
+
+		virtual bool IsValid() const override
+		{
+			return MessageID != static_cast<uint32>(INDEX_NONE);
+		}
+
+		uint32 MessageID = INDEX_NONE;
+
+		FString EventName;
+		FString ActorName;
+		FName ActorIconName;
+		FString CategoryName;
+
+		EAudioEventLogSoundCategory Category = EAudioEventLogSoundCategory::None;
+	};
+
+	class FAudioEventLogMessages
+	{
+		TAnalyzerMessageQueue<FAudioEventLogMessage> AudioEventLogMessages { 5.0 };
+
+		friend class FAudioEventLogTraceProvider;
+	};
+} // namespace UE::Audio::Insights
diff --git a/Engine/Plugins/AudioInsights/Source/AudioInsights/Public/Messages/SoundTraceMessages.h b/Engine/Plugins/AudioInsights/Source/AudioInsights/Public/Messages/SoundTraceMessages.h
--- a/Engine/Plugins/AudioInsights/Source/AudioInsights/Public/Messages/SoundTraceMessages.h
+++ b/Engine/Plugins/AudioInsights/Source/AudioInsights/Public/Messages/SoundTraceMessages.h
@@ -28,6 +28,15 @@ namespace UE::Audio::Insights
 		Pinned
 	};
 
+	namespace SoundClassNames
+	{
+		const FString MetaSoundSource = TEXT("MetaSoundSource");
+		const FString SoundWaveProcedural = TEXT("SoundWaveProcedural");
+		const FString SoundCue = TEXT("SoundCue");
+		const FString SoundWave = TEXT("SoundWave");
+		const FString SoundCueTemplate = TEXT("SoundCueTemplate");
+	}
+
 	// Trace messages
 	struct FSoundMessageBase
 	{
@@ -61,7 +70,7 @@ namespace UE::Audio::Insights
 		ESoundDashboardEntryType EntryType;
 
 #if WITH_EDITOR
-		FString OwnerActorName;
+		FString ActorLabel;
 #endif // WITH_EDITOR
 	};
 
@@ -253,7 +262,7 @@ namespace UE::Audio::Insights
 		::Audio::TCircularAudioBuffer<FDataPoint> RelativeRenderCostDataPoints;
 
 #if WITH_EDITOR
-		FText OwnerActorName;
+		FText ActorLabel;
 #endif // WITH_EDITOR
 
 	private:
diff --git a/Engine/Plugins/AudioInsights/Source/AudioInsights/Public/Views/AudioEventLogDashboardViewFactory.h b/Engine/Plugins/AudioInsights/Source/AudioInsights/Public/Views/AudioEventLogDashboardViewFactory.h
new file mode 100644
--- /dev/null
+++ b/Engine/Plugins/AudioInsights/Source/AudioInsights/Public/Views/AudioEventLogDashboardViewFactory.h
@@ -0,0 +1,51 @@
+// Copyright Epic Games, Inc. All Rights Reserved.
+#pragma once
+
+#include "Templates/SharedPointer.h"
+#include "Views/TableDashboardViewFactory.h"
+
+#define UE_API AUDIOINSIGHTS_API
+
+namespace UE::Audio::Insights
+{
+	class FAudioEventLogDashboardViewFactory : public FTraceObjectTableDashboardViewFactory
+	{
+	public:
+		UE_API FAudioEventLogDashboardViewFactory();
+		virtual ~FAudioEventLogDashboardViewFactory() = default;
+
+		UE_API virtual FName GetName() const override;
+		UE_API virtual FText GetDisplayName() const override;
+		UE_API virtual EDefaultDashboardTabStack GetDefaultTabStack() const override;
+		UE_API virtual FSlateIcon GetIcon() const override;
+		
+		UE_API virtual TSharedRef<SWidget> MakeWidget(TSharedRef<SDockTab> OwnerTab, const FSpawnTabArgs& SpawnTabArgs) override;
+		UE_API virtual void RefreshFilteredEntriesListView() override;
+
+	protected:
+		UE_API virtual void ProcessEntries(FTraceTableDashboardViewFactory::EProcessReason Reason) override;
+		UE_API virtual const TMap<FName, FTraceTableDashboardViewFactory::FColumnData>& GetColumns() const override;
+
+		UE_API virtual void SortTable() override;
+		UE_API virtual bool IsColumnSortable(const FName& ColumnId) const override;
+
+		UE_API virtual void OnSelectionChanged(TSharedPtr<IDashboardDataViewEntry> SelectedItem, ESelectInfo::Type SelectInfo) override;
+		UE_API virtual bool ClearSelectionOnClick() const override { return true; }
+		UE_API virtual const FTableRowStyle* GetRowStyle() const override;
+
+#if WITH_EDITOR
+		UE_API virtual TSharedRef<SWidget> MakeAssetMenuBar() const override { return SNullWidget::NullWidget; }
+#endif // WITH_EDITOR
+
+	private:
+		void SortByPredicate(TFunctionRef<bool(const class FAudioEventLogDashboardEntry&, const class FAudioEventLogDashboardEntry&)> Predicate);
+
+#if !WITH_EDITOR
+		void OnTimingViewTimeMarkerChanged(double InTimeMarker);
+#endif // !WITH_EDITOR
+
+		TSharedPtr<SWidget> DashboardWidget;
+	};
+} // namespace UE::Audio::Insights
+
+#undef UE_API
\ No newline at end of file
diff --git a/Engine/Plugins/AudioInsights/Source/AudioInsights/Public/Views/TableDashboardViewFactory.h b/Engine/Plugins/AudioInsights/Source/AudioInsights/Public/Views/TableDashboardViewFactory.h
--- a/Engine/Plugins/AudioInsights/Source/AudioInsights/Public/Views/TableDashboardViewFactory.h
+++ b/Engine/Plugins/AudioInsights/Source/AudioInsights/Public/Views/TableDashboardViewFactory.h
@@ -11,7 +11,6 @@
 
 #define UE_API AUDIOINSIGHTS_API
 
-
 namespace UE::Audio::Insights
 {
 	class IObjectDashboardEntry : public IDashboardDataViewEntry
@@ -49,7 +48,10 @@ namespace UE::Audio::Insights
 	protected:
 		struct SRowWidget : public SMultiColumnTableRow<TSharedPtr<IDashboardDataViewEntry>>
 		{
-			SLATE_BEGIN_ARGS(SRowWidget) { }
+			SLATE_BEGIN_ARGS(SRowWidget)
+				: _Style(nullptr)
+				{}
+				SLATE_STYLE_ARGUMENT(FTableRowStyle, Style)
 			SLATE_END_ARGS()
 
 			void Construct(const FArguments& InArgs, const TSharedRef<STableViewBase>& InOwnerTable, TSharedPtr<IDashboardDataViewEntry> InData, TSharedRef<FTraceTableDashboardViewFactory> InFactory);
@@ -122,28 +124,33 @@ namespace UE::Audio::Insights
 		}
 
 #if WITH_EDITOR
-		virtual bool IsDebugDrawEnabled() const { return false; }
-		virtual void DebugDraw(float InElapsed, const TArray<TSharedPtr<IDashboardDataViewEntry>>& InSelectedItems, ::Audio::FDeviceId InAudioDeviceId) const { };
+		UE_API virtual bool IsDebugDrawEnabled() const { return false; }
+		UE_API virtual void DebugDraw(float InElapsed, const TArray<TSharedPtr<IDashboardDataViewEntry>>& InSelectedItems, ::Audio::FDeviceId InAudioDeviceId) const { };
 #endif // WITH_EDITOR
 
 		struct FColumnData
 		{
 			const FText DisplayName;
 			const TFunction<FText(const IDashboardDataViewEntry&)> GetDisplayValue;
+			const TFunction<FName(const IDashboardDataViewEntry&)> GetIconName;
 			bool bDefaultHidden = false;
 			const float FillWidth = 1.0f;
 			const EHorizontalAlignment Alignment = HAlign_Left;
 		};
 
-		virtual const TMap<FName, FColumnData>& GetColumns() const = 0;
-		virtual void ProcessEntries(EProcessReason Reason) = 0;
+		UE_API virtual const TMap<FName, FColumnData>& GetColumns() const = 0;
+		UE_API virtual void ProcessEntries(EProcessReason Reason) = 0;
 
 		UE_API const FText& GetSearchFilterText() const;
 
-		virtual void SortTable() = 0;
+		UE_API virtual void SortTable() = 0;
+		UE_API virtual bool IsColumnSortable(const FName& ColumnId) const { return true; }
 
 		UE_API virtual TSharedPtr<SWidget> OnConstructContextMenu();
 		UE_API virtual void OnSelectionChanged(TSharedPtr<IDashboardDataViewEntry> SelectedItem, ESelectInfo::Type SelectInfo);
+		UE_API virtual bool ClearSelectionOnClick() const { return false; }
+
+		UE_API virtual const FTableRowStyle* GetRowStyle() const { return nullptr; }
 		UE_API virtual FSlateColor GetRowColor(const TSharedPtr<IDashboardDataViewEntry>& InRowDataPtr);
 
 		EProcessReason UpdateFilterReason = EProcessReason::None;
diff --git a/Engine/Plugins/AudioInsights/Source/AudioInsightsEditor/Private/AudioInsightsEditorModule.cpp b/Engine/Plugins/AudioInsights/Source/AudioInsightsEditor/Private/AudioInsightsEditorModule.cpp
--- a/Engine/Plugins/AudioInsights/Source/AudioInsightsEditor/Private/AudioInsightsEditorModule.cpp
+++ b/Engine/Plugins/AudioInsights/Source/AudioInsightsEditor/Private/AudioInsightsEditorModule.cpp
@@ -10,6 +10,7 @@
 #include "UObject/NameTypes.h"
 #include "Views/AudioAnalyzerRackDashboardViewFactory.h"
 #include "Views/AudioBusesDashboardViewFactory.h"
+#include "Views/AudioEventLogDashboardViewFactory.h"
 #include "Views/AudioMetersDashboardViewFactory.h"
 #include "Views/LogDashboardViewFactory.h"
 #include "Views/SoundDashboardViewFactory.h"
@@ -41,6 +42,7 @@ namespace UE::Audio::Insights
 			// @TODO UE-274216: Decide what to do with the Viewport dashboard
 			//DashboardFactory->RegisterViewFactory(MakeShared<FViewportDashboardViewFactory>());
 			DashboardFactory->RegisterViewFactory(MakeShared<FLogDashboardViewFactory>());
+			DashboardFactory->RegisterViewFactory(MakeShared<FAudioEventLogDashboardViewFactory>());
 			DashboardFactory->RegisterViewFactory(SoundsDashboard);
 			DashboardFactory->RegisterViewFactory(MakeShared<FVirtualLoopDashboardViewFactory>());
 			DashboardFactory->RegisterViewFactory(MakeShared<FSubmixesDashboardViewFactory>());
diff --git a/Engine/Plugins/AudioInsights/Source/AudioInsightsEditor/Private/Views/AudioBusesDashboardViewFactory.cpp b/Engine/Plugins/AudioInsights/Source/AudioInsightsEditor/Private/Views/AudioBusesDashboardViewFactory.cpp
--- a/Engine/Plugins/AudioInsights/Source/AudioInsightsEditor/Private/Views/AudioBusesDashboardViewFactory.cpp
+++ b/Engine/Plugins/AudioInsights/Source/AudioInsightsEditor/Private/Views/AudioBusesDashboardViewFactory.cpp
@@ -313,6 +313,7 @@ namespace UE::Audio::Insights
 					{
 						LOCTEXT("AudioBuses_ActiveDisplayName", "Active"),
 						[](const IDashboardDataViewEntry& InData) { return FText::GetEmpty(); },
+						nullptr		/* GetIconName */,
 						false,		/* bDefaultHidden */
 						0.08f,		/* FillWidth */
 						EHorizontalAlignment::HAlign_Center
@@ -327,6 +328,7 @@ namespace UE::Audio::Insights
 							const FAudioBusAssetDashboardEntry& AudioBusAssetDashboardEntry = AudioBusesPrivate::CastEntry(InData);
 							return AudioBusAssetDashboardEntry.EntryType == EAudioBusEntryType::AssetBased ? AudioBusAssetDashboardEntry.GetDisplayName() : FText::FromString(AudioBusAssetDashboardEntry.Name);
 						},
+						nullptr		/* GetIconName */,
 						false,		/* bDefaultHidden */
 						0.92f,		/* FillWidth */
 						EHorizontalAlignment::HAlign_Left
diff --git a/Engine/Plugins/AudioInsights/Source/AudioInsightsEditor/Private/Views/SubmixesDashboardViewFactory.cpp b/Engine/Plugins/AudioInsights/Source/AudioInsightsEditor/Private/Views/SubmixesDashboardViewFactory.cpp
--- a/Engine/Plugins/AudioInsights/Source/AudioInsightsEditor/Private/Views/SubmixesDashboardViewFactory.cpp
+++ b/Engine/Plugins/AudioInsights/Source/AudioInsightsEditor/Private/Views/SubmixesDashboardViewFactory.cpp
@@ -226,6 +226,7 @@ namespace UE::Audio::Insights
 					{
 						LOCTEXT("Submixes_ActiveDisplayName", "Active"),
 						[](const IDashboardDataViewEntry& InData) { return FText::GetEmpty(); },
+						nullptr		/* GetIconName */,
 						false,		/* bDefaultHidden */
 						0.08f,		/* FillWidth */
 						EHorizontalAlignment::HAlign_Center
@@ -236,6 +237,7 @@ namespace UE::Audio::Insights
 					{
 						LOCTEXT("Submixes_NameColumnDisplayName", "Name"),
 						[](const IDashboardDataViewEntry& InData) { return SubmixesPrivate::CastEntry(InData).GetDisplayName(); },
+						nullptr		/* GetIconName */,
 						false,		/* bDefaultHidden */
 						0.92f,		/* FillWidth */
 						EHorizontalAlignment::HAlign_Left
diff --git a/Engine/Plugins/Runtime/AudioModulation/Source/AudioModulationEditor/Private/Insights/Views/ControlBusDashboardViewFactory.cpp b/Engine/Plugins/Runtime/AudioModulation/Source/AudioModulationEditor/Private/Insights/Views/ControlBusDashboardViewFactory.cpp
--- a/Engine/Plugins/Runtime/AudioModulation/Source/AudioModulationEditor/Private/Insights/Views/ControlBusDashboardViewFactory.cpp
+++ b/Engine/Plugins/Runtime/AudioModulation/Source/AudioModulationEditor/Private/Insights/Views/ControlBusDashboardViewFactory.cpp
@@ -129,6 +129,7 @@ namespace AudioModulationEditor
 					{
 						LOCTEXT("ControlBus_BusIdColumnDisplayName", "Control Bus ID"),
 						[](const IDashboardDataViewEntry& InData) { return FText::AsNumber(ControlBusPrivate::CastEntry(InData).ControlBusId); },
+						nullptr	/* GetIconName */,
 						true /* bDefaultHidden */,
 						0.08f /* FillWidth */
 					}
@@ -138,6 +139,7 @@ namespace AudioModulationEditor
 					{
 						LOCTEXT("ControlBus_NameColumnDisplayName", "Name"),
 						[](const IDashboardDataViewEntry& InData) { return ControlBusPrivate::CastEntry(InData).GetDisplayNameAsFText(); },
+						nullptr	/* GetIconName */,
 						false /* bDefaultHidden */,
 						0.75f /* FillWidth */
 					}
@@ -147,6 +149,7 @@ namespace AudioModulationEditor
 					{
 						LOCTEXT("ControlBus_ParamNameColumnDisplayName", "Parameter"),
 						[](const IDashboardDataViewEntry& InData) { return FText::FromString(ControlBusPrivate::CastEntry(InData).ParamName); },
+						nullptr	/* GetIconName */,
 						false /* bDefaultHidden */,
 						0.15f /* FillWidth */
 					}
@@ -156,6 +159,7 @@ namespace AudioModulationEditor
 					{
 						LOCTEXT("ControlBus_ValueColumnDisplayName", "Value"),
 						[](const IDashboardDataViewEntry& InData) { return FText::AsNumber(ControlBusPrivate::CastEntry(InData).Value, FSlateStyle::Get().GetLinearVolumeFloatFormat()); },
+						nullptr	/* GetIconName */,
 						false /* bDefaultHidden */,
 						0.07f /* FillWidth */
 					}
diff --git a/Engine/Plugins/Runtime/AudioModulation/Source/AudioModulationEditor/Private/Insights/Views/ModulationMatrixDashboardViewFactory.cpp b/Engine/Plugins/Runtime/AudioModulation/Source/AudioModulationEditor/Private/Insights/Views/ModulationMatrixDashboardViewFactory.cpp
--- a/Engine/Plugins/Runtime/AudioModulation/Source/AudioModulationEditor/Private/Insights/Views/ModulationMatrixDashboardViewFactory.cpp
+++ b/Engine/Plugins/Runtime/AudioModulation/Source/AudioModulationEditor/Private/Insights/Views/ModulationMatrixDashboardViewFactory.cpp
@@ -104,6 +104,7 @@ namespace AudioModulationEditor
 						const FSourceId SourceId = CastEntry(InData).SourceId;
 						return SourceId != INDEX_NONE ? FText::AsNumber(SourceId) : FText::GetEmpty();
 					},
+					nullptr	/* GetIconName */,
 					true /* bDefaultHidden */,
 					0.1f /* FillWidth */
 				}
@@ -113,6 +114,7 @@ namespace AudioModulationEditor
 				{
 					LOCTEXT("ModulationMatrix_ModulatingSourceDisplayName", "Modulating Source"),
 					[](const IDashboardDataViewEntry& InData) { return FText::FromString(CastEntry(InData).DisplayName); },
+					nullptr	/* GetIconName */,
 					false /* bDefaultHidden */,
 					0.4f /* FillWidth */
 				}
@@ -139,6 +141,7 @@ namespace AudioModulationEditor
 
 						return FText::GetEmpty();
 					},
+					nullptr	/* GetIconName */,
 					false /* bDefaultHidden */,
 					0.125f /* FillWidth */
 				}
@@ -323,6 +326,7 @@ namespace AudioModulationEditor
 
 					return FText::GetEmpty();
 				},
+				nullptr	/* GetIconName */,
 				false /* bDefaultHidden */,
 				0.1f  /* FillWidth */
 			});
diff --git a/Engine/Source/Runtime/Engine/Private/Audio/AudioTraceUtil.cpp b/Engine/Source/Runtime/Engine/Private/Audio/AudioTraceUtil.cpp
new file mode 100644
--- /dev/null
+++ b/Engine/Source/Runtime/Engine/Private/Audio/AudioTraceUtil.cpp
@@ -0,0 +1,164 @@
+// Copyright Epic Games, Inc. All Rights Reserved.
+#include "Audio/AudioTraceUtil.h"
+
+#include "ActiveSound.h"
+#include "AudioDevice.h"
+#include "Styling/SlateIconFinder.h"
+#include "Trace/Trace.h"
+
+#define LOCTEXT_NAMESPACE "AudioTraceUtil"
+
+#if UE_AUDIO_PROFILERTRACE_ENABLED
+
+UE_TRACE_EVENT_BEGIN(Audio, EventLog)
+	UE_TRACE_EVENT_FIELD(uint32, DeviceId)
+	UE_TRACE_EVENT_FIELD(uint64, Timestamp)
+	UE_TRACE_EVENT_FIELD(uint32, PlayOrder)
+	UE_TRACE_EVENT_FIELD(UE::Trace::WideString, AssetPath)
+	UE_TRACE_EVENT_FIELD(UE::Trace::WideString, EventLogName)
+	UE_TRACE_EVENT_FIELD(UE::Trace::WideString, ActorName)
+	UE_TRACE_EVENT_FIELD(UE::Trace::WideString, ActorIconName)
+	UE_TRACE_EVENT_FIELD(UE::Trace::WideString, SoundClassName)
+UE_TRACE_EVENT_END()
+
+namespace Audio::Trace
+{
+	namespace Private
+	{
+		const FString MetaSoundSource = TEXT("MetaSoundSource");
+		const FString SoundWaveProcedural = TEXT("SoundWaveProcedural");
+		const FString SoundCue = TEXT("SoundCue");
+		const FString SoundWave = TEXT("SoundWave");
+		const FString SoundCueTemplate = TEXT("SoundCueTemplate");
+
+		const FString SoundStartLogMessage = LOCTEXT("EventLogTraceMessage_SoundStart", "Start").ToString();
+		const FString SoundStopLogMessage = LOCTEXT("EventLogTraceMessage_SoundStop", "Stop").ToString();
+
+		bool NameMatchesBaseAudioAsset(const FString& SoundClassName)
+		{
+			return SoundClassName == MetaSoundSource
+				|| SoundClassName == SoundWaveProcedural
+				|| SoundClassName == SoundCue
+				|| SoundClassName == SoundWave
+				|| SoundClassName == SoundCueTemplate;
+		}
+	}
+
+	namespace Util
+	{
+		FString GetSoundBaseAssetName(const TObjectPtr<UClass> SoundClass)
+		{
+			if (SoundClass == nullptr)
+			{
+				return FString();
+			}
+
+			TObjectPtr<UClass> SoundSuperClass = SoundClass->GetSuperClass();
+			FString SoundClassName = SoundClass->GetName();
+			while (!Private::NameMatchesBaseAudioAsset(SoundClassName) && !SoundClassName.IsEmpty())
+			{
+				SoundClassName = SoundSuperClass ? SoundSuperClass->GetName() : FString();
+				SoundSuperClass = SoundSuperClass ? SoundSuperClass->GetSuperClass() : nullptr;
+			}
+
+			return SoundClassName;
+		}
+
+		FString GetOwnerActorName(const FActiveSound& InActiveSound)
+		{
+			if (const TObjectPtr<UAudioComponent> AudioComponent = UAudioComponent::GetAudioComponentFromID(InActiveSound.GetAudioComponentID()))
+			{
+				if (const AActor* OwnerActor = AudioComponent->GetOwner())
+				{
+					return OwnerActor->GetName();
+				}
+			}
+
+			return FString();
+		}
+
+		FString GetOwnerActorLabel(const FActiveSound& InActiveSound)
+		{
+#if WITH_EDITOR
+			if (const TObjectPtr<UAudioComponent> AudioComponent = UAudioComponent::GetAudioComponentFromID(InActiveSound.GetAudioComponentID()))
+			{
+				if (const AActor* OwnerActor = AudioComponent->GetOwner())
+				{
+					return OwnerActor->GetActorLabel(false);
+				}
+			}
+#endif // WITH_EDITOR
+
+			return FString();
+		}
+
+		FName GetOwnerActorIconName(const FActiveSound& InActiveSound)
+		{
+#if WITH_EDITOR
+			if (const TObjectPtr<UAudioComponent> AudioComponent = UAudioComponent::GetAudioComponentFromID(InActiveSound.GetAudioComponentID()))
+			{
+				if (const AActor* OwnerActor = AudioComponent->GetOwner())
+				{
+					FName IconName = OwnerActor->GetCustomIconName();
+					if (IconName == NAME_None)
+					{
+						// Actor didn't specify an icon - fallback on the class icon
+						return FSlateIconFinder::FindIconForClass(OwnerActor->GetClass()).GetStyleName();
+					}
+				}
+			}
+#endif // WITH_EDITOR
+
+			return NAME_None;
+		}
+
+		TObjectPtr<UObject> GetSoundObjectPointer(const FActiveSound& InActiveSound)
+		{
+			return InActiveSound.GetSound() != nullptr ? FSoftObjectPath(InActiveSound.GetSound()->GetPathName()).ResolveObject() : nullptr;
+		}
+	} // namespace Util
+
+	namespace EventLog
+	{
+		void SendSoundAssetEventLog(const FActiveSound& InActiveSound, const FString& Event)
+		{
+			const bool bChannelEnabled = UE_TRACE_CHANNELEXPR_IS_ENABLED(AudioMixerChannel);
+
+			if (bChannelEnabled)
+			{
+				const TObjectPtr<UObject> SoundObjectPtr = Util::GetSoundObjectPointer(InActiveSound);
+				if (SoundObjectPtr == nullptr)
+				{
+					return;
+				}
+
+				const FString SoundBaseAssetName = Util::GetSoundBaseAssetName(SoundObjectPtr->GetClass());
+
+				UE_TRACE_LOG(Audio, EventLog, AudioMixerChannel)
+					<< EventLog.DeviceId(InActiveSound.AudioDevice->DeviceID)
+					<< EventLog.Timestamp(FPlatformTime::Cycles64())
+					<< EventLog.PlayOrder(InActiveSound.GetPlayOrder())
+					<< EventLog.AssetPath(*SoundObjectPtr->GetPathName())
+					<< EventLog.EventLogName(*Event)
+					<< EventLog.ActorName(*Util::GetOwnerActorName(InActiveSound))
+					<< EventLog.ActorIconName(*Util::GetOwnerActorIconName(InActiveSound).ToString())
+					<< EventLog.SoundClassName(*SoundBaseAssetName);
+			}
+		}
+
+		const FString& GetStartEventName()
+		{
+			return Private::SoundStartLogMessage;
+		}
+
+		const FString& GetStopEventName()
+		{
+			return Private::SoundStopLogMessage;
+		}
+	} // namespace EventLog
+
+} // namespace Audio::AudioTraceUtil
+
+#endif // UE_AUDIO_PROFILERTRACE_ENABLED
+
+#undef LOCTEXT_NAMESPACE // AudioTraceUtil
\ No newline at end of file
diff --git a/Engine/Source/Runtime/Engine/Private/AudioDevice.cpp b/Engine/Source/Runtime/Engine/Private/AudioDevice.cpp
--- a/Engine/Source/Runtime/Engine/Private/AudioDevice.cpp
+++ b/Engine/Source/Runtime/Engine/Private/AudioDevice.cpp
@@ -7,6 +7,7 @@
 #include "AudioEffect.h"
 #include "AudioPluginUtilities.h"
 #include "Audio/AudioDebug.h"
+#include "Audio/AudioTraceUtil.h"
 #include "GameFramework/GameUserSettings.h"
 #include "GameFramework/WorldSettings.h"
 #include "GeneralProjectSettings.h"
@@ -243,24 +244,32 @@ UE_TRACE_EVENT_BEGIN(Audio, SoundStart)
 	UE_TRACE_EVENT_FIELD(uint32, DeviceId)
 	UE_TRACE_EVENT_FIELD(uint64, Timestamp)
 	UE_TRACE_EVENT_FIELD(uint32, PlayOrder)
-	UE_TRACE_EVENT_FIELD(UE::Trace::WideString, Name)
+	UE_TRACE_EVENT_FIELD(UE::Trace::WideString, AssetPath)
+	UE_TRACE_EVENT_FIELD(UE::Trace::WideString, EventLogName)
 	UE_TRACE_EVENT_FIELD(UE::Trace::WideString, SoundClassName)
-	UE_TRACE_EVENT_FIELD(UE::Trace::WideString, OwnerActorName)
+	UE_TRACE_EVENT_FIELD(UE::Trace::WideString, ActorName)
+	UE_TRACE_EVENT_FIELD(UE::Trace::WideString, ActorLabel)
+	UE_TRACE_EVENT_FIELD(UE::Trace::WideString, ActorIconName)
 UE_TRACE_EVENT_END()
 
 UE_TRACE_EVENT_BEGIN(Audio, SoundIsAlivePing)
 	UE_TRACE_EVENT_FIELD(uint32, DeviceId)
 	UE_TRACE_EVENT_FIELD(uint64, Timestamp)
 	UE_TRACE_EVENT_FIELD(uint32, PlayOrder)
-	UE_TRACE_EVENT_FIELD(UE::Trace::WideString, Name)
+	UE_TRACE_EVENT_FIELD(UE::Trace::WideString, AssetPath)
 	UE_TRACE_EVENT_FIELD(UE::Trace::WideString, SoundClassName)
-	UE_TRACE_EVENT_FIELD(UE::Trace::WideString, OwnerActorName)
+	UE_TRACE_EVENT_FIELD(UE::Trace::WideString, ActorLabel)
 UE_TRACE_EVENT_END()
 
 UE_TRACE_EVENT_BEGIN(Audio, SoundStop)
 	UE_TRACE_EVENT_FIELD(uint32, DeviceId)
 	UE_TRACE_EVENT_FIELD(uint64, Timestamp)
 	UE_TRACE_EVENT_FIELD(uint32, PlayOrder)
+	UE_TRACE_EVENT_FIELD(UE::Trace::WideString, AssetPath)
+	UE_TRACE_EVENT_FIELD(UE::Trace::WideString, EventLogName)
+	UE_TRACE_EVENT_FIELD(UE::Trace::WideString, SoundClassName)
+	UE_TRACE_EVENT_FIELD(UE::Trace::WideString, ActorName)
+	UE_TRACE_EVENT_FIELD(UE::Trace::WideString, ActorIconName)
 UE_TRACE_EVENT_END()
 
 UE_TRACE_EVENT_BEGIN(Audio, VirtualLoopStopOrRealize)
@@ -5076,84 +5085,30 @@ void FAudioDevice::UpdateAudioEngineSubsystems()
 }
 
 #if UE_AUDIO_PROFILERTRACE_ENABLED
-namespace AudioProfilerTrace
-{
-	const FString MetaSoundSource = TEXT("MetaSoundSource");
-	const FString SoundWaveProcedural = TEXT("SoundWaveProcedural");
-	const FString SoundCue = TEXT("SoundCue");
-	const FString SoundWave = TEXT("SoundWave");
-	const FString SoundCueTemplate = TEXT("SoundCueTemplate");
-
-	bool NameMatchesBaseAudioAsset(const FString& SoundClassName)
-	{
-		return SoundClassName == MetaSoundSource
-			|| SoundClassName == SoundWaveProcedural
-			|| SoundClassName == SoundCue
-			|| SoundClassName == SoundWave
-			|| SoundClassName == SoundCueTemplate;
-	}
-
-	FString GetSoundBaseAssetName(const TObjectPtr<UClass> SoundClass)
-	{
-		if (SoundClass == nullptr)
-		{
-			return FString();
-		}
-
-		TObjectPtr<UClass> SoundSuperClass = SoundClass->GetSuperClass();
-		FString SoundClassName = SoundClass->GetName();
-		while (!NameMatchesBaseAudioAsset(SoundClassName) && !SoundClassName.IsEmpty())
-		{
-			SoundClassName = SoundSuperClass ? SoundSuperClass->GetName() : FString();
-			SoundSuperClass = SoundSuperClass ? SoundSuperClass->GetSuperClass() : nullptr;
-		}
-
-		return SoundClassName;
-	}
-
-	FString GetOwnerActorLabel(const FActiveSound& InActiveSound)
-	{
-#if WITH_EDITOR
-		if (const TObjectPtr<UAudioComponent> AudioComponent = UAudioComponent::GetAudioComponentFromID(InActiveSound.GetAudioComponentID()))
-		{
-			if (const AActor* OwnerActor = AudioComponent->GetOwner())
-			{
-				return OwnerActor->GetActorLabel(false);
-			}
-		}
-#endif // WITH_EDITOR
-
-		return FString();
-	}
-
-	TObjectPtr<UObject> GetSoundObjectPointer(const FActiveSound & InActiveSound)
-	{
-		return InActiveSound.GetSound() != nullptr ? FSoftObjectPath(InActiveSound.GetSound()->GetPathName()).ResolveObject() : nullptr;
-	}
-}
-
 void FAudioDevice::SendSoundStartedTraceMessage(const FActiveSound& InActiveSound) const
 {
 	const bool bChannelEnabled = UE_TRACE_CHANNELEXPR_IS_ENABLED(AudioMixerChannel);
 
 	if (bChannelEnabled)
 	{
-		const TObjectPtr<UObject> SoundObjectPtr = AudioProfilerTrace::GetSoundObjectPointer(InActiveSound);
+		const TObjectPtr<UObject> SoundObjectPtr = Audio::Trace::Util::GetSoundObjectPointer(InActiveSound);
 		if (SoundObjectPtr == nullptr)
 		{
 			return;
 		}
 
-		const FString SoundBaseAssetName = AudioProfilerTrace::GetSoundBaseAssetName(SoundObjectPtr->GetClass());
-		const FString OwnerActorLabel = AudioProfilerTrace::GetOwnerActorLabel(InActiveSound);
+		const FString OwnerActorLabel = Audio::Trace::Util::GetOwnerActorLabel(InActiveSound);
 
 		UE_TRACE_LOG(Audio, SoundStart, AudioMixerChannel)
 			<< SoundStart.DeviceId(InActiveSound.AudioDevice->DeviceID)
 			<< SoundStart.Timestamp(FPlatformTime::Cycles64())
 			<< SoundStart.PlayOrder(InActiveSound.PlayOrder)
-			<< SoundStart.Name(*SoundObjectPtr->GetPathName())
-			<< SoundStart.SoundClassName(*SoundBaseAssetName)
-			<< SoundStart.OwnerActorName(*OwnerActorLabel);
+			<< SoundStart.AssetPath(*SoundObjectPtr->GetPathName())
+			<< SoundStart.EventLogName(*Audio::Trace::EventLog::GetStartEventName())
+			<< SoundStart.SoundClassName(*Audio::Trace::Util::GetSoundBaseAssetName(SoundObjectPtr->GetClass()))
+			<< SoundStart.ActorName(*Audio::Trace::Util::GetOwnerActorName(InActiveSound))
+			<< SoundStart.ActorLabel(*Audio::Trace::Util::GetOwnerActorLabel(InActiveSound))
+			<< SoundStart.ActorIconName(*Audio::Trace::Util::GetOwnerActorIconName(InActiveSound).ToString());
 	}
 }
 
@@ -5163,22 +5118,22 @@ void FAudioDevice::SendSoundIsAlivePingTraceMessage(const FActiveSound& InActive
 
 	if (bChannelEnabled)
 	{
-		const TObjectPtr<UObject> SoundObjectPtr = AudioProfilerTrace::GetSoundObjectPointer(InActiveSound);
+		const TObjectPtr<UObject> SoundObjectPtr = Audio::Trace::Util::GetSoundObjectPointer(InActiveSound);
 		if (SoundObjectPtr == nullptr)
 		{
 			return;
 		}
 
-		const FString SoundBaseAssetName = AudioProfilerTrace::GetSoundBaseAssetName(SoundObjectPtr->GetClass());
-		const FString OwnerActorLabel = AudioProfilerTrace::GetOwnerActorLabel(InActiveSound);
+		const FString SoundBaseAssetName = Audio::Trace::Util::GetSoundBaseAssetName(SoundObjectPtr->GetClass());
+		const FString OwnerActorLabel = Audio::Trace::Util::GetOwnerActorLabel(InActiveSound);
 
 		UE_TRACE_LOG(Audio, SoundIsAlivePing, AudioMixerChannel)
 			<< SoundIsAlivePing.DeviceId(InActiveSound.AudioDevice->DeviceID)
 			<< SoundIsAlivePing.Timestamp(FPlatformTime::Cycles64())
 			<< SoundIsAlivePing.PlayOrder(InActiveSound.PlayOrder)
-			<< SoundIsAlivePing.Name(*SoundObjectPtr->GetPathName())
+			<< SoundIsAlivePing.AssetPath(*SoundObjectPtr->GetPathName())
 			<< SoundIsAlivePing.SoundClassName(*SoundBaseAssetName)
-			<< SoundIsAlivePing.OwnerActorName(*OwnerActorLabel);
+			<< SoundIsAlivePing.ActorLabel(*OwnerActorLabel);
 	}
 }
 #endif // UE_AUDIO_PROFILERTRACE_ENABLED
@@ -6065,10 +6020,19 @@ void FAudioDevice::RemoveActiveSound(FActiveSound* ActiveSound)
 	const bool bChannelEnabled = UE_TRACE_CHANNELEXPR_IS_ENABLED(AudioMixerChannel);
 	if (bChannelEnabled)
 	{
-		UE_TRACE_LOG(Audio, SoundStop, AudioMixerChannel)
-			<< SoundStop.DeviceId(ActiveSound->AudioDevice->DeviceID)
-			<< SoundStop.Timestamp(FPlatformTime::Cycles64())
-			<< SoundStop.PlayOrder(ActiveSound->PlayOrder);
+		const TObjectPtr<UObject> SoundObjectPtr = Audio::Trace::Util::GetSoundObjectPointer(*ActiveSound);
+		if (SoundObjectPtr)
+		{
+			UE_TRACE_LOG(Audio, SoundStop, AudioMixerChannel)
+				<< SoundStop.DeviceId(ActiveSound->AudioDevice->DeviceID)
+				<< SoundStop.Timestamp(FPlatformTime::Cycles64())
+				<< SoundStop.PlayOrder(ActiveSound->PlayOrder)
+				<< SoundStop.AssetPath(*SoundObjectPtr->GetPathName())
+				<< SoundStop.EventLogName(*Audio::Trace::EventLog::GetStopEventName())
+				<< SoundStop.SoundClassName(*Audio::Trace::Util::GetSoundBaseAssetName(SoundObjectPtr->GetClass()))
+				<< SoundStop.ActorName(*Audio::Trace::Util::GetOwnerActorName(*ActiveSound))
+				<< SoundStop.ActorIconName(*Audio::Trace::Util::GetOwnerActorIconName(*ActiveSound).ToString());
+		}
 	}
 #endif // UE_AUDIO_PROFILERTRACE_ENABLED
 
diff --git a/Engine/Source/Runtime/Engine/Private/SoundWave.cpp b/Engine/Source/Runtime/Engine/Private/SoundWave.cpp
--- a/Engine/Source/Runtime/Engine/Private/SoundWave.cpp
+++ b/Engine/Source/Runtime/Engine/Private/SoundWave.cpp
@@ -3,6 +3,7 @@
 #include "Sound/SoundWave.h"
 
 #include "Audio/AudioTimingLog.h"
+#include "Audio/AudioTraceUtil.h"
 #include "AudioCompressionSettingsUtils.h"
 #include "AudioDecompress.h"
 #include "AudioDerivedData.h"
@@ -116,9 +117,9 @@ UE_TRACE_EVENT_BEGIN(Audio, SoundWaveStart)
 	UE_TRACE_EVENT_FIELD(uint64, Timestamp)
 	UE_TRACE_EVENT_FIELD(uint32, PlayOrder)
 	UE_TRACE_EVENT_FIELD(uint32, ActiveSoundPlayOrder)
-	UE_TRACE_EVENT_FIELD(UE::Trace::WideString, Name)
+	UE_TRACE_EVENT_FIELD(UE::Trace::WideString, AssetPath)
 	UE_TRACE_EVENT_FIELD(UE::Trace::WideString, SoundClassName)
-	UE_TRACE_EVENT_FIELD(UE::Trace::WideString, OwnerActorName)
+	UE_TRACE_EVENT_FIELD(UE::Trace::WideString, ActorLabel)
 UE_TRACE_EVENT_END()
 
 UE_TRACE_EVENT_BEGIN(Audio, SoundWaveIsAlivePing)
@@ -126,9 +127,9 @@ UE_TRACE_EVENT_BEGIN(Audio, SoundWaveIsAlivePing)
 	UE_TRACE_EVENT_FIELD(uint64, Timestamp)
 	UE_TRACE_EVENT_FIELD(uint32, PlayOrder)
 	UE_TRACE_EVENT_FIELD(uint32, ActiveSoundPlayOrder)
-	UE_TRACE_EVENT_FIELD(UE::Trace::WideString, Name)
+	UE_TRACE_EVENT_FIELD(UE::Trace::WideString, AssetPath)
 	UE_TRACE_EVENT_FIELD(UE::Trace::WideString, SoundClassName)
-	UE_TRACE_EVENT_FIELD(UE::Trace::WideString, OwnerActorName)
+	UE_TRACE_EVENT_FIELD(UE::Trace::WideString, ActorLabel)
 UE_TRACE_EVENT_END()
 
 UE_TRACE_EVENT_BEGIN(Audio, SoundPriority)
@@ -3932,17 +3933,17 @@ void USoundWave::Parse(FAudioDevice* AudioDevice, const UPTRINT NodeWaveInstance
 			const TObjectPtr<UObject> SoundWaveObjectPtr = FSoftObjectPath(WaveInstance->WaveData->GetPathName()).ResolveObject();
 			if (SoundWaveObjectPtr)
 			{
-				const FString SoundWaveBaseAssetName = AudioProfilerTrace::GetSoundBaseAssetName(SoundWaveObjectPtr->GetClass());
-				const FString OwnerActorLabel = AudioProfilerTrace::GetOwnerActorLabel(ActiveSound);
+				const FString SoundWaveBaseAssetName = Audio::Trace::Util::GetSoundBaseAssetName(SoundWaveObjectPtr->GetClass());
+				const FString OwnerActorLabel = Audio::Trace::Util::GetOwnerActorLabel(ActiveSound);
 
 				UE_TRACE_LOG(Audio, SoundWaveStart, AudioMixerChannel)
 					<< SoundWaveStart.DeviceId(AudioDevice->DeviceID)
 					<< SoundWaveStart.Timestamp(FPlatformTime::Cycles64())
 					<< SoundWaveStart.PlayOrder(WaveInstance->GetPlayOrder())
 					<< SoundWaveStart.ActiveSoundPlayOrder(ActiveSound.GetPlayOrder())
-					<< SoundWaveStart.Name(*SoundWaveObjectPtr->GetPathName())
+					<< SoundWaveStart.AssetPath(*SoundWaveObjectPtr->GetPathName())
 					<< SoundWaveStart.SoundClassName(*SoundWaveBaseAssetName)
-					<< SoundWaveStart.OwnerActorName(*OwnerActorLabel);
+					<< SoundWaveStart.ActorLabel(*OwnerActorLabel);
 			}
 		}
 		// Occasionally send a ping to Audio Insights to confirm the sound is still playing.
@@ -3952,17 +3953,17 @@ void USoundWave::Parse(FAudioDevice* AudioDevice, const UPTRINT NodeWaveInstance
 			const TObjectPtr<UObject> SoundWaveObjectPtr = FSoftObjectPath(WaveInstance->WaveData->GetPathName()).ResolveObject();
 			if (SoundWaveObjectPtr)
 			{
-				const FString SoundWaveBaseAssetName = AudioProfilerTrace::GetSoundBaseAssetName(SoundWaveObjectPtr->GetClass());
-				const FString OwnerActorLabel = AudioProfilerTrace::GetOwnerActorLabel(ActiveSound);
+				const FString SoundWaveBaseAssetName = Audio::Trace::Util::GetSoundBaseAssetName(SoundWaveObjectPtr->GetClass());
+				const FString OwnerActorLabel = Audio::Trace::Util::GetOwnerActorLabel(ActiveSound);
 
 				UE_TRACE_LOG(Audio, SoundWaveIsAlivePing, AudioMixerChannel)
 					<< SoundWaveIsAlivePing.DeviceId(AudioDevice->DeviceID)
 					<< SoundWaveIsAlivePing.Timestamp(FPlatformTime::Cycles64())
 					<< SoundWaveIsAlivePing.PlayOrder(WaveInstance->GetPlayOrder())
 					<< SoundWaveIsAlivePing.ActiveSoundPlayOrder(ActiveSound.GetPlayOrder())
-					<< SoundWaveIsAlivePing.Name(*SoundWaveObjectPtr->GetPathName())
+					<< SoundWaveIsAlivePing.AssetPath(*SoundWaveObjectPtr->GetPathName())
 					<< SoundWaveIsAlivePing.SoundClassName(*SoundWaveBaseAssetName)
-					<< SoundWaveIsAlivePing.OwnerActorName(*OwnerActorLabel);
+					<< SoundWaveIsAlivePing.ActorLabel(*OwnerActorLabel);
 			}
 		}
 
diff --git a/Engine/Source/Runtime/Engine/Public/Audio/AudioTraceUtil.h b/Engine/Source/Runtime/Engine/Public/Audio/AudioTraceUtil.h
new file mode 100644
--- /dev/null
+++ b/Engine/Source/Runtime/Engine/Public/Audio/AudioTraceUtil.h
@@ -0,0 +1,31 @@
+// Copyright Epic Games, Inc. All Rights Reserved.
+#pragma once
+
+#include "AudioMixerTrace.h"
+#include "CoreMinimal.h"
+
+#if UE_AUDIO_PROFILERTRACE_ENABLED
+
+struct FActiveSound;
+namespace Audio::Trace
+{
+	namespace Util
+	{
+		ENGINE_API FString GetSoundBaseAssetName(const TObjectPtr<UClass> SoundClass);
+		ENGINE_API FString GetOwnerActorName(const FActiveSound& InActiveSound);
+		ENGINE_API FString GetOwnerActorLabel(const FActiveSound& InActiveSound);
+		ENGINE_API FName GetOwnerActorIconName(const FActiveSound& InActiveSound);
+		ENGINE_API TObjectPtr<UObject> GetSoundObjectPointer(const FActiveSound& InActiveSound);
+	} // namespace Util
+
+	namespace EventLog
+	{
+		ENGINE_API void SendSoundAssetEventLog(const FActiveSound& InActiveSound, const FString& Event);
+
+		ENGINE_API const FString& GetStartEventName();
+		ENGINE_API const FString& GetStopEventName();
+	} // namespace EventLog
+
+} // namespace Audio::Trace
+
+#endif // #if UE_AUDIO_PROFILERTRACE_ENABLED
\ No newline at end of file
diff --git a/Engine/Source/Runtime/Engine/Public/AudioDevice.h b/Engine/Source/Runtime/Engine/Public/AudioDevice.h
--- a/Engine/Source/Runtime/Engine/Public/AudioDevice.h
+++ b/Engine/Source/Runtime/Engine/Public/AudioDevice.h
@@ -149,14 +149,6 @@ namespace ESortedActiveWaveGetType
 	};
 }
 
-#if UE_AUDIO_PROFILERTRACE_ENABLED
-namespace AudioProfilerTrace
-{
-	FString GetSoundBaseAssetName(const TObjectPtr<UClass> SoundClass);
-	FString GetOwnerActorLabel(const FActiveSound& InActiveSound);
-}
-#endif // #if UE_AUDIO_PROFILERTRACE_ENABLED
-
 /**
  * Defines the properties of the listener
  */
