[Audio Insights] Column visibility options from settings cog.

diff --git a/Engine/Plugins/AudioInsights/Source/AudioInsights/Private/Views/SoundDashboardViewFactory.cpp b/Engine/Plugins/AudioInsights/Source/AudioInsights/Private/Views/SoundDashboardViewFactory.cpp
--- a/Engine/Plugins/AudioInsights/Source/AudioInsights/Private/Views/SoundDashboardViewFactory.cpp
+++ b/Engine/Plugins/AudioInsights/Source/AudioInsights/Private/Views/SoundDashboardViewFactory.cpp
@@ -1504,6 +1504,35 @@ namespace UE::Audio::Insights
 		MenuBuilder.AddMenuEntry(Commands.GetAutoExpandNothingCommand());
 	}
 
+	void FSoundDashboardViewFactory::BuildVisibleColumnsMenuContent(FMenuBuilder& MenuBuilder)
+	{
+		for (const SHeaderRow::FColumn& Column : HeaderRowWidget->GetColumns())
+		{
+			MenuBuilder.AddMenuEntry(
+				Column.DefaultText,
+				FText::GetEmpty(),
+				FSlateIcon(),
+				FUIAction(
+					FExecuteAction::CreateSP(this, &FSoundDashboardViewFactory::ToggleColumnVisibility, Column.ColumnId),
+					FCanExecuteAction::CreateLambda([]() { return true; }),
+					FGetActionCheckState::CreateSP(this, &FSoundDashboardViewFactory::GetIsColumnVisibleCheckedState, Column.ColumnId)
+				),
+				NAME_None,
+				EUserInterfaceActionType::ToggleButton);
+		}
+	}
+
+	void FSoundDashboardViewFactory::ToggleColumnVisibility(FName ColumnId)
+	{
+		const bool bIsColumnVisible = HeaderRowWidget->IsColumnVisible(ColumnId);
+		HeaderRowWidget->SetShowGeneratedColumn(ColumnId, !bIsColumnVisible);
+	}
+
+	ECheckBoxState FSoundDashboardViewFactory::GetIsColumnVisibleCheckedState(FName ColumnId) const
+	{
+		return HeaderRowWidget->IsColumnVisible(ColumnId) ? ECheckBoxState::Checked : ECheckBoxState::Unchecked;
+	}
+
 	void FSoundDashboardViewFactory::FilterText()
 	{
 		using namespace FSoundDashboardViewFactoryPrivate;
@@ -2484,6 +2513,13 @@ namespace UE::Audio::Insights
 				false
 			);
 
+			// Visible columns
+			MenuBuilder.AddSubMenu(
+				LOCTEXT("SoundsDashboard_VisibleColumns", "Visible Columns"),
+				LOCTEXT("SoundsDashboard_VisibleColumnsTooltip", "Show/hide columns"),
+				FNewMenuDelegate::CreateRaw(this, &FSoundDashboardViewFactory::BuildVisibleColumnsMenuContent)
+			);
+
 			// Show Stopped Sounds checkbox
 			MenuBuilder.AddMenuEntry(
 				FUIAction(FExecuteAction::CreateLambda([this]()
diff --git a/Engine/Plugins/AudioInsights/Source/AudioInsights/Public/Views/SoundDashboardViewFactory.h b/Engine/Plugins/AudioInsights/Source/AudioInsights/Public/Views/SoundDashboardViewFactory.h
--- a/Engine/Plugins/AudioInsights/Source/AudioInsights/Public/Views/SoundDashboardViewFactory.h
+++ b/Engine/Plugins/AudioInsights/Source/AudioInsights/Public/Views/SoundDashboardViewFactory.h
@@ -159,6 +159,9 @@ namespace UE::Audio::Insights
 		UE_API void ChangeAutoExpandMode(const ESoundDashboardAutoExpandOptions SelectedMode, const bool bUpdateEditorSettings = true);
 		UE_API void RefreshExpansionChecks(const TSharedPtr<IDashboardDataTreeViewEntry>& Entry);
 		UE_API void BuildAutoExpandMenuContent(FMenuBuilder& MenuBuilder);
+		UE_API void BuildVisibleColumnsMenuContent(FMenuBuilder& MenuBuilder);
+		UE_API void ToggleColumnVisibility(FName ColumnId);
+		UE_API ECheckBoxState GetIsColumnVisibleCheckedState(FName ColumnId) const;
 
 		UE_API void FilterText();
 		UE_API void RecursiveSort(TArray<TSharedPtr<IDashboardDataTreeViewEntry>>& OutTree, TFunctionRef<bool(const FSoundDashboardEntry&, const FSoundDashboardEntry&)> Predicate);
