[Audio Insights] TDeviceDataMapTraceProvider: don't reset DeviceDataMap when the trace has been stopped or an audio device has been destroyed. This worked fine when we were tracing only for PIE but now we want to preserve the data when stopping a trace in order to be able to examine/scrub it.

diff --git a/Engine/Plugins/AudioInsights/Source/AudioInsights/Public/AudioInsightsTraceProviderBase.h b/Engine/Plugins/AudioInsights/Source/AudioInsights/Public/AudioInsightsTraceProviderBase.h
--- a/Engine/Plugins/AudioInsights/Source/AudioInsights/Public/AudioInsightsTraceProviderBase.h
+++ b/Engine/Plugins/AudioInsights/Source/AudioInsights/Public/AudioInsightsTraceProviderBase.h
@@ -265,30 +265,40 @@ namespace UE::Audio::Insights
 		{
 			FTraceAuxiliary::OnTraceStarted.AddRaw(this, &TDeviceDataMapTraceProvider::OnTraceStarted);			
 			FTraceAuxiliary::OnTraceStopped.AddRaw(this, &TDeviceDataMapTraceProvider::OnTraceStopped);
-			FAudioDeviceManagerDelegates::OnAudioDeviceDestroyed.AddRaw(this, &TDeviceDataMapTraceProvider::OnAudioDeviceDestroyed);
+			FAudioDeviceManagerDelegates::OnAudioDeviceCreated.AddRaw(this, &TDeviceDataMapTraceProvider::OnAudioDeviceCreated);
 		}
 
 		void UnregisterDelegates()
 		{
-			FAudioDeviceManagerDelegates::OnAudioDeviceDestroyed.RemoveAll(this);
+			FAudioDeviceManagerDelegates::OnAudioDeviceCreated.RemoveAll(this);
 			FTraceAuxiliary::OnTraceStopped.RemoveAll(this);
 			FTraceAuxiliary::OnTraceStarted.RemoveAll(this);
 		}
 
 		void OnTraceStarted(FTraceAuxiliary::EConnectionType InTraceType, const FString& InTraceDestination)
 		{
+			Reset();
 			bIsTraceActive = true;
 		}
 
 		void OnTraceStopped(FTraceAuxiliary::EConnectionType InTraceType, const FString& InTraceDestination)
 		{
 			bIsTraceActive = false;
-			Reset();
 		}
 
-		void OnAudioDeviceDestroyed(::Audio::FDeviceId InDeviceId)
+		void OnAudioDeviceCreated(::Audio::FDeviceId InDeviceId)
 		{
-			DeviceDataMap.Remove(InDeviceId);
+			// Check if old audio device data needs to be removed from DeviceDataMap
+			if (const FAudioDeviceManager* DeviceManager = FAudioDeviceManager::Get())
+			{
+				for (auto It = DeviceDataMap.CreateIterator(); It; ++It)
+				{
+					if (!DeviceManager->IsValidAudioDevice(It.Key()))
+					{
+						It.RemoveCurrent();
+					}
+				}
+			}
 		}
 #else
 		const uint32 GetAudioDeviceIdFromDeviceDataMap() const
