[Audio Insights] Allow users to filter by categories in the event log

diff --git a/Engine/Plugins/AudioInsights/Source/AudioInsights/Private/Views/AudioEventLogDashboardViewFactory.cpp b/Engine/Plugins/AudioInsights/Source/AudioInsights/Private/Views/AudioEventLogDashboardViewFactory.cpp
--- a/Engine/Plugins/AudioInsights/Source/AudioInsights/Private/Views/AudioEventLogDashboardViewFactory.cpp
+++ b/Engine/Plugins/AudioInsights/Source/AudioInsights/Private/Views/AudioEventLogDashboardViewFactory.cpp
@@ -3,10 +3,11 @@
 
 #include "AudioInsightsModule.h"
 #include "AudioInsightsStyle.h"
+#include "AudioMixerTrace.h"
 #include "Messages/AudioEventLogTraceMessages.h"
 #include "Providers/AudioEventLogTraceProvider.h"
 
-#define LOCTEXT_NAMESPACE "AudioInsights"
+#define LOCTEXT_NAMESPACE "AudioInsightsEventLog"
 
 namespace UE::Audio::Insights
 {
@@ -19,6 +20,76 @@ namespace UE::Audio::Insights
 			return static_cast<const FAudioEventLogDashboardEntry&>(InData);
 		};
 
+		const FString CategorySoundActivity = LOCTEXT("EventLogCategory_SoundActivity", "Sound Activity").ToString();
+		const FString CategoryVirtualization = LOCTEXT("EventLogCategory_Virtualization", "Virtualization").ToString();
+		const FString CategoryPlayRequests = LOCTEXT("EventLogCategory_PlayRequests", "Play Requests").ToString();
+		const FString CategoryStopRequests = LOCTEXT("EventLogCategory_StopRequests", "Stop Requests").ToString();
+		const FString CategoryMessages = LOCTEXT("EventLogCategory_Messages", "Messages").ToString();
+		const FString CategoryCustom = LOCTEXT("EventLogCategory_Custom", "Custom").ToString();
+
+		const TMap<FString, TArray<FString>> GetInitEventTypeFilters()
+		{
+			using namespace ::Audio::Trace::EventLog;
+			return 
+			{
+				{ 
+					CategorySoundActivity,
+					{
+						ID::SoundStart,
+						ID::SoundStop,
+						ID::PauseSoundRequested,
+						ID::ResumeSoundRequested
+					}
+				},
+				
+				{
+					CategoryVirtualization,
+					{
+						ID::SoundVirtualized,
+						ID::SoundRealized
+					}
+				},
+				
+				{
+					CategoryPlayRequests,
+					{
+						ID::PlayRequestSoundHandle,
+						ID::PlayRequestAudioComponent,
+						ID::PlayRequestOneShot,
+						ID::PlayRequestSoundAtLocation,
+						ID::PlayRequestSound2D,
+						ID::PlayRequestSlateSound
+					}
+				},
+
+				{
+					CategoryStopRequests,
+					{
+						ID::StopAllRequested,
+						ID::StopRequestedSoundHandle,
+						ID::StopRequestAudioComponent,
+						ID::StopRequestActiveSound,
+						ID::StopRequestSoundsUsingResource,
+						ID::StopRequestConcurrency
+					}
+				},
+
+				{
+					CategoryMessages,
+					{
+						ID::FlushAudioDeviceRequested
+					}
+				},
+
+				{
+					CategoryCustom,
+					{
+						// Auto-populates with custom events from users which haven't been categorized 
+					}
+				}
+			};
+		}
+
 		const FName MessageIDColumnID = "MessageID";
 		const FName TimestampColumnID = "Timestamp";
 		const FName PlayOrderColumnID = "PlayOrder";
@@ -112,12 +183,22 @@ namespace UE::Audio::Insights
 	void FAudioEventLogDashboardViewFactory::ProcessEntries(FTraceTableDashboardViewFactory::EProcessReason Reason)
 	{
 		const FString FilterString = GetSearchFilterText().ToString();
-		FTraceTableDashboardViewFactory::FilterEntries<FAudioEventLogTraceProvider>([&FilterString](const IDashboardDataViewEntry& Entry)
+		FTraceTableDashboardViewFactory::FilterEntries<FAudioEventLogTraceProvider>([this, &FilterString](const IDashboardDataViewEntry& Entry)
 		{
+			UpdateCustomEventLogFilters(Entry);
+
 			const FAudioEventLogDashboardEntry& EventLogEntry = static_cast<const FAudioEventLogDashboardEntry&>(Entry);
-			if (EventLogEntry.GetDisplayName().ToString().Contains(FilterString) || EventLogEntry.ActorName.Contains(FilterString))
+
+			const FEventLogFilterID* FilterID = EventFilterTypes.Find(EventLogEntry.EventName);
+
+			const bool bPassesEventTypeFilter = FilterID ? PassesEventTypeFilter(*FilterID) : true;
+
+			if (bPassesEventTypeFilter)
 			{
-				return false;
+				if (EventLogEntry.GetDisplayName().ToString().Contains(FilterString) || EventLogEntry.ActorName.Contains(FilterString))
+				{
+					return false;
+				}
 			}
 
 			return true;
@@ -274,6 +355,48 @@ namespace UE::Audio::Insights
 		return ColumnId == TimestampColumnID;
 	}
 
+	TSharedPtr<SWidget> FAudioEventLogDashboardViewFactory::CreateFilterBarButtonWidget()
+	{
+		if (!FilterBarButton.IsValid())
+		{
+			CreateFilterBarWidget();
+			FilterBarButton = SBasicFilterBar<FEventLogFilterID>::MakeAddFilterButton(StaticCastSharedPtr<SAudioFilterBar<FEventLogFilterID>>(FilterBar).ToSharedRef()).ToSharedPtr();
+		}
+
+		return FilterBarButton;
+	}
+
+	TSharedRef<SWidget> FAudioEventLogDashboardViewFactory::CreateFilterBarWidget()
+	{
+		using namespace AudioEventLogPrivate;
+
+		if (!FilterBar.IsValid())
+		{
+			const TMap<FString, TArray<FString>> InitEventTypeFilters = GetInitEventTypeFilters();
+
+			TArray<TSharedRef<FFilterBase<FEventLogFilterID>>> Filters;
+
+			for (const auto& [Category, EventTypes] : InitEventTypeFilters)
+			{
+				const TSharedPtr<FFilterCategory>& FilterCategory = FilterCategories.FindOrAdd(Category, MakeShared<FFilterCategory>(FText::FromString(Category), FText::GetEmpty()));
+				for (const FString& EventType : EventTypes)
+				{
+					Filters.Add(CreateNewEventFilterType(FilterCategory, EventType));
+				}
+			}
+			
+			SAssignNew(FilterBar, SAudioFilterBar<FEventLogFilterID>)
+			.CustomFilters(Filters)
+			.UseSectionsForCategories(false)
+			.OnFilterChanged_Lambda([this, Filters]()
+			{
+				UpdateFilterReason = EProcessReason::FilterUpdated;
+			});
+		}
+
+		return FilterBar.ToSharedRef();
+	}
+
 	void FAudioEventLogDashboardViewFactory::OnSelectionChanged(TSharedPtr<IDashboardDataViewEntry> SelectedItem, ESelectInfo::Type SelectInfo)
 	{
 		using namespace AudioEventLogPrivate;
@@ -323,6 +446,61 @@ namespace UE::Audio::Insights
 		return &FSlateStyle::Get().GetWidgetStyle<FTableRowStyle>("TreeDashboard.TableViewRow");
 	}
 
+	void FAudioEventLogDashboardViewFactory::UpdateCustomEventLogFilters(const IDashboardDataViewEntry& Entry)
+	{
+		using namespace AudioEventLogPrivate;
+
+		const FAudioEventLogDashboardEntry& EventLogEntry = CastEntry(Entry);
+
+		if (!EventFilterTypes.Contains(EventLogEntry.EventName))
+		{
+			CreateFilterBarWidget();
+
+			const TSharedPtr<FFilterCategory>& CustomCategory = FilterCategories.FindOrAdd(CategoryCustom, MakeShared<FFilterCategory>(FText::FromString(CategoryCustom), FText::GetEmpty()));
+			FilterBar->AddFilter(CreateNewEventFilterType(CustomCategory, EventLogEntry.EventName));
+		}
+	}
+
+	TSharedRef<FAudioEventLogFilter> FAudioEventLogDashboardViewFactory::CreateNewEventFilterType(const TSharedPtr<FFilterCategory>& FilterCategory, const FString& EventType)
+	{
+		static FEventLogFilterID FilterID = 0;
+
+		TSharedPtr<FAudioEventLogFilter> Filter = MakeShared<FAudioEventLogFilter>(FilterID,
+																				   EventType,
+																				   FText::FromString(EventType),
+																				   NAME_None,
+																				   FText::GetEmpty(),
+																				   FLinearColor::MakeRandomColor(),
+																				   FilterCategory);
+		EventFilterTypes.Add(EventType, FilterID++);
+		return Filter.ToSharedRef();
+	}
+
+	bool FAudioEventLogDashboardViewFactory::PassesEventTypeFilter(const FEventLogFilterID EventLogID) const
+	{
+		if (!FilterBar.IsValid())
+		{
+			return true;
+		}
+
+		TSharedPtr<TFilterCollection<FEventLogFilterID>> ActiveFilters = FilterBar->GetAllActiveFilters();
+		
+		if (!ActiveFilters.IsValid() || ActiveFilters->Num() == 0)
+		{
+			return true;
+		}
+		
+		for (int32 FilterIndex = 0; FilterIndex < ActiveFilters->Num(); ++FilterIndex)
+		{
+			if (ActiveFilters->GetFilterAtIndex(FilterIndex)->PassesFilter(EventLogID))
+			{
+				return true;
+			}
+		}
+
+		return false;
+	}
+
 	void FAudioEventLogDashboardViewFactory::SortByPredicate(TFunctionRef<bool(const FAudioEventLogDashboardEntry&, const FAudioEventLogDashboardEntry&)> Predicate)
 	{
 		using namespace AudioEventLogPrivate;
diff --git a/Engine/Plugins/AudioInsights/Source/AudioInsights/Private/Views/TableDashboardViewFactory.cpp b/Engine/Plugins/AudioInsights/Source/AudioInsights/Private/Views/TableDashboardViewFactory.cpp
--- a/Engine/Plugins/AudioInsights/Source/AudioInsights/Private/Views/TableDashboardViewFactory.cpp
+++ b/Engine/Plugins/AudioInsights/Source/AudioInsights/Private/Views/TableDashboardViewFactory.cpp
@@ -169,6 +169,8 @@ namespace UE::Audio::Insights
 	{
 		if (!DashboardWidget.IsValid())
 		{
+			FilterBarButton = CreateFilterBarButtonWidget();
+
 			FilteredEntriesListView = SNew(SListView<TSharedPtr<IDashboardDataViewEntry>>)
 			.ListItemsSource(&DataViewEntries)
 			.OnContextMenuOpening(this, &FTraceTableDashboardViewFactory::OnConstructContextMenu)
@@ -196,6 +198,22 @@ namespace UE::Audio::Insights
 			.Padding(0, 2)
 			[
 				SNew(SHorizontalBox)
+
+				// Optional : Filters menu
+				+ SHorizontalBox::Slot()
+				.AutoWidth()
+				.Padding(FilterBarButton.IsValid() ? 3.0f : 0.0f, 0.0f, 0.0f, 0.0f)
+				[
+					SNew(SVerticalBox)
+						
+					+ SVerticalBox::Slot()
+					.MaxHeight(30.0f)
+					[
+						FilterBarButton.IsValid() ? FilterBarButton.ToSharedRef() : SNullWidget::NullWidget
+					]
+				]
+
+				// Search box
 				+ SHorizontalBox::Slot()
 				.AutoWidth()
 				[
@@ -205,6 +223,12 @@ namespace UE::Audio::Insights
 					.MinDesiredWidth(100)
 					.OnTextChanged(this, &FTraceTableDashboardViewFactory::SetSearchBoxFilterText)
 				]
+
+				// Optional : Active filters area
+				+ SHorizontalBox::Slot()
+				[
+					CreateFilterBarWidget()
+				]
 			]
 			+ SVerticalBox::Slot()
 			.FillHeight(1.0f)
diff --git a/Engine/Plugins/AudioInsights/Source/AudioInsights/Public/Views/AudioEventLogDashboardViewFactory.h b/Engine/Plugins/AudioInsights/Source/AudioInsights/Public/Views/AudioEventLogDashboardViewFactory.h
--- a/Engine/Plugins/AudioInsights/Source/AudioInsights/Public/Views/AudioEventLogDashboardViewFactory.h
+++ b/Engine/Plugins/AudioInsights/Source/AudioInsights/Public/Views/AudioEventLogDashboardViewFactory.h
@@ -1,13 +1,49 @@
 // Copyright Epic Games, Inc. All Rights Reserved.
 #pragma once
 
+#include "Filters/GenericFilter.h"
 #include "Templates/SharedPointer.h"
+#include "Views/SAudioFilterBar.h"
 #include "Views/TableDashboardViewFactory.h"
 
 #define UE_API AUDIOINSIGHTS_API
 
 namespace UE::Audio::Insights
 {
+	using FEventLogFilterID = int32;
+
+	class FAudioEventLogFilter : public FGenericFilter<FEventLogFilterID>
+	{
+	public:
+		FAudioEventLogFilter
+		(
+			const FEventLogFilterID InFilterID,
+			const FString& InName,
+			const FText& InDisplayName,
+			const FName& InIconName,
+			const FText& InToolTipText,
+			FLinearColor InColor,
+			TSharedPtr<FFilterCategory> InCategory
+		)
+		: FGenericFilter<FEventLogFilterID>(InCategory, InName, InDisplayName, FGenericFilter<FEventLogFilterID>::FOnItemFiltered())
+		, FilterID(InFilterID)
+		{
+			ToolTip = InToolTipText;
+			Color = InColor;
+			IconName = InIconName;
+		}
+
+		bool IsActive() const { return bIsActive; }
+		FEventLogFilterID GetFilterID() const { return FilterID; }
+
+	private:
+		virtual void ActiveStateChanged(bool bActive) override { bIsActive = bActive; }
+		virtual bool PassesFilter(FEventLogFilterID InItem) const override { return FilterID == InItem; }
+
+		const FEventLogFilterID FilterID;
+		bool bIsActive = false;
+	};
+
 	class FAudioEventLogDashboardViewFactory : public FTraceObjectTableDashboardViewFactory
 	{
 	public:
@@ -29,6 +65,9 @@ namespace UE::Audio::Insights
 		UE_API virtual void SortTable() override;
 		UE_API virtual bool IsColumnSortable(const FName& ColumnId) const override;
 
+		UE_API virtual TSharedPtr<SWidget> CreateFilterBarButtonWidget() override;
+		UE_API virtual TSharedRef<SWidget> CreateFilterBarWidget() override;
+
 		UE_API virtual void OnSelectionChanged(TSharedPtr<IDashboardDataViewEntry> SelectedItem, ESelectInfo::Type SelectInfo) override;
 		UE_API virtual bool ClearSelectionOnClick() const override { return true; }
 		UE_API virtual const FTableRowStyle* GetRowStyle() const override;
@@ -38,6 +77,10 @@ namespace UE::Audio::Insights
 #endif // WITH_EDITOR
 
 	private:
+		void UpdateCustomEventLogFilters(const IDashboardDataViewEntry& Entry);
+		TSharedRef<FAudioEventLogFilter> CreateNewEventFilterType(const TSharedPtr<FFilterCategory>& FilterCategory, const FString& EventType);
+		bool PassesEventTypeFilter(const FEventLogFilterID EventLogID) const;
+
 		void SortByPredicate(TFunctionRef<bool(const class FAudioEventLogDashboardEntry&, const class FAudioEventLogDashboardEntry&)> Predicate);
 
 #if !WITH_EDITOR
@@ -45,6 +88,10 @@ namespace UE::Audio::Insights
 #endif // !WITH_EDITOR
 
 		TSharedPtr<SWidget> DashboardWidget;
+		TSharedPtr<SAudioFilterBar<FEventLogFilterID>> FilterBar;
+
+		TMap<FString, TSharedPtr<FFilterCategory>> FilterCategories;
+		TMap<FString, FEventLogFilterID> EventFilterTypes;
 	};
 } // namespace UE::Audio::Insights
 
diff --git a/Engine/Plugins/AudioInsights/Source/AudioInsights/Public/Views/SAudioFilterBar.h b/Engine/Plugins/AudioInsights/Source/AudioInsights/Public/Views/SAudioFilterBar.h
--- a/Engine/Plugins/AudioInsights/Source/AudioInsights/Public/Views/SAudioFilterBar.h
+++ b/Engine/Plugins/AudioInsights/Source/AudioInsights/Public/Views/SAudioFilterBar.h
@@ -13,8 +13,11 @@ namespace UE::Audio::Insights
 	public:
 		using FOnFilterChanged = typename SBasicFilterBar<EnumType>::FOnFilterChanged;
 
-		SLATE_BEGIN_ARGS(SAudioFilterBar){}
+		SLATE_BEGIN_ARGS(SAudioFilterBar)
+			: _UseSectionsForCategories(true)
+		{}
 			SLATE_ARGUMENT(TArray<TSharedRef<FFilterBase<EnumType>>>, CustomFilters)
+			SLATE_ARGUMENT(bool, UseSectionsForCategories)
 			SLATE_EVENT(SAudioFilterBar<EnumType>::FOnFilterChanged, OnFilterChanged)
 		SLATE_END_ARGS()
 
@@ -27,7 +30,7 @@ namespace UE::Audio::Insights
 			SBasicFilterBar<EnumType>::Construct(Args
 				.CanChangeOrientation(false)
 				.FilterBarLayout(EFilterBarLayout::Horizontal)
-				.UseSectionsForCategories(true)
+				.UseSectionsForCategories(InArgs._UseSectionsForCategories)
 			);
 		}
 
@@ -78,38 +81,41 @@ namespace UE::Audio::Insights
 
 		void PopulateCustomFilters(UToolMenu* InMenu)
 		{
-			if (!SBasicFilterBar<EnumType>::bUseSectionsForCategories)
+			if (SBasicFilterBar<EnumType>::bUseSectionsForCategories)
 			{
-				return;
-			}
-
-			for (const TSharedPtr<FFilterCategory>& Category : SBasicFilterBar<EnumType>::AllFilterCategories)
-			{
-				if (!Category.IsValid())
+				for (const TSharedPtr<FFilterCategory>& Category : SBasicFilterBar<EnumType>::AllFilterCategories)
 				{
-					continue;
-				}
+					if (!Category.IsValid())
+					{
+						continue;
+					}
 
-				FToolMenuSection& Section = InMenu->AddSection(*Category->Title.ToString(), Category->Title);
+					FToolMenuSection& Section = InMenu->AddSection(*Category->Title.ToString(), Category->Title);
 
-				for (const TSharedRef<FFilterBase<EnumType>>& FrontendFilter : SBasicFilterBar<EnumType>::AllFrontendFilters)
-				{
-					if (FrontendFilter->GetCategory() == Category)
+					for (const TSharedRef<FFilterBase<EnumType>>& FrontendFilter : SBasicFilterBar<EnumType>::AllFrontendFilters)
 					{
-						Section.AddMenuEntry(
-							NAME_None,
-							FrontendFilter->GetDisplayName(),
-							FrontendFilter->GetToolTipText(),
-							FSlateStyle::Get().CreateIcon(FrontendFilter->GetIconName()),
-							FUIAction(
-								FExecuteAction::CreateSP(const_cast<SAudioFilterBar*>(this), &SAudioFilterBar::FrontendFilterClicked, FrontendFilter),
-								FCanExecuteAction(),
-								FIsActionChecked::CreateSP(this, &SAudioFilterBar::IsFrontendFilterInUse, FrontendFilter)),
-							EUserInterfaceActionType::ToggleButton
-						);
+						if (FrontendFilter->GetCategory() == Category)
+						{
+							Section.AddMenuEntry(
+								NAME_None,
+								FrontendFilter->GetDisplayName(),
+								FrontendFilter->GetToolTipText(),
+								FSlateStyle::Get().CreateIcon(FrontendFilter->GetIconName()),
+								FUIAction(
+									FExecuteAction::CreateSP(const_cast<SAudioFilterBar*>(this), &SAudioFilterBar::FrontendFilterClicked, FrontendFilter),
+									FCanExecuteAction(),
+									FIsActionChecked::CreateSP(this, &SAudioFilterBar::IsFrontendFilterInUse, FrontendFilter)),
+								EUserInterfaceActionType::ToggleButton
+							);
+						}
 					}
 				}
 			}
+			else
+			{
+				SBasicFilterBar<EnumType>::PopulateCustomFilters(InMenu);
+			}
+			
 		}
 	};
 } // namespace UE::Audio::Insights
diff --git a/Engine/Plugins/AudioInsights/Source/AudioInsights/Public/Views/TableDashboardViewFactory.h b/Engine/Plugins/AudioInsights/Source/AudioInsights/Public/Views/TableDashboardViewFactory.h
--- a/Engine/Plugins/AudioInsights/Source/AudioInsights/Public/Views/TableDashboardViewFactory.h
+++ b/Engine/Plugins/AudioInsights/Source/AudioInsights/Public/Views/TableDashboardViewFactory.h
@@ -92,7 +92,7 @@ namespace UE::Audio::Insights
 						return StaticCastSharedPtr<IDashboardDataViewEntry>(Pair.Value);
 					};
 
-					if (SearchBoxFilterText.IsEmpty())
+					if (SearchBoxFilterText.IsEmpty() && !FilterBarButton.IsValid())
 					{
 						Algo::Transform(*DeviceData, DataViewEntries, TransformEntry);
 					}
@@ -146,6 +146,10 @@ namespace UE::Audio::Insights
 		UE_API virtual void SortTable() = 0;
 		UE_API virtual bool IsColumnSortable(const FName& ColumnId) const { return true; }
 
+
+		UE_API virtual TSharedPtr<SWidget> CreateFilterBarButtonWidget() { return nullptr; }
+		UE_API virtual TSharedRef<SWidget> CreateFilterBarWidget() { return SNullWidget::NullWidget; }
+
 		UE_API virtual TSharedPtr<SWidget> OnConstructContextMenu();
 		UE_API virtual void OnSelectionChanged(TSharedPtr<IDashboardDataViewEntry> SelectedItem, ESelectInfo::Type SelectInfo);
 		UE_API virtual bool ClearSelectionOnClick() const { return false; }
@@ -164,6 +168,7 @@ namespace UE::Audio::Insights
 		TSharedPtr<SWidget> DashboardWidget;
 		TSharedPtr<SHeaderRow> HeaderRowWidget;
 		TSharedPtr<SListView<TSharedPtr<IDashboardDataViewEntry>>> FilteredEntriesListView;
+		TSharedPtr<SWidget> FilterBarButton;
 
 		FName SortByColumn;
 		EColumnSortMode::Type SortMode = EColumnSortMode::None;
