[Audio Insights] FTraceModule::ExecuteConsoleCommand: unified session instances finding and command execution for editor and standalone (fix for session not found when starting monitoring in standalone)

diff --git a/Engine/Plugins/AudioInsights/Source/AudioInsights/Private/AudioInsightsTraceModule.cpp b/Engine/Plugins/AudioInsights/Source/AudioInsights/Private/AudioInsightsTraceModule.cpp
--- a/Engine/Plugins/AudioInsights/Source/AudioInsights/Private/AudioInsightsTraceModule.cpp
+++ b/Engine/Plugins/AudioInsights/Source/AudioInsights/Private/AudioInsightsTraceModule.cpp
@@ -482,21 +482,22 @@ namespace UE::Audio::Insights
 
 		if (SessionManager.IsValid())
 		{
-#if WITH_EDITOR
-			for (const TSharedPtr<ISessionInstanceInfo>& SessionInstance : SessionManager->GetSelectedInstances())
+			TArray<TSharedPtr<ISessionInfo>> Sessions;
+			SessionManager->GetSessions(Sessions);
+
+			for (const TSharedPtr<ISessionInfo>& Session : Sessions)
 			{
-				if (SessionInstance.IsValid())
+				if (Session.IsValid())
 				{
-					SessionInstance->ExecuteCommand(InCommandStr);
+					TArray<TSharedPtr<ISessionInstanceInfo>> SessionInstances;
+					Session->GetInstances(SessionInstances);
+
+					for (const TSharedPtr<ISessionInstanceInfo>& SessionInstance : SessionInstances)
+					{
+						SessionInstance->ExecuteCommand(InCommandStr);
+					}
 				}
 			}
-#else
-			TSharedPtr<ISessionInstanceInfo> SessionInstance = SessionManager->GetInstance(InstanceID);
-			if (SessionInstance.IsValid())
-			{
-				SessionInstance->ExecuteCommand(InCommandStr);
-			}
-#endif // WITH_EDITOR
 		}
 	}
 
