[Audio Insights] TDeviceDataMapTraceProvider::OnTraceStarted: no need to call Reset here since Provider->Reset(); is called in FTraceProviderBase::FTraceAnalyzerBase::OnAnalysisBegin

diff --git a/Engine/Plugins/AudioInsights/Source/AudioInsights/Public/AudioInsightsTraceProviderBase.h b/Engine/Plugins/AudioInsights/Source/AudioInsights/Public/AudioInsightsTraceProviderBase.h
--- a/Engine/Plugins/AudioInsights/Source/AudioInsights/Public/AudioInsightsTraceProviderBase.h
+++ b/Engine/Plugins/AudioInsights/Source/AudioInsights/Public/AudioInsightsTraceProviderBase.h
@@ -277,7 +277,6 @@ namespace UE::Audio::Insights
 
 		void OnTraceStarted(FTraceAuxiliary::EConnectionType InTraceType, const FString& InTraceDestination)
 		{
-			Reset();
 			bIsTraceActive = true;
 		}
 
