[Audio Insights] Added standalone Submixes dashboard

diff --git a/Engine/Plugins/AudioInsights/Source/AudioInsights/Private/AudioInsightsModule.cpp b/Engine/Plugins/AudioInsights/Source/AudioInsights/Private/AudioInsightsModule.cpp
--- a/Engine/Plugins/AudioInsights/Source/AudioInsights/Private/AudioInsightsModule.cpp
+++ b/Engine/Plugins/AudioInsights/Source/AudioInsights/Private/AudioInsightsModule.cpp
@@ -18,6 +18,7 @@
 #include "Views/AudioEventLogDashboardViewFactory.h"
 #include "Views/SoundDashboardViewFactory.h"
 #include "Views/SoundPlotsDashboardViewFactory.h"
+#include "Views/SubmixDashboardViewFactory.h"
 #include "Views/VirtualLoopDashboardViewFactory.h"
 #endif // !WITH_EDITOR
 
@@ -57,6 +58,7 @@ namespace UE::Audio::Insights
 			DashboardFactory->RegisterViewFactory(SoundsDashboard);
 			DashboardFactory->RegisterViewFactory(MakeShared<FVirtualLoopDashboardViewFactory>());
 			DashboardFactory->RegisterViewFactory(PlotsDashboard);
+			DashboardFactory->RegisterViewFactory(MakeShared<FSubmixDashboardViewFactory>());
 
 			PlotsDashboard->InitPlots(SoundsDashboard);
 #endif // !WITH_EDITOR
diff --git a/Engine/Plugins/AudioInsights/Source/AudioInsights/Private/Providers/SubmixTraceProvider.cpp b/Engine/Plugins/AudioInsights/Source/AudioInsights/Private/Providers/SubmixTraceProvider.cpp
new file mode 100644
--- /dev/null
+++ b/Engine/Plugins/AudioInsights/Source/AudioInsights/Private/Providers/SubmixTraceProvider.cpp
@@ -0,0 +1,343 @@
+// Copyright Epic Games, Inc. All Rights Reserved.
+
+#include "Providers/SubmixTraceProvider.h"
+
+#if !WITH_EDITOR
+#include "Common/PagedArray.h"
+#include "Async/ParallelFor.h"
+#endif // !WITH_EDITOR
+
+namespace UE::Audio::Insights
+{
+	namespace FSubmixTraceProviderPrivate
+	{
+#if !WITH_EDITOR
+		template<typename T>
+		const T* FindClosestMessageToTimestamp(const TraceServices::TPagedArray<T>& InCachedMessages, const double InTimeMarker, const uint32 InSubmixId)
+		{
+			const int32 ClosestMessageToTimeStampIndex = TraceServices::PagedArrayAlgo::BinarySearchClosestBy(InCachedMessages, InTimeMarker, [](const T& Msg) { return Msg.Timestamp; });
+
+			// Iterate backwards from TimeMarker until we find the matching SubmixId
+			for (auto It = InCachedMessages.GetIteratorFromItem(ClosestMessageToTimeStampIndex); It; --It)
+			{
+				if (It->SubmixId == InSubmixId)
+				{
+					return &(*It);
+				}
+			}
+
+			return nullptr;
+		}
+#endif // !WITH_EDITOR
+	}
+
+	FSubmixTraceProvider::FSubmixTraceProvider()
+		: TDeviceDataMapTraceProvider<uint32, TSharedPtr<FSubmixDashboardEntry>>(GetName_Static())
+	{
+	}
+	
+	FName FSubmixTraceProvider::GetName_Static()
+	{
+		static const FLazyName SubmixTraceProviderName = "SubmixProvider";
+		return SubmixTraceProviderName;
+	}
+
+	UE::Trace::IAnalyzer* FSubmixTraceProvider::ConstructAnalyzer(TraceServices::IAnalysisSession& InSession)
+	{
+		class FSubmixTraceAnalyzer : public FTraceAnalyzerBase
+		{
+		public:
+			FSubmixTraceAnalyzer(TSharedRef<FSubmixTraceProvider> InProvider, TraceServices::IAnalysisSession& InSession)
+				: FTraceAnalyzerBase(InProvider)
+				, Session(InSession)
+			{
+			}
+
+			virtual void OnAnalysisBegin(const FOnAnalysisContext& Context) override
+			{
+				FTraceAnalyzerBase::OnAnalysisBegin(Context);
+
+				UE::Trace::IAnalyzer::FInterfaceBuilder& Builder = Context.InterfaceBuilder;
+
+				Builder.RouteEvent(RouteId_Loaded,         "Audio", "SubmixLoadedOnTraceStart");
+				Builder.RouteEvent(RouteId_HasActivity,    "Audio", "SubmixHasActivity");
+				Builder.RouteEvent(RouteId_Loaded,         "Audio", "SubmixLoaded");
+				Builder.RouteEvent(RouteId_EnvelopeValues, "Audio", "SubmixEnvelopeValues");
+				Builder.RouteEvent(RouteId_Unloaded,       "Audio", "SubmixUnloaded");
+			}
+
+			virtual bool OnEvent(uint16 RouteId, EStyle Style, const FOnEventContext& Context) override
+			{
+				LLM_SCOPE_BYNAME(TEXT("Insights/FSubmixTraceAnalyzer"));
+
+				FSubmixMessages& Messages = GetProvider<FSubmixTraceProvider>().TraceMessages;
+
+				switch (RouteId)
+				{
+					case RouteId_Loaded:
+					{
+						Messages.LoadedMessages.Enqueue(FSubmixLoadedMessage{ Context });
+						break;
+					}
+
+					case RouteId_HasActivity:
+					{
+						Messages.HasActivityMessages.Enqueue(FSubmixHasActivityMessage{ Context });
+						break;
+					}
+
+					case RouteId_EnvelopeValues:
+					{
+						Messages.EnvelopeValuesMessages.Enqueue(FSubmixEnvelopeValuesMessage{ Context });
+						break;
+					}
+
+					case RouteId_Unloaded:
+					{
+						Messages.UnloadedMessages.Enqueue(FSubmixUnloadedMessage{ Context });
+						break;
+					}
+
+					default:
+					{
+						return OnEventFailure(RouteId, Style, Context);
+					}
+				}
+
+				const double Timestamp = Context.EventTime.AsSeconds(Context.EventData.GetValue<uint64>("Timestamp"));
+
+				{
+					TraceServices::FAnalysisSessionEditScope SessionEditScope(Session);
+					Session.UpdateDurationSeconds(Timestamp);
+				}
+
+				return OnEventSuccess(RouteId, Style, Context);
+			}
+
+		private:
+			enum : uint16
+			{
+				RouteId_Loaded,
+				RouteId_HasActivity,
+				RouteId_EnvelopeValues,
+				RouteId_Unloaded
+			};
+
+			TraceServices::IAnalysisSession& Session;
+		};
+
+		return new FSubmixTraceAnalyzer(AsShared(), InSession);
+	}
+
+	bool FSubmixTraceProvider::ProcessMessages()
+	{
+		// Helper lambdas
+#if !WITH_EDITOR
+		auto CreateEntry = [this](const FSubmixMessageBase& Msg)
+		{
+			TSharedPtr<FSubmixDashboardEntry>* ToReturn = nullptr;
+			
+			UpdateDeviceEntry(Msg.DeviceId, Msg.SubmixId, [&ToReturn, &Msg](TSharedPtr<FSubmixDashboardEntry>& Entry)
+			{
+				if (!Entry.IsValid())
+				{
+					Entry = MakeShared<FSubmixDashboardEntry>();
+					Entry->SubmixId = Msg.SubmixId;
+				}
+
+				Entry->Timestamp = Msg.Timestamp;
+
+				ToReturn = &Entry;
+			});
+
+			return ToReturn;
+		};
+#endif // !WITH_EDITOR
+
+		auto GetEntry = [this](const FSubmixMessageBase& Msg)
+		{
+			return FindDeviceEntry(Msg.DeviceId, Msg.SubmixId);
+		};
+
+		// Process messages
+#if !WITH_EDITOR
+		ProcessMessageQueue<FSubmixLoadedMessage>(TraceMessages.LoadedMessages, CreateEntry,
+		[this](const FSubmixLoadedMessage& Msg, TSharedPtr<FSubmixDashboardEntry>* OutEntry)
+		{
+			if (SessionCachedMessages.IsValid())
+			{
+				SessionCachedMessages->LoadedCachedMessages.EmplaceBack(Msg);
+			}
+
+			if (OutEntry)
+			{
+				FSubmixDashboardEntry& EntryRef = *OutEntry->Get();
+				EntryRef.Name = Msg.Name;
+				EntryRef.NumChannels = Msg.NumChannels;
+			}
+		});
+#endif // !WITH_EDITOR
+
+		ProcessMessageQueue<FSubmixHasActivityMessage>(TraceMessages.HasActivityMessages, GetEntry,
+		[this](const FSubmixHasActivityMessage& Msg, TSharedPtr<FSubmixDashboardEntry>* OutEntry)
+		{
+#if !WITH_EDITOR
+			if (SessionCachedMessages.IsValid())
+			{
+				SessionCachedMessages->HasActivityCachedMessages.EmplaceBack(Msg);
+			}
+#endif // !WITH_EDITOR
+
+			if (OutEntry)
+			{
+				FSubmixDashboardEntry& EntryRef = *OutEntry->Get();
+				EntryRef.bHasActivity = Msg.bHasActivity;
+				EntryRef.Timestamp = Msg.Timestamp;
+			}
+		});
+
+		ProcessMessageQueue<FSubmixEnvelopeValuesMessage>(TraceMessages.EnvelopeValuesMessages, GetEntry,
+		[this](const FSubmixEnvelopeValuesMessage& Msg, TSharedPtr<FSubmixDashboardEntry>* OutEntry)
+		{
+			if (OutEntry)
+			{
+				FSubmixDashboardEntry& EntryRef = *OutEntry->Get();
+
+				EntryRef.EnvelopeValues = Msg.EnvelopeValues;
+				EntryRef.Timestamp = Msg.Timestamp;
+			}
+		});
+
+#if !WITH_EDITOR
+		ProcessMessageQueue<FSubmixUnloadedMessage>(TraceMessages.UnloadedMessages, GetEntry,
+		[this](const FSubmixUnloadedMessage& Msg, TSharedPtr<FSubmixDashboardEntry>* OutEntry)
+		{
+			if (SessionCachedMessages.IsValid())
+			{
+				SessionCachedMessages->UnloadedCachedMessages.EmplaceBack(Msg);
+			}
+
+			if (OutEntry && (*OutEntry)->Timestamp < Msg.Timestamp)
+			{
+				RemoveDeviceEntry(Msg.DeviceId, Msg.SubmixId);
+			}
+		});
+#endif // !WITH_EDITOR
+
+		return true;
+	}
+
+#if !WITH_EDITOR
+	void FSubmixTraceProvider::InitSessionCachedMessages(TraceServices::IAnalysisSession& InSession)
+	{
+		SessionCachedMessages = MakeUnique<FSubmixSessionCachedMessages>(InSession);
+	}
+
+	void FSubmixTraceProvider::OnTimingViewTimeMarkerChanged(double TimeMarker)
+	{
+		using namespace FSubmixTraceProviderPrivate;
+
+		if (!SessionCachedMessages.IsValid())
+		{
+			return;
+		}
+
+		DeviceDataMap.Empty();
+
+		// Collect all the submix loaded messages registered until this point in time 
+		for (const FSubmixLoadedMessage& SubmixLoadedCachedMessage : SessionCachedMessages->LoadedCachedMessages)
+		{
+			if (SubmixLoadedCachedMessage.Timestamp > TimeMarker)
+			{
+				break;
+			}
+
+			UpdateDeviceEntry(SubmixLoadedCachedMessage.DeviceId, SubmixLoadedCachedMessage.SubmixId, [&SubmixLoadedCachedMessage](TSharedPtr<FSubmixDashboardEntry>& Entry)
+			{
+				if (!Entry.IsValid())
+				{
+					Entry = MakeShared<FSubmixDashboardEntry>();
+					Entry->SubmixId = SubmixLoadedCachedMessage.SubmixId;
+				}
+
+				Entry->Timestamp = SubmixLoadedCachedMessage.Timestamp;
+
+				Entry->Name = *SubmixLoadedCachedMessage.Name;
+				Entry->NumChannels = SubmixLoadedCachedMessage.NumChannels;
+			});
+		}
+
+		// Selectively remove loaded messages collected in the step above by knowing which submixes were unloaded.
+		// With this we will know what are the available submixes at this point in time.
+		for (const FSubmixUnloadedMessage& SubmixUnloadedCachedMessage : SessionCachedMessages->UnloadedCachedMessages)
+		{
+			if (SubmixUnloadedCachedMessage.Timestamp > TimeMarker)
+			{
+				break;
+			}
+
+			auto* OutEntry = FindDeviceEntry(SubmixUnloadedCachedMessage.DeviceId, SubmixUnloadedCachedMessage.SubmixId);
+
+			if (OutEntry && (*OutEntry)->Timestamp < SubmixUnloadedCachedMessage.Timestamp)
+			{
+				RemoveDeviceEntry(SubmixUnloadedCachedMessage.DeviceId, SubmixUnloadedCachedMessage.SubmixId);
+			}
+		}
+
+		const FDeviceData* DeviceData = FindFilteredDeviceData();
+		if (DeviceData)
+		{
+			// Collect has activity messages from submixes (based on active submixes SubmixId)
+			struct CachedEntryInfo
+			{
+				TOptional<FSubmixHasActivityMessage> HasActivityMessage;
+			};
+
+			TArray<uint32> SubmixIdArray;
+			(*DeviceData).GenerateKeyArray(SubmixIdArray);
+
+			TArray<CachedEntryInfo> CachedEntryInfos;
+			CachedEntryInfos.SetNum(SubmixIdArray.Num());
+
+			ParallelFor(SubmixIdArray.Num(),
+			[&SubmixIdArray, &CachedEntryInfos, TimeMarker, this](const int32 Index)
+			{
+				const uint32 SubmixId = SubmixIdArray[Index];
+
+				// HasActivity
+				const FSubmixHasActivityMessage* FoundHasActivityCachedMessage = FindClosestMessageToTimestamp(SessionCachedMessages->HasActivityCachedMessages, TimeMarker, SubmixId);
+				if (FoundHasActivityCachedMessage)
+				{
+					CachedEntryInfos[Index].HasActivityMessage = *FoundHasActivityCachedMessage;
+				}
+			});
+
+			// Update the device entries with the collected info
+			for (const CachedEntryInfo& CachedEntryInfo : CachedEntryInfos)
+			{
+				if (CachedEntryInfo.HasActivityMessage.IsSet())
+				{
+					const FSubmixHasActivityMessage& HasActivityMessage = CachedEntryInfo.HasActivityMessage.GetValue();
+
+					UpdateDeviceEntry(HasActivityMessage.DeviceId, HasActivityMessage.SubmixId, [&HasActivityMessage](TSharedPtr<FSubmixDashboardEntry>& Entry)
+					{
+						if (!Entry.IsValid())
+						{
+							Entry = MakeShared<FSubmixDashboardEntry>();
+							Entry->SubmixId = HasActivityMessage.SubmixId;
+						}
+
+						Entry->Timestamp = HasActivityMessage.Timestamp;
+
+						Entry->bHasActivity = HasActivityMessage.bHasActivity;
+					});
+				}
+			}
+		}
+
+		// Call parent method to update LastMessageId
+		FTraceProviderBase::OnTimingViewTimeMarkerChanged(TimeMarker);
+	}
+#endif // !WITH_EDITOR
+
+} // namespace UE::Audio::Insights
diff --git a/Engine/Plugins/AudioInsights/Source/AudioInsights/Private/Providers/SubmixTraceProvider.h b/Engine/Plugins/AudioInsights/Source/AudioInsights/Private/Providers/SubmixTraceProvider.h
new file mode 100644
--- /dev/null
+++ b/Engine/Plugins/AudioInsights/Source/AudioInsights/Private/Providers/SubmixTraceProvider.h
@@ -0,0 +1,36 @@
+// Copyright Epic Games, Inc. All Rights Reserved.
+
+#pragma once
+
+#include "AudioInsightsTraceProviderBase.h"
+#include "Messages/SubmixTraceMessages.h"
+#include "UObject/NameTypes.h"
+
+namespace UE::Audio::Insights
+{
+	class FSubmixTraceProvider : public TDeviceDataMapTraceProvider<uint32, TSharedPtr<FSubmixDashboardEntry>>, public TSharedFromThis<FSubmixTraceProvider>
+	{
+	public:
+		FSubmixTraceProvider();
+		virtual ~FSubmixTraceProvider() = default;
+
+		static FName GetName_Static();
+		
+		virtual UE::Trace::IAnalyzer* ConstructAnalyzer(TraceServices::IAnalysisSession& InSession) override;
+
+#if !WITH_EDITOR
+		virtual void InitSessionCachedMessages(TraceServices::IAnalysisSession& InSession) override;
+#endif // !WITH_EDITOR
+
+	private:
+		virtual bool ProcessMessages() override;
+
+#if !WITH_EDITOR
+		virtual void OnTimingViewTimeMarkerChanged(double TimeMarker) override;
+
+		TUniquePtr<FSubmixSessionCachedMessages> SessionCachedMessages;
+#endif // !WITH_EDITOR
+
+		FSubmixMessages TraceMessages;
+	};
+} // namespace UE::Audio::Insights
diff --git a/Engine/Plugins/AudioInsights/Source/AudioInsights/Private/Views/SubmixDashboardViewFactory.cpp b/Engine/Plugins/AudioInsights/Source/AudioInsights/Private/Views/SubmixDashboardViewFactory.cpp
new file mode 100644
--- /dev/null
+++ b/Engine/Plugins/AudioInsights/Source/AudioInsights/Private/Views/SubmixDashboardViewFactory.cpp
@@ -0,0 +1,229 @@
+// Copyright Epic Games, Inc. All Rights Reserved.
+#include "Views/SubmixDashboardViewFactory.h"
+
+#include "AudioInsightsStyle.h"
+#include "IAudioInsightsModule.h"
+#include "Internationalization/Text.h"
+#include "Providers/SubmixTraceProvider.h"
+
+#define LOCTEXT_NAMESPACE "AudioInsights"
+
+namespace UE::Audio::Insights
+{
+	namespace FSubmixDashboardViewFactoryPrivate
+	{
+		const FSubmixDashboardEntry& CastEntry(const IDashboardDataViewEntry& InData)
+		{
+			return static_cast<const FSubmixDashboardEntry&>(InData);
+		};
+	}
+
+	FSubmixDashboardViewFactory::FSubmixDashboardViewFactory()
+	{
+		IAudioInsightsTraceModule& AudioInsightsTraceModule = IAudioInsightsModule::GetChecked().GetTraceModule();
+
+		SubmixProvider = MakeShared<FSubmixTraceProvider>();
+
+		AudioInsightsTraceModule.AddTraceProvider(SubmixProvider);
+
+		Providers = TArray<TSharedPtr<FTraceProviderBase>>
+		{
+			SubmixProvider
+		};
+
+		SortByColumn = "Name";
+		SortMode     = EColumnSortMode::Ascending;
+	}
+
+	FName FSubmixDashboardViewFactory::GetName() const
+	{
+		return "Submix";
+	}
+
+	FText FSubmixDashboardViewFactory::GetDisplayName() const
+	{
+		return LOCTEXT("AudioDashboard_Submix_DisplayName", "Submixes");
+	}
+	
+	FSlateIcon FSubmixDashboardViewFactory::GetIcon() const
+	{
+		return FSlateStyle::Get().CreateIcon("AudioInsights.Icon.Submix");
+	}
+
+	EDefaultDashboardTabStack FSubmixDashboardViewFactory::GetDefaultTabStack() const
+	{
+		return EDefaultDashboardTabStack::Analysis;
+	}
+
+	TSharedRef<SWidget> FSubmixDashboardViewFactory::GenerateWidgetForColumn(TSharedRef<IDashboardDataViewEntry> InRowData, const FName& InColumnName)
+	{
+		using namespace FSubmixDashboardViewFactoryPrivate;
+
+		if (InColumnName == "Active")
+		{
+			static const FLinearColor DarkGreen(0.027f, 0.541f, 0.22f);
+			static const float Radius = 4.0f;
+			static const FVector2f Size(7.0f, 7.0f);
+			static const FSlateRoundedBoxBrush GreenRoundedBrush(DarkGreen, Radius, Size);
+
+			return SNew(SBox)
+				.Clipping(EWidgetClipping::ClipToBounds)
+				.Padding(6.0f)
+				.Visibility_Lambda([InRowData]()
+				{ 
+					const FSubmixDashboardEntry& SubmixDashboardEntry = CastEntry(InRowData.Get());
+					return SubmixDashboardEntry.bHasActivity ? EVisibility::Visible : EVisibility::Hidden;
+				})
+				[
+					SNew(SImage)
+					.Image(&GreenRoundedBrush)
+				];
+		}
+		else if (InColumnName == "Name")
+		{
+			const FColumnData& ColumnData = GetColumns()[InColumnName];
+			const FText ValueText = ColumnData.GetDisplayValue(InRowData.Get());
+
+			if (ValueText.IsEmpty())
+			{
+				return SNullWidget::NullWidget;
+			}
+
+			return SNew(SHorizontalBox)
+				.Clipping(EWidgetClipping::ClipToBounds)
+				+ SHorizontalBox::Slot()
+				.AutoWidth()
+				[
+					SNew(STextBlock)
+					.Text(ValueText)
+					.OnDoubleClicked_Lambda([this, InRowData](const FGeometry& MyGeometry, const FPointerEvent& PointerEvent)
+					{
+						return FReply::Unhandled();
+					})
+				];
+		}
+
+		return SNullWidget::NullWidget;
+	}
+
+	void FSubmixDashboardViewFactory::ProcessEntries(FTraceTableDashboardViewFactory::EProcessReason Reason)
+	{
+		const FString FilterString = GetSearchFilterText().ToString();
+		
+		FTraceTableDashboardViewFactory::FilterEntries<FSubmixTraceProvider>([&FilterString](const IDashboardDataViewEntry& Entry)
+		{
+			const FSubmixDashboardEntry& SubmixEntry = static_cast<const FSubmixDashboardEntry&>(Entry);
+			
+			return !SubmixEntry.GetDisplayName().ToString().Contains(FilterString);
+		});
+	}
+
+	TSharedRef<SWidget> FSubmixDashboardViewFactory::MakeWidget(TSharedRef<SDockTab> OwnerTab, const FSpawnTabArgs& SpawnTabArgs)
+	{
+		using namespace FSubmixDashboardViewFactoryPrivate;
+
+		if (!DashboardWidget.IsValid())
+		{
+			DashboardWidget = FTraceTableDashboardViewFactory::MakeWidget(OwnerTab, SpawnTabArgs);
+
+			if (FilteredEntriesListView.IsValid())
+			{
+				FilteredEntriesListView->SetSelectionMode(ESelectionMode::Single);
+			}
+		}
+
+		return DashboardWidget->AsShared();
+	}
+
+	const TMap<FName, FTraceTableDashboardViewFactory::FColumnData>& FSubmixDashboardViewFactory::GetColumns() const
+	{
+		using namespace FSubmixDashboardViewFactoryPrivate;
+
+		auto CreateColumnData = []()
+		{
+			return TMap<FName, FTraceTableDashboardViewFactory::FColumnData>
+			{
+				{
+					"Active",
+					{
+						LOCTEXT("SubmixDashboard_ActiveDisplayName", "Active"),
+						[](const IDashboardDataViewEntry& InData) { return FText::GetEmpty(); },
+						nullptr, /* GetIconName */
+						false,	 /* bDefaultHidden */
+						0.08f,	 /* FillWidth */
+						EHorizontalAlignment::HAlign_Center
+					}
+				},
+				{
+					"Name",
+					{
+						LOCTEXT("SubmixDashboard_NameColumnDisplayName", "Name"),
+						[](const IDashboardDataViewEntry& InData) { return CastEntry(InData).GetDisplayName(); },
+						nullptr, /* GetIconName */
+						false,	 /* bDefaultHidden */
+						0.92f,	 /* FillWidth */
+						EHorizontalAlignment::HAlign_Left
+					}
+				}
+			};
+		};
+		
+		static const TMap<FName, FTraceTableDashboardViewFactory::FColumnData> ColumnData = CreateColumnData();
+		
+		return ColumnData;
+	}
+
+	void FSubmixDashboardViewFactory::SortTable()
+	{
+		using namespace FSubmixDashboardViewFactoryPrivate;
+
+		if (SortByColumn == "Active")
+		{
+			if (SortMode == EColumnSortMode::Ascending)
+			{
+				DataViewEntries.Sort([](const TSharedPtr<IDashboardDataViewEntry>& A, const TSharedPtr<IDashboardDataViewEntry>& B)
+				{
+					const FSubmixDashboardEntry& AData = CastEntry(*A.Get());
+					const FSubmixDashboardEntry& BData = CastEntry(*B.Get());
+
+					return !AData.bHasActivity && BData.bHasActivity;
+				});
+			}
+			else if (SortMode == EColumnSortMode::Descending)
+			{
+				DataViewEntries.Sort([](const TSharedPtr<IDashboardDataViewEntry>& A, const TSharedPtr<IDashboardDataViewEntry>& B)
+				{
+					const FSubmixDashboardEntry& AData = CastEntry(*A.Get());
+					const FSubmixDashboardEntry& BData = CastEntry(*B.Get());
+
+					return !BData.bHasActivity && AData.bHasActivity;
+				});
+			}
+		}
+		else if (SortByColumn == "Name")
+		{
+			if (SortMode == EColumnSortMode::Ascending)
+			{
+				DataViewEntries.Sort([](const TSharedPtr<IDashboardDataViewEntry>& A, const TSharedPtr<IDashboardDataViewEntry>& B)
+				{
+					const FSubmixDashboardEntry& AData = CastEntry(*A.Get());
+					const FSubmixDashboardEntry& BData = CastEntry(*B.Get());
+
+					return AData.GetDisplayName().CompareToCaseIgnored(BData.GetDisplayName()) < 0;
+				});
+			}
+			else if (SortMode == EColumnSortMode::Descending)
+			{
+				DataViewEntries.Sort([](const TSharedPtr<IDashboardDataViewEntry>& A, const TSharedPtr<IDashboardDataViewEntry>& B)
+				{
+					const FSubmixDashboardEntry& AData = CastEntry(*A.Get());
+					const FSubmixDashboardEntry& BData = CastEntry(*B.Get());
+
+					return BData.GetDisplayName().CompareToCaseIgnored(AData.GetDisplayName()) < 0;
+				});
+			}
+		}
+	}
+} // namespace UE::Audio::Insights
+
+#undef LOCTEXT_NAMESPACE
diff --git a/Engine/Plugins/AudioInsights/Source/AudioInsights/Public/Messages/SubmixTraceMessages.h b/Engine/Plugins/AudioInsights/Source/AudioInsights/Public/Messages/SubmixTraceMessages.h
new file mode 100644
--- /dev/null
+++ b/Engine/Plugins/AudioInsights/Source/AudioInsights/Public/Messages/SubmixTraceMessages.h
@@ -0,0 +1,132 @@
+// Copyright Epic Games, Inc. All Rights Reserved.
+#pragma once
+
+#include "Messages/AnalyzerMessageQueue.h"
+#include "Trace/Analyzer.h"
+#include "Views/TableDashboardViewFactory.h"
+
+#if !WITH_EDITOR
+#include "Common/PagedArray.h"
+#include "TraceServices/Model/AnalysisSession.h"
+#endif // !WITH_EDITOR
+
+namespace UE::Audio::Insights
+{
+	struct FSubmixMessageBase
+	{
+		FSubmixMessageBase() = default;
+
+		FSubmixMessageBase(const UE::Trace::IAnalyzer::FOnEventContext& InContext)
+		{
+			const UE::Trace::IAnalyzer::FEventData& EventData = InContext.EventData;
+
+			DeviceId  = static_cast<::Audio::FDeviceId>(EventData.GetValue<uint32>("DeviceId"));
+			SubmixId  = EventData.GetValue<uint32>("SubmixId");
+			Timestamp = InContext.EventTime.AsSeconds(EventData.GetValue<uint64>("Timestamp"));
+		}
+
+		::Audio::FDeviceId DeviceId = INDEX_NONE;
+		uint32 SubmixId = INDEX_NONE;
+		double Timestamp = 0.0;
+	};
+
+	struct FSubmixLoadedMessage : public FSubmixMessageBase
+	{
+		FSubmixLoadedMessage() = default;
+
+		FSubmixLoadedMessage(const UE::Trace::IAnalyzer::FOnEventContext& InContext)
+			: FSubmixMessageBase(InContext)
+		{
+			const UE::Trace::IAnalyzer::FEventData& EventData = InContext.EventData;
+
+			EventData.GetString("Name", Name);
+			NumChannels = EventData.GetValue<int32>("NumChannels");
+		}
+
+		FString Name;
+		int32 NumChannels = 0;
+	};
+
+	struct FSubmixHasActivityMessage : public FSubmixMessageBase
+	{
+		FSubmixHasActivityMessage() = default;
+
+		FSubmixHasActivityMessage(const UE::Trace::IAnalyzer::FOnEventContext& InContext)
+			: FSubmixMessageBase(InContext)
+		{
+			const UE::Trace::IAnalyzer::FEventData& EventData = InContext.EventData;
+
+			bHasActivity = EventData.GetValue<bool>("HasActivity");
+		}
+
+		bool bHasActivity = false;
+	};
+
+	struct FSubmixEnvelopeValuesMessage : public FSubmixMessageBase
+	{
+		FSubmixEnvelopeValuesMessage() = default;
+
+		FSubmixEnvelopeValuesMessage(const UE::Trace::IAnalyzer::FOnEventContext& InContext)
+			: FSubmixMessageBase(InContext)
+		{
+			const UE::Trace::IAnalyzer::FEventData& EventData = InContext.EventData;
+
+			const TArrayView<const float> EnvelopeValuesView = EventData.GetArrayView<float>("EnvelopeValues");
+
+			for (const auto& EnvelopeValue : EnvelopeValuesView)
+			{
+				EnvelopeValues.Emplace(EnvelopeValue);
+			}
+		}
+
+		TArray<float> EnvelopeValues;
+	};
+
+	using FSubmixUnloadedMessage = FSubmixMessageBase;
+
+	class FSubmixMessages
+	{
+		UE::Audio::Insights::TAnalyzerMessageQueue<FSubmixLoadedMessage> LoadedMessages{ 0.1 };
+		UE::Audio::Insights::TAnalyzerMessageQueue<FSubmixHasActivityMessage> HasActivityMessages{ 0.1 };
+		UE::Audio::Insights::TAnalyzerMessageQueue<FSubmixEnvelopeValuesMessage> EnvelopeValuesMessages{ 0.1 };
+		UE::Audio::Insights::TAnalyzerMessageQueue<FSubmixUnloadedMessage> UnloadedMessages{ 0.1 };
+
+		friend class FSubmixTraceProvider;
+	};
+
+#if !WITH_EDITOR
+	struct FSubmixSessionCachedMessages
+	{
+		FSubmixSessionCachedMessages(TraceServices::IAnalysisSession& InSession)
+			: LoadedCachedMessages(InSession.GetLinearAllocator(), 4096)
+			, HasActivityCachedMessages(InSession.GetLinearAllocator(), 4096)
+			, UnloadedCachedMessages(InSession.GetLinearAllocator(), 4096)
+		{
+
+		}
+
+		TraceServices::TPagedArray<FSubmixLoadedMessage> LoadedCachedMessages;
+		TraceServices::TPagedArray<FSubmixHasActivityMessage> HasActivityCachedMessages;
+		TraceServices::TPagedArray<FSubmixUnloadedMessage> UnloadedCachedMessages;
+	};
+#endif // !WITH_EDITOR
+
+	struct FSubmixDashboardEntry : public IObjectDashboardEntry
+	{
+		FSubmixDashboardEntry() = default;
+		virtual ~FSubmixDashboardEntry() = default;
+
+		virtual FText GetDisplayName() const override { return FText::FromString(FSoftObjectPath(Name).GetAssetName()); }
+		virtual const UObject* GetObject() const override { return FSoftObjectPath(Name).ResolveObject(); }
+		virtual UObject* GetObject() override { return FSoftObjectPath(Name).ResolveObject(); }
+		virtual bool IsValid() const override { return SubmixId != static_cast<uint32>(INDEX_NONE); }
+
+		uint32 SubmixId = INDEX_NONE;
+		double Timestamp = 0.0;
+
+		FString Name;
+		int32 NumChannels = 0;
+		bool bHasActivity = false;
+		TArray<float> EnvelopeValues;
+	};
+} // namespace UE::Audio::Insights
diff --git a/Engine/Plugins/AudioInsights/Source/AudioInsights/Public/Views/SubmixDashboardViewFactory.h b/Engine/Plugins/AudioInsights/Source/AudioInsights/Public/Views/SubmixDashboardViewFactory.h
new file mode 100644
--- /dev/null
+++ b/Engine/Plugins/AudioInsights/Source/AudioInsights/Public/Views/SubmixDashboardViewFactory.h
@@ -0,0 +1,36 @@
+// Copyright Epic Games, Inc. All Rights Reserved.
+#pragma once
+
+#include "Templates/SharedPointer.h"
+#include "Views/TableDashboardViewFactory.h"
+
+#define UE_API AUDIOINSIGHTS_API
+
+namespace UE::Audio::Insights
+{
+	class FSubmixTraceProvider;
+
+	class FSubmixDashboardViewFactory : public FTraceObjectTableDashboardViewFactory
+	{
+	public:
+		UE_API FSubmixDashboardViewFactory();
+		UE_API virtual ~FSubmixDashboardViewFactory() = default;
+
+		UE_API virtual FName GetName() const override;
+		UE_API virtual FText GetDisplayName() const override;
+		UE_API virtual FSlateIcon GetIcon() const override;
+		UE_API virtual EDefaultDashboardTabStack GetDefaultTabStack() const override;
+		UE_API virtual TSharedRef<SWidget> MakeWidget(TSharedRef<SDockTab> OwnerTab, const FSpawnTabArgs& SpawnTabArgs) override;
+
+	protected:
+		UE_API virtual TSharedRef<SWidget> GenerateWidgetForColumn(TSharedRef<IDashboardDataViewEntry> InRowData, const FName& InColumnName) override;
+		UE_API virtual void ProcessEntries(FTraceTableDashboardViewFactory::EProcessReason Reason) override;
+		UE_API virtual const TMap<FName, FTraceTableDashboardViewFactory::FColumnData>& GetColumns() const override;
+
+		UE_API virtual void SortTable() override;
+
+		TSharedPtr<FSubmixTraceProvider> SubmixProvider;
+	};
+} // namespace UE::Audio::Insights
+
+#undef UE_API
