[Audio Insights] Added SAudioMeterView: a container with an audio meter with a name below it (for the standalone Audio Meters dashboard)

diff --git a/Engine/Plugins/AudioInsights/Source/AudioInsights/AudioInsights.Build.cs b/Engine/Plugins/AudioInsights/Source/AudioInsights/AudioInsights.Build.cs
--- a/Engine/Plugins/AudioInsights/Source/AudioInsights/AudioInsights.Build.cs
+++ b/Engine/Plugins/AudioInsights/Source/AudioInsights/AudioInsights.Build.cs
@@ -25,6 +25,7 @@ public class AudioInsights : ModuleRules
 			new string[]
 			{
 				"AudioMixerCore",
+				"AudioWidgetsCore",
 				"CoreUObject",
 				"InputCore",
 				"OutputLog",
diff --git a/Engine/Plugins/AudioInsights/Source/AudioInsights/Private/Views/SAudioMeterView.cpp b/Engine/Plugins/AudioInsights/Source/AudioInsights/Private/Views/SAudioMeterView.cpp
new file mode 100644
--- /dev/null
+++ b/Engine/Plugins/AudioInsights/Source/AudioInsights/Private/Views/SAudioMeterView.cpp
@@ -0,0 +1,126 @@
+// Copyright Epic Games, Inc. All Rights Reserved.
+
+#include "Views/SAudioMeterView.h"
+
+#include "AudioInsightsDataSource.h"
+#include "DSP/Dsp.h"
+#include "SAudioMeterWidget.h"
+#include "Widgets/Layout/SBox.h"
+#include "Widgets/SBoxPanel.h"
+#include "Widgets/Text/STextBlock.h"
+
+#define LOCTEXT_NAMESPACE "AudioInsights"
+
+namespace UE::Audio::Insights
+{
+	namespace SAudioMeterViewPrivate
+	{
+		const TAttribute<TArray<FAudioMeterChannelInfo>> CreateMeterChannelInfoAttribute(const TSharedPtr<IAudioMeterableEntry>& InAudioMeterableEntry)
+		{
+			return TAttribute<TArray<FAudioMeterChannelInfo>>::Create([InAudioMeterableEntry]()
+				{
+					TArray<FAudioMeterChannelInfo> MeterChannelInfo;
+
+					if (InAudioMeterableEntry.IsValid())
+					{
+						const TArray<float>& EnvelopeValues = InAudioMeterableEntry->GetEnvelopeValues();
+
+						if (!EnvelopeValues.IsEmpty())
+						{
+							MeterChannelInfo.Reserve(EnvelopeValues.Num());
+
+							for (const float EnvelopeValue : EnvelopeValues)
+							{
+								MeterChannelInfo.Emplace(::Audio::ConvertToDecibels(EnvelopeValue), -160.0f, -160.0f);
+							}
+						}
+						else
+						{
+							const int32 NumChannels = InAudioMeterableEntry->GetNumChannels();
+							
+							for (int32 Index = 0; Index < NumChannels; ++Index)
+							{
+								MeterChannelInfo.Emplace(-160.0f, -160.0f, -160.0f);
+							}
+						}
+
+					}
+
+					return MeterChannelInfo;
+				});
+		}
+	}
+
+	void SAudioMeterView::Construct(const SAudioMeterView::FArguments& InArgs, const TSharedPtr<IAudioMeterableEntry>& InAudioMeterableEntry)
+	{
+		using namespace SAudioMeterViewPrivate;
+
+		ChildSlot
+		[
+			SNew(SVerticalBox)
+			+ SVerticalBox::Slot()
+			[
+				SNew(SVerticalBox)
+				.Clipping(EWidgetClipping::ClipToBounds)
+				// Height padding
+				+ SVerticalBox::Slot()
+				.FillHeight(0.05)
+				[
+					SNew(SBox)
+				]
+				// Audio Meter container
+				+ SVerticalBox::Slot()
+				.FillHeight(0.8)
+				[
+					// Audio Meter
+					SNew(SHorizontalBox)
+					+ SHorizontalBox::Slot()
+					.HAlign(HAlign_Center)
+					.VAlign(VAlign_Fill)
+					[
+						SAssignNew(AudioMeterWidget, SAudioMeterWidget)
+						.Orientation(EOrientation::Orient_Vertical)
+						.BackgroundColor(FLinearColor::Transparent)
+						.MeterBackgroundColor(FAudioMeterDefaultColorWidgetStyle::GetDefault().MeterBackgroundColor)
+						.MeterValueColor(FAudioMeterDefaultColorWidgetStyle::GetDefault().MeterValueColor)
+						.MeterPeakColor(FAudioMeterDefaultColorWidgetStyle::GetDefault().MeterPeakColor)
+						.MeterClippingColor(FAudioMeterDefaultColorWidgetStyle::GetDefault().MeterClippingColor)
+						.MeterScaleColor(FAudioMeterDefaultColorWidgetStyle::GetDefault().MeterScaleColor)
+						.MeterScaleLabelColor(FAudioMeterDefaultColorWidgetStyle::GetDefault().MeterScaleLabelColor)
+						.MeterChannelInfo(CreateMeterChannelInfoAttribute(InAudioMeterableEntry))
+					]
+				]
+				// Height padding
+				+ SVerticalBox::Slot()
+				.FillHeight(0.025)
+				[
+					SNew(SBox)
+				]
+				// Audio meter name label
+				+ SVerticalBox::Slot()
+				.FillHeight(0.1)
+				[
+					SNew(STextBlock)
+					.Text(InArgs._NameText)
+					.Justification(ETextJustify::Center)
+					.ColorAndOpacity(InArgs._NameTextColor)
+				]
+				// Height padding
+				+ SVerticalBox::Slot()
+				.FillHeight(0.025)
+				[
+					SNew(SBox)
+				]
+			]
+		];
+	}
+
+	void SAudioMeterView::SetAudioMeterChannelInfo(const TSharedRef<IAudioMeterableEntry>& InAudioMeterableEntry)
+	{
+		using namespace SAudioMeterViewPrivate;
+
+		AudioMeterWidget->SetMeterChannelInfo(CreateMeterChannelInfoAttribute(InAudioMeterableEntry));
+	}
+} // namespace UE::Audio::Insights
+
+#undef LOCTEXT_NAMESPACE
diff --git a/Engine/Plugins/AudioInsights/Source/AudioInsights/Public/AudioInsightsDataSource.h b/Engine/Plugins/AudioInsights/Source/AudioInsights/Public/AudioInsightsDataSource.h
--- a/Engine/Plugins/AudioInsights/Source/AudioInsights/Public/AudioInsightsDataSource.h
+++ b/Engine/Plugins/AudioInsights/Source/AudioInsights/Public/AudioInsightsDataSource.h
@@ -31,4 +31,13 @@ namespace UE::Audio::Insights
 		TArray<TSharedPtr<IDashboardDataTreeViewEntry>> Children;
 		bool bIsExpanded = false;
 	};
+
+	class IAudioMeterableEntry
+	{
+	public:
+		virtual ~IAudioMeterableEntry() = default;
+
+		virtual int32 GetNumChannels() const = 0;
+		virtual const TArray<float>& GetEnvelopeValues() const = 0;
+	};
 } // namespace UE::Audio::Insights
diff --git a/Engine/Plugins/AudioInsights/Source/AudioInsights/Public/Views/SAudioMeterView.h b/Engine/Plugins/AudioInsights/Source/AudioInsights/Public/Views/SAudioMeterView.h
new file mode 100644
--- /dev/null
+++ b/Engine/Plugins/AudioInsights/Source/AudioInsights/Public/Views/SAudioMeterView.h
@@ -0,0 +1,31 @@
+// Copyright Epic Games, Inc. All Rights Reserved.
+
+#pragma once
+
+#include "Widgets/SCompoundWidget.h"
+
+class SAudioMeterWidget;
+
+namespace UE::Audio::Insights
+{
+	class IAudioMeterableEntry;
+
+	class SAudioMeterView : public SCompoundWidget
+	{
+	public:
+		SLATE_BEGIN_ARGS(SAudioMeterView)
+			: _NameText(FText::GetEmpty())
+			, _NameTextColor(FSlateColor(FColor::White))
+		{}
+			SLATE_ARGUMENT(FText, NameText)
+			SLATE_ARGUMENT(FSlateColor, NameTextColor)
+		SLATE_END_ARGS()
+
+		void Construct(const FArguments& InArgs, const TSharedPtr<IAudioMeterableEntry>& InAudioMeterableEntry);
+
+		void SetAudioMeterChannelInfo(const TSharedRef<IAudioMeterableEntry>& InAudioMeterableEntry);
+
+	protected:
+		TSharedPtr<SAudioMeterWidget> AudioMeterWidget;
+	};
+} // namespace UE::Audio::Insights
