[Audio Insights] Add Select All option to Visible Columns submenu in sound dashboard Settings cog menu.
Also: fix persisting column visibility for renamed column, with an ensure added to catch such bugs.

diff --git a/Engine/Plugins/AudioInsights/Source/AudioInsights/Private/Views/SoundDashboardViewFactory.cpp b/Engine/Plugins/AudioInsights/Source/AudioInsights/Private/Views/SoundDashboardViewFactory.cpp
--- a/Engine/Plugins/AudioInsights/Source/AudioInsights/Private/Views/SoundDashboardViewFactory.cpp
+++ b/Engine/Plugins/AudioInsights/Source/AudioInsights/Private/Views/SoundDashboardViewFactory.cpp
@@ -2,6 +2,7 @@
 #include "Views/SoundDashboardViewFactory.h"
 
 #include "Algo/Accumulate.h"
+#include "Algo/AnyOf.h"
 #include "AudioInsightsModule.h"
 #include "AudioInsightsStyle.h"
 #include "AudioInsightsUtils.h"
@@ -1506,6 +1507,18 @@ namespace UE::Audio::Insights
 
 	void FSoundDashboardViewFactory::BuildVisibleColumnsMenuContent(FMenuBuilder& MenuBuilder)
 	{
+		MenuBuilder.AddMenuEntry(
+			TAttribute<FText>::CreateSP(this, &FSoundDashboardViewFactory::GetToggleAllColumnsVisibilityText),
+			FText::GetEmpty(),
+			FSlateIcon(),
+			FUIAction(
+				FExecuteAction::CreateSP(this, &FSoundDashboardViewFactory::ToggleAllColumnsVisibility),
+				FCanExecuteAction::CreateLambda([]() { return true; }),
+				FGetActionCheckState::CreateSP(this, &FSoundDashboardViewFactory::GetToggleAllColumnsVisibilityCheckState)
+			),
+			NAME_None,
+			EUserInterfaceActionType::ToggleButton);
+
 		for (const SHeaderRow::FColumn& Column : HeaderRowWidget->GetColumns())
 		{
 			MenuBuilder.AddMenuEntry(
@@ -1533,6 +1546,61 @@ namespace UE::Audio::Insights
 		return HeaderRowWidget->IsColumnVisible(ColumnId) ? ECheckBoxState::Checked : ECheckBoxState::Unchecked;
 	}
 
+	void FSoundDashboardViewFactory::ToggleAllColumnsVisibility()
+	{
+		const ECheckBoxState CurrentState = GetToggleAllColumnsVisibilityCheckState();
+		const bool bVisibleIsDesired = CurrentState == ECheckBoxState::Unchecked || CurrentState == ECheckBoxState::Undetermined;
+
+		// Null our VisibleColumnSettings TOptional while updating HeaderRowWidget so that we ignore OnHiddenColumnsListChanged callbacks:
+		VisibleColumnSettings.Reset();
+
+		for (const SHeaderRow::FColumn& Column : HeaderRowWidget->GetColumns())
+		{
+			if (Column.bIsVisible != bVisibleIsDesired)
+			{
+				HeaderRowWidget->SetShowGeneratedColumn(Column.ColumnId, bVisibleIsDesired);
+			}
+		}
+
+		// Reinitialize the VisibleColumnSettings and perform a single OnHiddenColumnsListChanged call to update it and persist the settings:
+		VisibleColumnSettings.Emplace();
+		OnHiddenColumnsListChanged();
+	}
+
+	ECheckBoxState FSoundDashboardViewFactory::GetToggleAllColumnsVisibilityCheckState() const
+	{
+		const TIndirectArray<SHeaderRow::FColumn>& Columns = HeaderRowWidget->GetColumns();
+		const bool bAnyColumnsVisible = Algo::AnyOf(Columns, [](const SHeaderRow::FColumn& Column) { return Column.bIsVisible; });
+		const bool bAnyColumnsInvisible = Algo::AnyOf(Columns, [](const SHeaderRow::FColumn& Column) { return !Column.bIsVisible; });
+
+		if (bAnyColumnsVisible && bAnyColumnsInvisible)
+		{
+			return ECheckBoxState::Undetermined;
+		}
+		else if (bAnyColumnsVisible)
+		{
+			return ECheckBoxState::Checked;
+		}
+		else
+		{
+			return ECheckBoxState::Unchecked;
+		}
+	}
+
+	FText FSoundDashboardViewFactory::GetToggleAllColumnsVisibilityText() const
+	{
+		switch (GetToggleAllColumnsVisibilityCheckState())
+		{
+		case ECheckBoxState::Unchecked:
+		case ECheckBoxState::Undetermined:
+			return LOCTEXT("SoundsDashboard_VisibleColumns_SelectAll", "Select: All");
+		case ECheckBoxState::Checked:
+			return LOCTEXT("SoundsDashboard_VisibleColumns_SelectNone", "Select: None");
+		default:
+			return FText::GetEmpty();
+		}
+	}
+
 	void FSoundDashboardViewFactory::FilterText()
 	{
 		using namespace FSoundDashboardViewFactoryPrivate;
@@ -1646,7 +1714,7 @@ namespace UE::Audio::Insights
 					.IconName		  = "AudioInsights.Icon.SoundDashboard.Mute",
 					.TooltipText	  = MuteColumnTooltip,
 					.bShowDisplayName = false,
-					.bDefaultHidden   = !VisibleColumns.bMute,
+					.bDefaultHidden   = !VisibleColumns.GetIsVisible(MuteColumnName),
 					.FillWidth        = 0.05f,
 					.Alignment        = EHorizontalAlignment::HAlign_Center
 				}
@@ -1658,7 +1726,7 @@ namespace UE::Audio::Insights
 					.IconName         = "AudioInsights.Icon.SoundDashboard.Solo",
 					.TooltipText	  = SoloColumnTooltip,
 					.bShowDisplayName = false,
-					.bDefaultHidden   = !VisibleColumns.bSolo,
+					.bDefaultHidden   = !VisibleColumns.GetIsVisible(SoloColumnName),
 					.FillWidth        = 0.05f,
 					.Alignment        = EHorizontalAlignment::HAlign_Center
 				}
@@ -1671,7 +1739,7 @@ namespace UE::Audio::Insights
 					.IconName		  = "AudioInsights.Icon.SoundDashboard.Plots",
 					.TooltipText	  = PlotColumnTooltip,
 					.bShowDisplayName = false,
-					.bDefaultHidden   = !VisibleColumns.bPlot,
+					.bDefaultHidden   = !VisibleColumns.GetIsVisible(PlotColumnName),
 					.FillWidth		  = 0.05f,
 					.Alignment		  = EHorizontalAlignment::HAlign_Center
 				}
@@ -1683,7 +1751,7 @@ namespace UE::Audio::Insights
 					.IconName         = NAME_None,
 					.TooltipText	  = NameColumnTooltip,
 					.bShowDisplayName = true,
-					.bDefaultHidden   = !VisibleColumns.bName,
+					.bDefaultHidden   = !VisibleColumns.GetIsVisible(NameColumnName),
 					.FillWidth        = 0.5f,
 					.Alignment        = EHorizontalAlignment::HAlign_Left
 				}
@@ -1695,7 +1763,7 @@ namespace UE::Audio::Insights
 					.IconName		  = NAME_None,
 					.TooltipText	  = PlayOrderColumnTooltip,
 					.bShowDisplayName = true,
-					.bDefaultHidden   = !VisibleColumns.bPlayOrder,
+					.bDefaultHidden   = !VisibleColumns.GetIsVisible(PlayOrderColumnName),
 					.FillWidth		  = 0.1f,
 					.Alignment		  = EHorizontalAlignment::HAlign_Left
 				}
@@ -1707,7 +1775,7 @@ namespace UE::Audio::Insights
 					.IconName		  = NAME_None,
 					.TooltipText	  = PriorityColumnTooltip,
 					.bShowDisplayName = true,
-					.bDefaultHidden	  = !VisibleColumns.bPriority,
+					.bDefaultHidden	  = !VisibleColumns.GetIsVisible(PriorityColumnName),
 					.FillWidth		  = 0.08f,
 					.Alignment		  = EHorizontalAlignment::HAlign_Left
 				}
@@ -1719,7 +1787,7 @@ namespace UE::Audio::Insights
 					.IconName         = NAME_None,
 					.TooltipText	  = DistanceColumnTooltip,
 					.bShowDisplayName = true,
-					.bDefaultHidden   = !VisibleColumns.bDistance,
+					.bDefaultHidden   = !VisibleColumns.GetIsVisible(DistanceColumnName),
 					.FillWidth        = 0.1f,
 					.Alignment        = EHorizontalAlignment::HAlign_Left
 				}
@@ -1731,7 +1799,7 @@ namespace UE::Audio::Insights
 					.IconName		  = NAME_None,
 					.TooltipText	  = DistanceAttenuationColumnTooltip,
 					.bShowDisplayName = true,
-					.bDefaultHidden   = !VisibleColumns.bDistanceAttenuation,
+					.bDefaultHidden   = !VisibleColumns.GetIsVisible(DistanceAttenuationColumnName),
 					.FillWidth		  = 0.1f,
 					.Alignment		  = EHorizontalAlignment::HAlign_Left
 				}
@@ -1743,7 +1811,7 @@ namespace UE::Audio::Insights
 					.IconName         = NAME_None,
 					.TooltipText	  = AmplitudeColumnTooltip,
 					.bShowDisplayName = true,
-					.bDefaultHidden   = !VisibleColumns.bAmplitude,
+					.bDefaultHidden   = !VisibleColumns.GetIsVisible(AmplitudeColumnName),
 					.FillWidth        = 0.12f,
 					.Alignment        = EHorizontalAlignment::HAlign_Left
 				}
@@ -1755,7 +1823,7 @@ namespace UE::Audio::Insights
 					.IconName         = NAME_None,
 					.TooltipText	  = VolumeColumnTooltip,
 					.bShowDisplayName = true,
-					.bDefaultHidden   = !VisibleColumns.bVolume,
+					.bDefaultHidden   = !VisibleColumns.GetIsVisible(VolumeColumnName),
 					.FillWidth        = 0.1f,
 					.Alignment        = EHorizontalAlignment::HAlign_Left
 				}
@@ -1767,7 +1835,7 @@ namespace UE::Audio::Insights
 					.IconName		  = NAME_None,
 					.TooltipText	  = LPFFreqColumnTooltip,
 					.bShowDisplayName = true,
-					.bDefaultHidden   = !VisibleColumns.bLPFFreq,
+					.bDefaultHidden   = !VisibleColumns.GetIsVisible(LPFFreqColumnName),
 					.FillWidth		  = 0.1f,
 					.Alignment		  = EHorizontalAlignment::HAlign_Left
 				}
@@ -1779,7 +1847,7 @@ namespace UE::Audio::Insights
 					.IconName		  = NAME_None,
 					.TooltipText	  = HPFFreqColumnTooltip,
 					.bShowDisplayName = true,
-					.bDefaultHidden   = !VisibleColumns.bHPFFreq,
+					.bDefaultHidden   = !VisibleColumns.GetIsVisible(HPFFreqColumnName),
 					.FillWidth		  = 0.1f,
 					.Alignment		  = EHorizontalAlignment::HAlign_Left
 				}
@@ -1791,7 +1859,7 @@ namespace UE::Audio::Insights
 					.IconName         = NAME_None,
 					.TooltipText	  = PitchColumnTooltip,
 					.bShowDisplayName = true,
-					.bDefaultHidden   = !VisibleColumns.bPitch,
+					.bDefaultHidden   = !VisibleColumns.GetIsVisible(PitchColumnName),
 					.FillWidth        = 0.1f,
 					.Alignment        = EHorizontalAlignment::HAlign_Left
 				}
@@ -1803,7 +1871,7 @@ namespace UE::Audio::Insights
 					.IconName		  = NAME_None,
 					.TooltipText	  = RelativeRenderCostColumnTooltip,
 					.bShowDisplayName = true,
-					.bDefaultHidden   = !VisibleColumns.bRelativeRenderCost,
+					.bDefaultHidden   = !VisibleColumns.GetIsVisible(RelativeRenderCostColumnName),
 					.FillWidth		  = 0.1f,
 					.Alignment		  = EHorizontalAlignment::HAlign_Left
 				}
@@ -1816,7 +1884,7 @@ namespace UE::Audio::Insights
 					.IconName		  = NAME_None,
 					.TooltipText	  = ActorLabelColumnTooltip,
 					.bShowDisplayName = true,
-					.bDefaultHidden   = !VisibleColumns.bActor,
+					.bDefaultHidden   = !VisibleColumns.GetIsVisible(ActorLabelColumnName),
 					.FillWidth		  = 0.1f,
 					.Alignment		  = EHorizontalAlignment::HAlign_Left
 				}
@@ -1829,7 +1897,7 @@ namespace UE::Audio::Insights
 					.IconName		  = NAME_None,
 					.TooltipText	  = CategoryColumnTooltip,
 					.bShowDisplayName = true,
-					.bDefaultHidden   = !VisibleColumns.bCategory,
+					.bDefaultHidden   = !VisibleColumns.GetIsVisible(CategoryColumnName),
 					.FillWidth		  = 0.1f,
 					.Alignment		  = EHorizontalAlignment::HAlign_Left
 				}
diff --git a/Engine/Plugins/AudioInsights/Source/AudioInsights/Public/Settings/SoundDashboardSettings.h b/Engine/Plugins/AudioInsights/Source/AudioInsights/Public/Settings/SoundDashboardSettings.h
--- a/Engine/Plugins/AudioInsights/Source/AudioInsights/Public/Settings/SoundDashboardSettings.h
+++ b/Engine/Plugins/AudioInsights/Source/AudioInsights/Public/Settings/SoundDashboardSettings.h
@@ -105,7 +105,7 @@ struct FSoundDashboardVisibleColumns
 	bool bRelativeRenderCost = false;
 
 	UPROPERTY(EditAnywhere, config, Category = SoundDashboard)
-	bool bActor = false;
+	bool bActorLabel = false;
 
 	UPROPERTY(EditAnywhere, config, Category = SoundDashboard)
 	bool bCategory = false;
@@ -121,6 +121,7 @@ struct FSoundDashboardVisibleColumns
 			}
 		}
 
+		ensureMsgf(false, TEXT("No VisibleColumns setting for ColumnId: %s"), *ColumnId.ToString());
 		return false;
 	}
 
diff --git a/Engine/Plugins/AudioInsights/Source/AudioInsights/Public/Views/SoundDashboardViewFactory.h b/Engine/Plugins/AudioInsights/Source/AudioInsights/Public/Views/SoundDashboardViewFactory.h
--- a/Engine/Plugins/AudioInsights/Source/AudioInsights/Public/Views/SoundDashboardViewFactory.h
+++ b/Engine/Plugins/AudioInsights/Source/AudioInsights/Public/Views/SoundDashboardViewFactory.h
@@ -162,6 +162,9 @@ namespace UE::Audio::Insights
 		UE_API void BuildVisibleColumnsMenuContent(FMenuBuilder& MenuBuilder);
 		UE_API void ToggleColumnVisibility(FName ColumnId);
 		UE_API ECheckBoxState GetIsColumnVisibleCheckedState(FName ColumnId) const;
+		UE_API void ToggleAllColumnsVisibility();
+		UE_API ECheckBoxState GetToggleAllColumnsVisibilityCheckState() const;
+		UE_API FText GetToggleAllColumnsVisibilityText() const;
 
 		UE_API void FilterText();
 		UE_API void RecursiveSort(TArray<TSharedPtr<IDashboardDataTreeViewEntry>>& OutTree, TFunctionRef<bool(const FSoundDashboardEntry&, const FSoundDashboardEntry&)> Predicate);
