[Audio Insights] Mute and Solo selected vs. filtered sounds

diff --git a/Engine/Plugins/AudioInsights/Source/AudioInsights/Private/Views/SoundDashboardViewFactory.cpp b/Engine/Plugins/AudioInsights/Source/AudioInsights/Private/Views/SoundDashboardViewFactory.cpp
--- a/Engine/Plugins/AudioInsights/Source/AudioInsights/Private/Views/SoundDashboardViewFactory.cpp
+++ b/Engine/Plugins/AudioInsights/Source/AudioInsights/Private/Views/SoundDashboardViewFactory.cpp
@@ -182,6 +182,46 @@ namespace UE::Audio::Insights
 				];
 		};
 
+		TSharedRef<SBox> CreateMuteSoloSelectedButtonContent(const FName& InIconName,
+			const FName& InTextStyle,
+			TFunction<FSlateColor()> InColorFunction,
+			TFunction<FText()> InTextFunction)
+		{
+			TSharedRef<SHorizontalBox> ButtonContainerWidget = SNew(SHorizontalBox);
+
+			if (!InIconName.IsNone())
+			{
+				ButtonContainerWidget->AddSlot()
+					.AutoWidth()
+					.VAlign(VAlign_Center)
+					[
+						SNew(SImage)
+							.ColorAndOpacity_Lambda([InColorFunction]() { return InColorFunction(); })
+							.Image(FSlateStyle::Get().GetBrush(InIconName))
+					];
+			}
+
+			const float LeftPadding = InIconName.IsNone() ? 0.0f : 4.0f;
+
+			ButtonContainerWidget->AddSlot()
+				.VAlign(VAlign_Center)
+				.Padding(LeftPadding, 0.0f, 0.0f, 0.0f)
+				.AutoWidth()
+				[
+					SNew(STextBlock)
+						.TextStyle(&FSlateStyle::Get().GetWidgetStyle<FTextBlockStyle>(InTextStyle))
+						.Justification(ETextJustify::Center)
+						.Text_Lambda([InTextFunction]() { return InTextFunction(); })
+						.ColorAndOpacity_Lambda([InColorFunction]() { return InColorFunction(); })
+				];
+
+			return SNew(SBox)
+				.HeightOverride(16.0f)
+				[
+					ButtonContainerWidget
+				];
+		};
+
 		void RefreshInitExpandStateRecursive(FSoundDashboardEntry& SoundEntry)
 		{
 			SoundEntry.bHasSetInitExpansion = false;
@@ -839,6 +879,21 @@ namespace UE::Audio::Insights
 	{
 		using namespace FSoundDashboardViewFactoryPrivate;
 
+		TFunction<FSlateColor()> MuteColorFunc = [this]() { return bIsMuteFilteredMode ? FSlateColor(FColor::Green) : FSlateColor::UseForeground(); };
+		TFunction<FText()> MuteTextFunc = [&FilteredList = FilteredEntriesListView]()
+			{
+				return FilteredList.IsValid() && FilteredList->GetNumItemsSelected() > 0 ?
+					LOCTEXT("SoundDashboard_MuteButtonTextSelect", "Mute Selected") :
+					LOCTEXT("SoundDashboard_MuteButtonTextFilter", "Mute Filtered");
+			};
+		TFunction<FSlateColor()> SoloColorFunc = [this]() { return bIsSoloFilteredMode ? FSlateColor(FColor::Green) : FSlateColor::UseForeground(); };
+		TFunction<FText()> SoloTextFunc = [&FilteredList = FilteredEntriesListView]()
+			{
+				return FilteredList.IsValid() && FilteredList->GetNumItemsSelected() > 0 ?
+					LOCTEXT("SoundDashboard_SoloButtonTextSelect", "Solo Selected") :
+					LOCTEXT("SoundDashboard_SoloButtonTextFilter", "Solo Filtered");
+			};
+
 		return SNew(SHorizontalBox)
 			// Mute button
 			+ SHorizontalBox::Slot()
@@ -852,13 +907,21 @@ namespace UE::Audio::Insights
 				{
 					if (FilteredEntriesListView.IsValid())
 					{
-						ToggleMuteSoloEntries(FilteredEntriesListView->GetSelectedItems(), EMuteSoloMode::Mute);
+						if (FilteredEntriesListView->GetNumItemsSelected() > 0)
+						{
+							ToggleMuteSoloEntries(FilteredEntriesListView->GetSelectedItems(), EMuteSoloMode::Mute);
+						}
+						else
+						{
+							bIsMuteFilteredMode = !bIsMuteFilteredMode;
+							MuteSoloFilteredEntries();
+						}
 					}
 
 					return FReply::Handled();
 				})
 				[
-					CreateButtonContentWidget("AudioInsights.Icon.SoundDashboard.Mute", LOCTEXT("SoundDashboard_MuteButtonText", "Mute Selected"), "SmallButtonText")
+					CreateMuteSoloSelectedButtonContent("AudioInsights.Icon.SoundDashboard.Mute", "SmallButtonText", MuteColorFunc, MuteTextFunc)
 				]
 			]
 			// Solo button
@@ -873,13 +936,21 @@ namespace UE::Audio::Insights
 				{
 					if (FilteredEntriesListView.IsValid())
 					{
-						ToggleMuteSoloEntries(FilteredEntriesListView->GetSelectedItems(), EMuteSoloMode::Solo);
+						if (FilteredEntriesListView->GetNumItemsSelected() > 0)
+						{
+							ToggleMuteSoloEntries(FilteredEntriesListView->GetSelectedItems(), EMuteSoloMode::Solo);
+						}
+						else
+						{
+							bIsSoloFilteredMode = !bIsSoloFilteredMode;
+							MuteSoloFilteredEntries();
+						}
 					}
 
 					return FReply::Handled();
 				})
 				[
-					CreateButtonContentWidget("AudioInsights.Icon.SoundDashboard.Solo", LOCTEXT("SoundDashboard_SoloButtonText", "Solo Selected"), "SmallButtonText")
+					CreateMuteSoloSelectedButtonContent("AudioInsights.Icon.SoundDashboard.Solo", "SmallButtonText", SoloColorFunc, SoloTextFunc)
 				]
 			]
 			// Clear Mutes/Solos button
@@ -1317,6 +1388,11 @@ namespace UE::Audio::Insights
 
 		const FSoundDashboardEntry& SoundDashboardEntry = CastEntry(*Item);
 
+		if (bIsMuteFilteredMode || bIsSoloFilteredMode)
+		{
+			MuteSoloFilteredEntries();
+		}
+
 		return SNew(SRowWidget, OwnerTable, Item, AsShared())
 			.RenderOpacity_Lambda([Item]
 			{
@@ -1633,6 +1709,12 @@ namespace UE::Audio::Insights
 				SetFilteredVisibility(*CategoryEntry, FilterString);
 			}
 		}
+
+		if (bIsMuteFilteredMode || bIsSoloFilteredMode)
+		{
+			ClearMutesAndSolos();
+			MuteSoloFilteredEntries();
+		}
 	}
 
 	TSharedPtr<SWidget> FSoundDashboardViewFactory::OnConstructContextMenu()
@@ -2891,6 +2973,67 @@ namespace UE::Audio::Insights
 #endif // ENABLE_AUDIO_DEBUG
 	}
 
+	void FSoundDashboardViewFactory::MuteSoloFilteredEntries()
+	{
+#if ENABLE_AUDIO_DEBUG
+		using namespace FSoundDashboardViewFactoryPrivate;
+
+		FAudioDeviceManager* AudioDeviceManager = FAudioDeviceManager::Get();
+		if (AudioDeviceManager == nullptr)
+		{
+			return;
+		}
+
+		if (!bIsMuteFilteredMode && !bIsSoloFilteredMode)
+		{
+			ClearMutesAndSolos();
+			return;
+		}
+
+		::Audio::FAudioDebugger& AudioDebugger = AudioDeviceManager->GetDebugger();
+
+		const EMuteSoloMode Mode = bIsSoloFilteredMode ? EMuteSoloMode::Solo : EMuteSoloMode::Mute;
+		bool bSomethingIsVisible = false;
+
+		for (const TSharedPtr<IDashboardDataTreeViewEntry>& Entry : DataViewEntries)
+		{
+			const FSoundDashboardEntry& CategoryEntry = FSoundDashboardViewFactoryPrivate::CastEntry(*Entry);
+			if (CategoryEntry.bIsCategory)
+			{
+				for (const TSharedPtr<IDashboardDataTreeViewEntry>& EntryChild : Entry->Children)
+				{
+					const FSoundDashboardEntry& SoundEntry = FSoundDashboardViewFactoryPrivate::CastEntry(*EntryChild);
+					if (SoundEntry.bIsVisible)
+					{
+						bSomethingIsVisible = true;
+						if (!FSoundDashboardViewFactoryPrivate::IsMuteSolo(AudioDebugger, *EntryChild, true /*bInCheckChildren*/, Mode))
+						{
+							SetMuteSolo(*EntryChild, Mode, true /*bInOnOff*/);
+						}
+					}
+				}
+			}
+		}
+
+		// If you have filtered out everything but are soloing to listen things that do match the filter,
+		// we mute the rest of the mix at this moment. 
+		if (!GetSearchFilterText().ToString().IsEmpty() && Mode == EMuteSoloMode::Solo && !bSomethingIsVisible)
+		{
+			for (TSharedPtr<IDashboardDataTreeViewEntry>& Entry : DataViewEntries)
+			{
+				FSoundDashboardEntry& CategoryEntry = FSoundDashboardViewFactoryPrivate::CastEntry(*Entry);
+				if (CategoryEntry.bIsCategory)
+				{
+					for (TSharedPtr<IDashboardDataTreeViewEntry>& EntryChild : Entry->Children)
+					{
+						SetMuteSolo(*EntryChild, EMuteSoloMode::Mute, true /*bInOnOff*/);
+					}
+				}
+			}
+		}
+#endif //ENABLE_AUDIO_DEBUG
+	}
+
 	TArray<TObjectPtr<UObject>> FSoundDashboardViewFactory::GetSelectedEditableAssets() const
 	{
 		TArray<TObjectPtr<UObject>> Objects;
diff --git a/Engine/Plugins/AudioInsights/Source/AudioInsights/Public/Views/SoundDashboardViewFactory.h b/Engine/Plugins/AudioInsights/Source/AudioInsights/Public/Views/SoundDashboardViewFactory.h
--- a/Engine/Plugins/AudioInsights/Source/AudioInsights/Public/Views/SoundDashboardViewFactory.h
+++ b/Engine/Plugins/AudioInsights/Source/AudioInsights/Public/Views/SoundDashboardViewFactory.h
@@ -188,6 +188,7 @@ namespace UE::Audio::Insights
 			TFunctionRef<bool(const IDashboardDataTreeViewEntry& /*InEntry*/, const bool /*bInCheckChildren*/)> IsMuteSoloFunc);
 
 		UE_API void ToggleMuteSoloEntries(const TArray<TSharedPtr<IDashboardDataTreeViewEntry>>& InEntries, const EMuteSoloMode InMuteSoloMode);
+		UE_API void MuteSoloFilteredEntries();
 
 		UE_API TArray<TObjectPtr<UObject>> GetSelectedEditableAssets() const;
 #else
@@ -254,6 +255,9 @@ namespace UE::Audio::Insights
 		bool bShowRecentlyStoppedSounds = false;
 		bool bTreeViewModeChanged = false;
 		bool bDisplayAmpPeakInDb = true;
+		bool bIsMuteFilteredMode = false;
+		bool bIsSoloFilteredMode = false;
+
 		TOptional<FSoundDashboardVisibleColumns> VisibleColumnSettings;
 	};
 } // namespace UE::Audio::Insights
