Fix ENABLE_AUDIO_DEBUG build break

diff --git a/Engine/Plugins/AudioInsights/Source/AudioInsights/Private/Views/SoundDashboardViewFactory.cpp b/Engine/Plugins/AudioInsights/Source/AudioInsights/Private/Views/SoundDashboardViewFactory.cpp
--- a/Engine/Plugins/AudioInsights/Source/AudioInsights/Private/Views/SoundDashboardViewFactory.cpp
+++ b/Engine/Plugins/AudioInsights/Source/AudioInsights/Private/Views/SoundDashboardViewFactory.cpp
@@ -1388,10 +1388,12 @@ namespace UE::Audio::Insights
 
 		const FSoundDashboardEntry& SoundDashboardEntry = CastEntry(*Item);
 
+#if ENABLE_AUDIO_DEBUG
 		if (bIsMuteFilteredMode || bIsSoloFilteredMode)
 		{
 			MuteSoloFilteredEntries();
 		}
+#endif //ENABLE_AUDIO_DEBUG
 
 		return SNew(SRowWidget, OwnerTable, Item, AsShared())
 			.RenderOpacity_Lambda([Item]
@@ -1710,11 +1712,13 @@ namespace UE::Audio::Insights
 			}
 		}
 
+#if ENABLE_AUDIO_DEBUG
 		if (bIsMuteFilteredMode || bIsSoloFilteredMode)
 		{
 			ClearMutesAndSolos();
 			MuteSoloFilteredEntries();
 		}
+#endif // ENABLE_AUDIO_DEBUG
 	}
 
 	TSharedPtr<SWidget> FSoundDashboardViewFactory::OnConstructContextMenu()
