[Audio Insights] Submix dashboard: added submix asset register/unregister support (for Editor version only)

diff --git a/Engine/Plugins/AudioInsights/Source/AudioInsights/Private/Providers/SubmixAssetProvider.cpp b/Engine/Plugins/AudioInsights/Source/AudioInsights/Private/Providers/SubmixAssetProvider.cpp
new file mode 100644
--- /dev/null
+++ b/Engine/Plugins/AudioInsights/Source/AudioInsights/Private/Providers/SubmixAssetProvider.cpp
@@ -0,0 +1,164 @@
+// Copyright Epic Games, Inc. All Rights Reserved.
+
+#include "Providers/SubmixAssetProvider.h"
+
+#if WITH_EDITOR
+#include "IAudioInsightsModule.h"
+#include "AssetRegistry/AssetData.h"
+#include "AssetRegistry/AssetRegistryModule.h"
+#include "AudioDeviceManager.h"
+#include "AudioMixerDevice.h"
+#include "Sound/SoundSubmix.h"
+
+namespace UE::Audio::Insights
+{
+	FSubmixAssetProvider::FSubmixAssetProvider()
+	{
+		FAssetRegistryModule& AssetRegistryModule = FModuleManager::LoadModuleChecked<FAssetRegistryModule>("AssetRegistry");
+
+		AssetRegistryModule.Get().OnAssetAdded().AddRaw(this, &FSubmixAssetProvider::OnAssetAdded);
+		AssetRegistryModule.Get().OnAssetRemoved().AddRaw(this, &FSubmixAssetProvider::OnAssetRemoved);
+		AssetRegistryModule.Get().OnFilesLoaded().AddRaw(this, &FSubmixAssetProvider::OnFilesLoaded);
+
+		FAudioDeviceManagerDelegates::OnAudioDeviceCreated.AddRaw(this, &FSubmixAssetProvider::OnActiveAudioDeviceChanged);
+		FAudioDeviceManagerDelegates::OnAudioDeviceDestroyed.AddRaw(this, &FSubmixAssetProvider::OnActiveAudioDeviceChanged);
+
+		FTraceAuxiliary::OnTraceStarted.AddRaw(this, &FSubmixAssetProvider::OnTraceStarted);
+	}
+
+	FSubmixAssetProvider::~FSubmixAssetProvider()
+	{
+		if (FAssetRegistryModule* AssetRegistryModule = FModuleManager::GetModulePtr<FAssetRegistryModule>("AssetRegistry"))
+		{
+			AssetRegistryModule->Get().OnAssetAdded().RemoveAll(this);
+			AssetRegistryModule->Get().OnAssetRemoved().RemoveAll(this);
+			AssetRegistryModule->Get().OnFilesLoaded().RemoveAll(this);
+		}
+
+		FAudioDeviceManagerDelegates::OnAudioDeviceCreated.RemoveAll(this);
+		FAudioDeviceManagerDelegates::OnAudioDeviceDestroyed.RemoveAll(this);
+
+		FTraceAuxiliary::OnTraceStarted.RemoveAll(this);
+	}
+
+	void FSubmixAssetProvider::OnAssetAdded(const FAssetData& InAssetData)
+	{
+		if (bAreFilesLoaded && InAssetData.AssetClassPath == FTopLevelAssetPath(USoundSubmix::StaticClass()))
+		{
+			constexpr bool bShouldNotify = true;
+			AddSubmixAsset(InAssetData, bShouldNotify);
+		}
+	}
+
+	void FSubmixAssetProvider::OnAssetRemoved(const FAssetData& InAssetData)
+	{
+		if (InAssetData.AssetClassPath == FTopLevelAssetPath(USoundSubmix::StaticClass()))
+		{
+			RemoveSubmixAsset(InAssetData);
+		}
+	}
+
+	void FSubmixAssetProvider::OnFilesLoaded()
+	{
+		bAreFilesLoaded = true;
+		UpdateSubmixAssetNames();
+	}
+
+	void FSubmixAssetProvider::OnActiveAudioDeviceChanged(::Audio::FDeviceId InDeviceId)
+	{
+		UpdateSubmixAssetNames();
+	}
+
+	void FSubmixAssetProvider::OnTraceStarted(FTraceAuxiliary::EConnectionType TraceType, const FString& TraceDestination)
+	{
+		UpdateSubmixAssetNames();
+	}
+
+	void FSubmixAssetProvider::RequestEntriesUpdate()
+	{
+		UpdateSubmixAssetNames();
+	}
+	
+	void FSubmixAssetProvider::AddSubmixAsset(const FAssetData& InAssetData, const bool bInShouldNotify)
+	{
+		const bool bIsSubmixEntryAlreadAdded = SubmixAssetEntries.ContainsByPredicate(
+			[AssetName = InAssetData.GetObjectPathString()](const TSharedPtr<FSubmixDashboardEntry>& SubmixDashboardEntry)
+			{
+				return SubmixDashboardEntry->Name == AssetName;
+			});
+
+		if (!bIsSubmixEntryAlreadAdded)
+		{
+			const IAudioInsightsModule& AudioInsightsModule = IAudioInsightsModule::GetEditorChecked();
+			const ::Audio::FDeviceId AudioDeviceId = AudioInsightsModule.GetDeviceId();
+
+			TSharedPtr<FSubmixDashboardEntry>& SubmixDashboardEntry = SubmixAssetEntries.Add_GetRef(MakeShared<FSubmixDashboardEntry>());
+			SubmixDashboardEntry->Timestamp = FPlatformTime::Seconds() - GStartTime;
+			SubmixDashboardEntry->SubmixId  = GetTypeHash(InAssetData.GetObjectPathString());
+			SubmixDashboardEntry->Name      = InAssetData.GetObjectPathString();
+
+			if (const FAudioDeviceManager* AudioDeviceManager = FAudioDeviceManager::Get())
+			{
+				if (const ::Audio::FMixerDevice* MixerDevice = static_cast<const ::Audio::FMixerDevice*>(AudioDeviceManager->GetAudioDeviceRaw(AudioDeviceId)))
+				{
+					SubmixDashboardEntry->NumChannels = MixerDevice->GetNumDeviceChannels();
+				}
+			}
+
+			if (bInShouldNotify)
+			{
+				OnSubmixAssetAdded.ExecuteIfBound(SubmixDashboardEntry);
+			}
+		}
+	}
+
+	void FSubmixAssetProvider::RemoveSubmixAsset(const FAssetData& InAssetData)
+	{
+		const TSharedPtr<FSubmixDashboardEntry>* FoundSubmixDashboardEntry = SubmixAssetEntries.FindByPredicate(
+			[AssetName = InAssetData.GetObjectPathString()](const TSharedPtr<FSubmixDashboardEntry>& InSubmixDashboardEntry)
+			{
+				return InSubmixDashboardEntry->Name == AssetName;
+			});
+
+		if (FoundSubmixDashboardEntry)
+		{
+			const TSharedPtr<FSubmixDashboardEntry> SubmixDashboardEntry = *FoundSubmixDashboardEntry;
+
+			if (SubmixDashboardEntry.IsValid())
+			{
+				const uint32 SubmixId = SubmixDashboardEntry->SubmixId;
+
+				SubmixAssetEntries.Remove(SubmixDashboardEntry);
+
+				OnSubmixAssetRemoved.ExecuteIfBound(SubmixId);
+			}
+		}
+	}
+
+	void FSubmixAssetProvider::UpdateSubmixAssetNames()
+	{
+		// Get all USoundSubmix assets
+		TArray<FAssetData> AssetDataArray;
+
+		FAssetRegistryModule& AssetRegistryModule = FModuleManager::LoadModuleChecked<FAssetRegistryModule>("AssetRegistry");
+		AssetRegistryModule.Get().GetAssetsByClass(FTopLevelAssetPath(USoundSubmix::StaticClass()), AssetDataArray);
+
+		// Build SubmixAssetEntries
+		SubmixAssetEntries.Reset();
+
+		constexpr bool bShouldNotify = false;
+		for (const FAssetData& AssetData : AssetDataArray)
+		{
+			AddSubmixAsset(AssetData, bShouldNotify);
+		}
+
+		SubmixAssetEntries.Sort([](const TSharedPtr<FSubmixDashboardEntry>& A, const TSharedPtr<FSubmixDashboardEntry>& B)
+		{
+			return A->GetDisplayName().CompareToCaseIgnored(B->GetDisplayName()) < 0;
+		});
+
+		OnSubmixAssetListUpdated.ExecuteIfBound(SubmixAssetEntries);
+	}
+} // namespace UE::Audio::Insights
+
+#endif // WITH_EDITOR
diff --git a/Engine/Plugins/AudioInsights/Source/AudioInsights/Private/Providers/SubmixAssetProvider.h b/Engine/Plugins/AudioInsights/Source/AudioInsights/Private/Providers/SubmixAssetProvider.h
new file mode 100644
--- /dev/null
+++ b/Engine/Plugins/AudioInsights/Source/AudioInsights/Private/Providers/SubmixAssetProvider.h
@@ -0,0 +1,50 @@
+// Copyright Epic Games, Inc. All Rights Reserved.
+#pragma once
+
+#if WITH_EDITOR
+#include "Delegates/Delegate.h"
+#include "Messages/SubmixTraceMessages.h"
+#include "ProfilingDebugging/TraceAuxiliary.h"
+
+struct FAssetData;
+
+namespace UE::Audio::Insights
+{
+	class FSubmixAssetProvider : public TSharedFromThis<FSubmixAssetProvider>
+	{
+	public:
+		FSubmixAssetProvider();
+		virtual ~FSubmixAssetProvider();
+
+		void RequestEntriesUpdate();
+
+		DECLARE_DELEGATE_OneParam(FOnSubmixAssetAdded, const TSharedPtr<FSubmixDashboardEntry>& /*SubmixDashboardEntry*/);
+		FOnSubmixAssetAdded OnSubmixAssetAdded;
+
+		DECLARE_DELEGATE_OneParam(FOnSubmixAssetRemoved, const uint32 /*SubmixId*/);
+		FOnSubmixAssetRemoved OnSubmixAssetRemoved;
+
+		DECLARE_DELEGATE_OneParam(FOnSubmixAssetListUpdated, const TArray<TSharedPtr<FSubmixDashboardEntry>>& /*SubmixDashboardEntries*/);
+		FOnSubmixAssetListUpdated OnSubmixAssetListUpdated;
+
+	private:
+		void OnAssetAdded(const FAssetData& InAssetData);
+		void OnAssetRemoved(const FAssetData& InAssetData);
+		void OnFilesLoaded();
+
+		void OnActiveAudioDeviceChanged(::Audio::FDeviceId InDeviceId);
+		void OnTraceStarted(FTraceAuxiliary::EConnectionType TraceType, const FString& TraceDestination);
+
+		void AddSubmixAsset(const FAssetData& InAssetData, const bool bInShouldNotify);
+		void RemoveSubmixAsset(const FAssetData& InAssetData);
+
+		void UpdateSubmixAssetNames();
+
+		bool bAreFilesLoaded = false;
+		bool bAssetEntriesNeedRefreshing = false;
+
+		TArray<TSharedPtr<FSubmixDashboardEntry>> SubmixAssetEntries;
+	};
+} // namespace UE::Audio::Insights
+
+#endif // WITH_EDITOR
diff --git a/Engine/Plugins/AudioInsights/Source/AudioInsights/Private/Providers/SubmixTraceProvider.cpp b/Engine/Plugins/AudioInsights/Source/AudioInsights/Private/Providers/SubmixTraceProvider.cpp
--- a/Engine/Plugins/AudioInsights/Source/AudioInsights/Private/Providers/SubmixTraceProvider.cpp
+++ b/Engine/Plugins/AudioInsights/Source/AudioInsights/Private/Providers/SubmixTraceProvider.cpp
@@ -2,10 +2,13 @@
 
 #include "Providers/SubmixTraceProvider.h"
 
-#if !WITH_EDITOR
-#include "Common/PagedArray.h"
+#if WITH_EDITOR
+#include "AudioDevice.h"
+#include "AudioDeviceManager.h"
+#else
 #include "Async/ParallelFor.h"
-#endif // !WITH_EDITOR
+#include "Common/PagedArray.h"
+#endif // WITH_EDITOR
 
 namespace UE::Audio::Insights
 {
@@ -34,6 +37,20 @@ namespace UE::Audio::Insights
 	FSubmixTraceProvider::FSubmixTraceProvider()
 		: TDeviceDataMapTraceProvider<uint32, TSharedPtr<FSubmixDashboardEntry>>(GetName_Static())
 	{
+#if WITH_EDITOR
+		SubmixAssetProvider.OnSubmixAssetAdded.BindRaw(this, &FSubmixTraceProvider::HandleOnSubmixAssetAdded);
+		SubmixAssetProvider.OnSubmixAssetRemoved.BindRaw(this, &FSubmixTraceProvider::HandleOnSubmixAssetRemoved);
+		SubmixAssetProvider.OnSubmixAssetListUpdated.BindRaw(this, &FSubmixTraceProvider::HandleOnSubmixAssetListUpdated);
+#endif // WITH_EDITOR
+	}
+
+	FSubmixTraceProvider::~FSubmixTraceProvider()
+	{
+#if WITH_EDITOR
+		SubmixAssetProvider.OnSubmixAssetAdded.Unbind();
+		SubmixAssetProvider.OnSubmixAssetRemoved.Unbind();
+		SubmixAssetProvider.OnSubmixAssetListUpdated.Unbind();
+#endif // WITH_EDITOR
 	}
 	
 	FName FSubmixTraceProvider::GetName_Static()
@@ -227,7 +244,92 @@ namespace UE::Audio::Insights
 		return true;
 	}
 
-#if !WITH_EDITOR
+#if WITH_EDITOR
+	void FSubmixTraceProvider::HandleOnSubmixAssetAdded(const TSharedPtr<FSubmixDashboardEntry>& InSubmixDashboardEntry)
+	{
+		if (const FAudioDeviceManager* DeviceManager = FAudioDeviceManager::Get())
+		{
+			const TArray<FAudioDevice*> AudioDevices = DeviceManager->GetAudioDevices();
+
+			for (const FAudioDevice* AudioDevice : AudioDevices)
+			{
+				if (AudioDevice)
+				{
+					const TSharedPtr<FSubmixDashboardEntry>* FoundSubmixDashboardEntry = FindDeviceEntry(AudioDevice->DeviceID, InSubmixDashboardEntry->SubmixId);
+
+					if (FoundSubmixDashboardEntry == nullptr)
+					{
+						UpdateDeviceEntry(AudioDevice->DeviceID, InSubmixDashboardEntry->SubmixId, [&InSubmixDashboardEntry](TSharedPtr<FSubmixDashboardEntry>& Entry)
+						{
+							if (!Entry.IsValid())
+							{
+								Entry = InSubmixDashboardEntry;
+							}
+						});
+					}
+
+					OnSubmixAssetAdded.ExecuteIfBound();
+				}
+			}
+		}
+	}
+
+	void FSubmixTraceProvider::HandleOnSubmixAssetRemoved(const uint32 InSubmixId)
+	{
+		if (const FAudioDeviceManager* DeviceManager = FAudioDeviceManager::Get())
+		{
+			const TArray<FAudioDevice*> AudioDevices = DeviceManager->GetAudioDevices();
+
+			for (const FAudioDevice* AudioDevice : AudioDevices)
+			{
+				if (AudioDevice)
+				{
+					RemoveDeviceEntry(AudioDevice->DeviceID, InSubmixId);
+
+					OnSubmixAssetRemoved.ExecuteIfBound();
+				}
+			}
+		}
+	}
+
+	void FSubmixTraceProvider::HandleOnSubmixAssetListUpdated(const TArray<TSharedPtr<FSubmixDashboardEntry>>& InSubmixAssetEntries)
+	{
+		const FAudioDeviceManager* DeviceManager = FAudioDeviceManager::Get();
+
+		if (DeviceManager == nullptr)
+		{
+			return;
+		}
+
+		DeviceDataMap.Reset();
+
+		for (const FAudioDevice* AudioDevice : DeviceManager->GetAudioDevices())
+		{
+			if (AudioDevice == nullptr)
+			{
+				continue;
+			}
+
+			for (const auto& SubmixDashboardEntry : InSubmixAssetEntries)
+			{
+				const TSharedPtr<FSubmixDashboardEntry>* FoundSubmixDashboardEntry = FindDeviceEntry(AudioDevice->DeviceID, SubmixDashboardEntry->SubmixId);
+
+				if (FoundSubmixDashboardEntry == nullptr)
+				{
+					UpdateDeviceEntry(AudioDevice->DeviceID, SubmixDashboardEntry->SubmixId, [&SubmixDashboardEntry](TSharedPtr<FSubmixDashboardEntry>& Entry)
+					{
+						if (!Entry.IsValid())
+						{
+							Entry = SubmixDashboardEntry;
+						}
+					});
+				}
+			}
+		}
+
+		OnSubmixAssetListUpdated.ExecuteIfBound();
+	}
+#else
 	void FSubmixTraceProvider::InitSessionCachedMessages(TraceServices::IAnalysisSession& InSession)
 	{
 		SessionCachedMessages = MakeUnique<FSubmixSessionCachedMessages>(InSession);
@@ -338,6 +440,6 @@ namespace UE::Audio::Insights
 		// Call parent method to update LastMessageId
 		FTraceProviderBase::OnTimingViewTimeMarkerChanged(TimeMarker);
 	}
-#endif // !WITH_EDITOR
+#endif // WITH_EDITOR
 
 } // namespace UE::Audio::Insights
diff --git a/Engine/Plugins/AudioInsights/Source/AudioInsights/Private/Providers/SubmixTraceProvider.h b/Engine/Plugins/AudioInsights/Source/AudioInsights/Private/Providers/SubmixTraceProvider.h
--- a/Engine/Plugins/AudioInsights/Source/AudioInsights/Private/Providers/SubmixTraceProvider.h
+++ b/Engine/Plugins/AudioInsights/Source/AudioInsights/Private/Providers/SubmixTraceProvider.h
@@ -6,13 +6,18 @@
 #include "Messages/SubmixTraceMessages.h"
 #include "UObject/NameTypes.h"
 
+#if WITH_EDITOR
+#include "Delegates/Delegate.h"
+#include "Providers/SubmixAssetProvider.h"
+#endif // WITH_EDITOR
+
 namespace UE::Audio::Insights
 {
 	class FSubmixTraceProvider : public TDeviceDataMapTraceProvider<uint32, TSharedPtr<FSubmixDashboardEntry>>, public TSharedFromThis<FSubmixTraceProvider>
 	{
 	public:
 		FSubmixTraceProvider();
-		virtual ~FSubmixTraceProvider() = default;
+		virtual ~FSubmixTraceProvider();
 
 		static FName GetName_Static();
 		
@@ -20,17 +25,31 @@ namespace UE::Audio::Insights
 
 #if !WITH_EDITOR
 		virtual void InitSessionCachedMessages(TraceServices::IAnalysisSession& InSession) override;
+#else
+		DECLARE_DELEGATE(FOnSubmixAssetsUpdated);
+		FOnSubmixAssetsUpdated OnSubmixAssetAdded;
+		FOnSubmixAssetsUpdated OnSubmixAssetRemoved;
+		FOnSubmixAssetsUpdated OnSubmixAssetListUpdated;
 #endif // !WITH_EDITOR
 
 	private:
 		virtual bool ProcessMessages() override;
 
-#if !WITH_EDITOR
+#if WITH_EDITOR
+		void HandleOnSubmixAssetAdded(const TSharedPtr<FSubmixDashboardEntry>& InSubmixDashboardEntry);
+		void HandleOnSubmixAssetRemoved(const uint32 InSubmixId);
+		void HandleOnSubmixAssetListUpdated(const TArray<TSharedPtr<FSubmixDashboardEntry>>& InSubmixAssetEntries);
+#else
 		virtual void OnTimingViewTimeMarkerChanged(double TimeMarker) override;
 
 		TUniquePtr<FSubmixSessionCachedMessages> SessionCachedMessages;
-#endif // !WITH_EDITOR
+#endif // WITH_EDITOR
 
 		FSubmixMessages TraceMessages;
+
+#if WITH_EDITOR
+		FSubmixAssetProvider SubmixAssetProvider;
+#endif // WITH_EDITOR
+
 	};
 } // namespace UE::Audio::Insights
diff --git a/Engine/Plugins/AudioInsights/Source/AudioInsights/Private/Views/SubmixDashboardViewFactory.cpp b/Engine/Plugins/AudioInsights/Source/AudioInsights/Private/Views/SubmixDashboardViewFactory.cpp
--- a/Engine/Plugins/AudioInsights/Source/AudioInsights/Private/Views/SubmixDashboardViewFactory.cpp
+++ b/Engine/Plugins/AudioInsights/Source/AudioInsights/Private/Views/SubmixDashboardViewFactory.cpp
@@ -6,6 +6,11 @@
 #include "Internationalization/Text.h"
 #include "Providers/SubmixTraceProvider.h"
 
+#if WITH_EDITOR
+#include "Editor.h"
+#include "Subsystems/AssetEditorSubsystem.h"
+#endif // WITH_EDITOR
+
 #define LOCTEXT_NAMESPACE "AudioInsights"
 
 namespace UE::Audio::Insights
@@ -33,6 +38,21 @@ namespace UE::Audio::Insights
 
 		SortByColumn = "Name";
 		SortMode     = EColumnSortMode::Ascending;
+
+#if WITH_EDITOR
+		SubmixProvider->OnSubmixAssetAdded.BindRaw(this, &FSubmixDashboardViewFactory::RequestListRefresh);
+		SubmixProvider->OnSubmixAssetRemoved.BindRaw(this, &FSubmixDashboardViewFactory::RequestListRefresh);
+		SubmixProvider->OnSubmixAssetListUpdated.BindRaw(this, &FSubmixDashboardViewFactory::RequestListRefresh);
+#endif // WITH_EDITOR
+	}
+
+	FSubmixDashboardViewFactory::~FSubmixDashboardViewFactory()
+	{
+#if WITH_EDITOR
+		SubmixProvider->OnSubmixAssetAdded.Unbind();
+		SubmixProvider->OnSubmixAssetRemoved.Unbind();
+		SubmixProvider->OnSubmixAssetListUpdated.Unbind();
+#endif // WITH_EDITOR
 	}
 
 	FName FSubmixDashboardViewFactory::GetName() const
@@ -98,6 +118,22 @@ namespace UE::Audio::Insights
 					.Text(ValueText)
 					.OnDoubleClicked_Lambda([this, InRowData](const FGeometry& MyGeometry, const FPointerEvent& PointerEvent)
 					{
+#if WITH_EDITOR
+						if (GEditor)
+						{
+							const TSharedPtr<IObjectDashboardEntry> ObjectData = StaticCastSharedPtr<IObjectDashboardEntry>(InRowData.ToSharedPtr());
+							if (ObjectData.IsValid())
+							{
+								const TObjectPtr<UObject> Object = ObjectData->GetObject();
+								if (Object && Object->IsAsset())
+								{
+									GEditor->GetEditorSubsystem<UAssetEditorSubsystem>()->OpenEditorForAsset(Object);
+									return FReply::Handled();
+								}
+							}
+						}
+#endif // WITH_EDITOR
+
 						return FReply::Unhandled();
 					})
 				];
@@ -148,9 +184,9 @@ namespace UE::Audio::Insights
 					{
 						LOCTEXT("SubmixDashboard_ActiveDisplayName", "Active"),
 						[](const IDashboardDataViewEntry& InData) { return FText::GetEmpty(); },
-						nullptr, /* GetIconName */
-						false,	 /* bDefaultHidden */
-						0.08f,	 /* FillWidth */
+						nullptr		/* GetIconName */,
+						false,		/* bDefaultHidden */
+						0.08f,		/* FillWidth */
 						EHorizontalAlignment::HAlign_Center
 					}
 				},
@@ -159,9 +195,9 @@ namespace UE::Audio::Insights
 					{
 						LOCTEXT("SubmixDashboard_NameColumnDisplayName", "Name"),
 						[](const IDashboardDataViewEntry& InData) { return CastEntry(InData).GetDisplayName(); },
-						nullptr, /* GetIconName */
-						false,	 /* bDefaultHidden */
-						0.92f,	 /* FillWidth */
+						nullptr		/* GetIconName */,
+						false,		/* bDefaultHidden */
+						0.92f,		/* FillWidth */
 						EHorizontalAlignment::HAlign_Left
 					}
 				}
@@ -224,6 +260,14 @@ namespace UE::Audio::Insights
 			}
 		}
 	}
+
+	void FSubmixDashboardViewFactory::RequestListRefresh()
+	{
+		if (FilteredEntriesListView.IsValid())
+		{
+			FilteredEntriesListView->RequestListRefresh();
+		}
+	}
 } // namespace UE::Audio::Insights
 
 #undef LOCTEXT_NAMESPACE
diff --git a/Engine/Plugins/AudioInsights/Source/AudioInsights/Public/Views/SubmixDashboardViewFactory.h b/Engine/Plugins/AudioInsights/Source/AudioInsights/Public/Views/SubmixDashboardViewFactory.h
--- a/Engine/Plugins/AudioInsights/Source/AudioInsights/Public/Views/SubmixDashboardViewFactory.h
+++ b/Engine/Plugins/AudioInsights/Source/AudioInsights/Public/Views/SubmixDashboardViewFactory.h
@@ -14,7 +14,7 @@ namespace UE::Audio::Insights
 	{
 	public:
 		UE_API FSubmixDashboardViewFactory();
-		UE_API virtual ~FSubmixDashboardViewFactory() = default;
+		UE_API virtual ~FSubmixDashboardViewFactory();
 
 		UE_API virtual FName GetName() const override;
 		UE_API virtual FText GetDisplayName() const override;
@@ -29,6 +29,8 @@ namespace UE::Audio::Insights
 
 		UE_API virtual void SortTable() override;
 
+		UE_API void RequestListRefresh();
+
 		TSharedPtr<FSubmixTraceProvider> SubmixProvider;
 	};
 } // namespace UE::Audio::Insights
