[Audio Insights] Removed IAudioMeterableEntry in favor of using composition over inheritance with FAudioMeterInfo + added SAudioMeterView::SetName

diff --git a/Engine/Plugins/AudioInsights/Source/AudioInsights/Private/Providers/AudioMeterProvider.h b/Engine/Plugins/AudioInsights/Source/AudioInsights/Private/Providers/AudioMeterProvider.h
new file mode 100644
--- /dev/null
+++ b/Engine/Plugins/AudioInsights/Source/AudioInsights/Private/Providers/AudioMeterProvider.h
@@ -0,0 +1,34 @@
+// Copyright Epic Games, Inc. All Rights Reserved.
+#pragma once
+
+#include "Delegates/Delegate.h"
+#include "HAL/Platform.h"
+
+#define UE_API AUDIOINSIGHTS_API
+
+class FText;
+struct FSlateColor;
+
+namespace UE::Audio::Insights
+{
+	struct FAudioMeterInfo
+	{
+		int32 NumChannels = 0;
+		TArray<float> EnvelopeValues;
+	};
+
+	class FAudioMeterProvider
+	{
+	public:
+		DECLARE_DELEGATE_FourParams(FOnAddAudioMeter, const uint32 /*EntryId*/, const TSharedRef<FAudioMeterInfo>& /*AudioMeterInfo*/, const FText& /*Name*/, const FSlateColor& /*NameColor*/);
+		UE_API inline static FOnAddAudioMeter OnAddAudioMeter;
+
+		DECLARE_DELEGATE_OneParam(FOnRemoveAudioMeter, const uint32 /*EntryId*/);
+		UE_API inline static FOnRemoveAudioMeter OnRemoveAudioMeter;
+
+		DECLARE_DELEGATE_ThreeParams(FOnUpdateAudioMeterInfo, const uint32 /*EntryId*/, const TSharedRef<FAudioMeterInfo>& /*AudioMeterInfo*/, const FText& /*Name*/);
+		UE_API inline static FOnUpdateAudioMeterInfo OnUpdateAudioMeterInfo;
+	};
+} // namespace UE::Audio::Insights
+
+#undef UE_API
diff --git a/Engine/Plugins/AudioInsights/Source/AudioInsights/Private/Providers/SubmixAssetProvider.cpp b/Engine/Plugins/AudioInsights/Source/AudioInsights/Private/Providers/SubmixAssetProvider.cpp
--- a/Engine/Plugins/AudioInsights/Source/AudioInsights/Private/Providers/SubmixAssetProvider.cpp
+++ b/Engine/Plugins/AudioInsights/Source/AudioInsights/Private/Providers/SubmixAssetProvider.cpp
@@ -101,7 +101,7 @@ namespace UE::Audio::Insights
 			{
 				if (const ::Audio::FMixerDevice* MixerDevice = static_cast<const ::Audio::FMixerDevice*>(AudioDeviceManager->GetAudioDeviceRaw(AudioDeviceId)))
 				{
-					SubmixDashboardEntry->NumChannels = MixerDevice->GetNumDeviceChannels();
+					SubmixDashboardEntry->AudioMeterInfo->NumChannels = MixerDevice->GetNumDeviceChannels();
 				}
 			}
 
diff --git a/Engine/Plugins/AudioInsights/Source/AudioInsights/Private/Providers/SubmixTraceProvider.cpp b/Engine/Plugins/AudioInsights/Source/AudioInsights/Private/Providers/SubmixTraceProvider.cpp
--- a/Engine/Plugins/AudioInsights/Source/AudioInsights/Private/Providers/SubmixTraceProvider.cpp
+++ b/Engine/Plugins/AudioInsights/Source/AudioInsights/Private/Providers/SubmixTraceProvider.cpp
@@ -190,7 +190,7 @@ namespace UE::Audio::Insights
 			{
 				FSubmixDashboardEntry& EntryRef = *OutEntry->Get();
 				EntryRef.Name = Msg.Name;
-				EntryRef.NumChannels = Msg.NumChannels;
+				EntryRef.AudioMeterInfo->NumChannels = Msg.NumChannels;
 			}
 		});
 #endif // !WITH_EDITOR
@@ -220,7 +220,7 @@ namespace UE::Audio::Insights
 			{
 				FSubmixDashboardEntry& EntryRef = *OutEntry->Get();
 
-				EntryRef.EnvelopeValues = Msg.EnvelopeValues;
+				EntryRef.AudioMeterInfo->EnvelopeValues = Msg.EnvelopeValues;
 				EntryRef.Timestamp = Msg.Timestamp;
 			}
 		});
@@ -365,7 +365,7 @@ namespace UE::Audio::Insights
 				Entry->Timestamp = SubmixLoadedCachedMessage.Timestamp;
 
 				Entry->Name = *SubmixLoadedCachedMessage.Name;
-				Entry->NumChannels = SubmixLoadedCachedMessage.NumChannels;
+				Entry->AudioMeterInfo->NumChannels = SubmixLoadedCachedMessage.NumChannels;
 			});
 		}
 
diff --git a/Engine/Plugins/AudioInsights/Source/AudioInsights/Private/Views/SAudioMeterView.cpp b/Engine/Plugins/AudioInsights/Source/AudioInsights/Private/Views/SAudioMeterView.cpp
--- a/Engine/Plugins/AudioInsights/Source/AudioInsights/Private/Views/SAudioMeterView.cpp
+++ b/Engine/Plugins/AudioInsights/Source/AudioInsights/Private/Views/SAudioMeterView.cpp
@@ -4,6 +4,7 @@
 
 #include "AudioInsightsDataSource.h"
 #include "DSP/Dsp.h"
+#include "Providers/AudioMeterProvider.h"
 #include "SAudioMeterWidget.h"
 #include "Widgets/Layout/SBox.h"
 #include "Widgets/SBoxPanel.h"
@@ -15,35 +16,31 @@ namespace UE::Audio::Insights
 {
 	namespace SAudioMeterViewPrivate
 	{
-		const TAttribute<TArray<FAudioMeterChannelInfo>> CreateMeterChannelInfoAttribute(const TSharedPtr<IAudioMeterableEntry>& InAudioMeterableEntry)
+		const TAttribute<TArray<FAudioMeterChannelInfo>> CreateMeterChannelInfoAttribute(const TSharedRef<FAudioMeterInfo>& InAudioMeterInfo)
 		{
-			return TAttribute<TArray<FAudioMeterChannelInfo>>::Create([InAudioMeterableEntry]()
+			return TAttribute<TArray<FAudioMeterChannelInfo>>::Create([InAudioMeterInfo]()
 				{
 					TArray<FAudioMeterChannelInfo> MeterChannelInfo;
 
-					if (InAudioMeterableEntry.IsValid())
+					const TArray<float>& EnvelopeValues = InAudioMeterInfo->EnvelopeValues;
+
+					if (!EnvelopeValues.IsEmpty())
 					{
-						const TArray<float>& EnvelopeValues = InAudioMeterableEntry->GetEnvelopeValues();
+						MeterChannelInfo.Reserve(EnvelopeValues.Num());
 
-						if (!EnvelopeValues.IsEmpty())
+						for (const float EnvelopeValue : EnvelopeValues)
 						{
-							MeterChannelInfo.Reserve(EnvelopeValues.Num());
-
-							for (const float EnvelopeValue : EnvelopeValues)
-							{
-								MeterChannelInfo.Emplace(::Audio::ConvertToDecibels(EnvelopeValue), -160.0f, -160.0f);
-							}
+							MeterChannelInfo.Emplace(::Audio::ConvertToDecibels(EnvelopeValue), -160.0f, -160.0f);
 						}
-						else
-						{
-							const int32 NumChannels = InAudioMeterableEntry->GetNumChannels();
+					}
+					else
+					{
+						const int32 NumChannels = InAudioMeterInfo->NumChannels;
 							
-							for (int32 Index = 0; Index < NumChannels; ++Index)
-							{
-								MeterChannelInfo.Emplace(-160.0f, -160.0f, -160.0f);
-							}
+						for (int32 Index = 0; Index < NumChannels; ++Index)
+						{
+							MeterChannelInfo.Emplace(-160.0f, -160.0f, -160.0f);
 						}
-
 					}
 
 					return MeterChannelInfo;
@@ -51,10 +48,12 @@ namespace UE::Audio::Insights
 		}
 	}
 
-	void SAudioMeterView::Construct(const SAudioMeterView::FArguments& InArgs, const TSharedPtr<IAudioMeterableEntry>& InAudioMeterableEntry)
+	void SAudioMeterView::Construct(const SAudioMeterView::FArguments& InArgs, const TSharedRef<FAudioMeterInfo>& InAudioMeterInfo)
 	{
 		using namespace SAudioMeterViewPrivate;
 
+		NameText = InArgs._NameText;
+
 		ChildSlot
 		[
 			SNew(SVerticalBox)
@@ -87,7 +86,7 @@ namespace UE::Audio::Insights
 						.MeterClippingColor(FAudioMeterDefaultColorWidgetStyle::GetDefault().MeterClippingColor)
 						.MeterScaleColor(FAudioMeterDefaultColorWidgetStyle::GetDefault().MeterScaleColor)
 						.MeterScaleLabelColor(FAudioMeterDefaultColorWidgetStyle::GetDefault().MeterScaleLabelColor)
-						.MeterChannelInfo(CreateMeterChannelInfoAttribute(InAudioMeterableEntry))
+						.MeterChannelInfo(CreateMeterChannelInfoAttribute(InAudioMeterInfo))
 					]
 				]
 				// Height padding
@@ -101,7 +100,10 @@ namespace UE::Audio::Insights
 				.FillHeight(0.1)
 				[
 					SNew(STextBlock)
-					.Text(InArgs._NameText)
+					.Text_Lambda([this]()
+					{
+						return NameText;
+					})
 					.Justification(ETextJustify::Center)
 					.ColorAndOpacity(InArgs._NameTextColor)
 				]
@@ -115,11 +117,16 @@ namespace UE::Audio::Insights
 		];
 	}
 
-	void SAudioMeterView::SetAudioMeterChannelInfo(const TSharedRef<IAudioMeterableEntry>& InAudioMeterableEntry)
+	void SAudioMeterView::SetAudioMeterChannelInfo(const TSharedRef<FAudioMeterInfo>& InAudioMeterInfo)
 	{
 		using namespace SAudioMeterViewPrivate;
 
-		AudioMeterWidget->SetMeterChannelInfo(CreateMeterChannelInfoAttribute(InAudioMeterableEntry));
+		AudioMeterWidget->SetMeterChannelInfo(CreateMeterChannelInfoAttribute(InAudioMeterInfo));
+	}
+
+	void SAudioMeterView::SetName(const FText& InName)
+	{
+		NameText = InName;
 	}
 } // namespace UE::Audio::Insights
 
diff --git a/Engine/Plugins/AudioInsights/Source/AudioInsights/Public/AudioInsightsDataSource.h b/Engine/Plugins/AudioInsights/Source/AudioInsights/Public/AudioInsightsDataSource.h
--- a/Engine/Plugins/AudioInsights/Source/AudioInsights/Public/AudioInsightsDataSource.h
+++ b/Engine/Plugins/AudioInsights/Source/AudioInsights/Public/AudioInsightsDataSource.h
@@ -31,13 +31,4 @@ namespace UE::Audio::Insights
 		TArray<TSharedPtr<IDashboardDataTreeViewEntry>> Children;
 		bool bIsExpanded = false;
 	};
-
-	class IAudioMeterableEntry
-	{
-	public:
-		virtual ~IAudioMeterableEntry() = default;
-
-		virtual int32 GetNumChannels() const = 0;
-		virtual const TArray<float>& GetEnvelopeValues() const = 0;
-	};
 } // namespace UE::Audio::Insights
diff --git a/Engine/Plugins/AudioInsights/Source/AudioInsights/Public/Messages/SubmixTraceMessages.h b/Engine/Plugins/AudioInsights/Source/AudioInsights/Public/Messages/SubmixTraceMessages.h
--- a/Engine/Plugins/AudioInsights/Source/AudioInsights/Public/Messages/SubmixTraceMessages.h
+++ b/Engine/Plugins/AudioInsights/Source/AudioInsights/Public/Messages/SubmixTraceMessages.h
@@ -1,6 +1,7 @@
 // Copyright Epic Games, Inc. All Rights Reserved.
 #pragma once
 
+#include "Providers/AudioMeterProvider.h"
 #include "Messages/AnalyzerMessageQueue.h"
 #include "Trace/Analyzer.h"
 #include "Views/TableDashboardViewFactory.h"
@@ -113,7 +114,12 @@ namespace UE::Audio::Insights
 
 	struct FSubmixDashboardEntry : public IObjectDashboardEntry
 	{
-		FSubmixDashboardEntry() = default;
+		FSubmixDashboardEntry()
+			: AudioMeterInfo(MakeShared<FAudioMeterInfo>())
+		{
+
+		}
+
 		virtual ~FSubmixDashboardEntry() = default;
 
 		virtual FText GetDisplayName() const override { return FText::FromString(FSoftObjectPath(Name).GetAssetName()); }
@@ -124,9 +130,9 @@ namespace UE::Audio::Insights
 		uint32 SubmixId = INDEX_NONE;
 		double Timestamp = 0.0;
 
+		TSharedRef<FAudioMeterInfo> AudioMeterInfo;
+
 		FString Name;
-		int32 NumChannels = 0;
 		bool bHasActivity = false;
-		TArray<float> EnvelopeValues;
 	};
 } // namespace UE::Audio::Insights
diff --git a/Engine/Plugins/AudioInsights/Source/AudioInsights/Public/Views/SAudioMeterView.h b/Engine/Plugins/AudioInsights/Source/AudioInsights/Public/Views/SAudioMeterView.h
--- a/Engine/Plugins/AudioInsights/Source/AudioInsights/Public/Views/SAudioMeterView.h
+++ b/Engine/Plugins/AudioInsights/Source/AudioInsights/Public/Views/SAudioMeterView.h
@@ -2,13 +2,14 @@
 
 #pragma once
 
+#include "Internationalization/Text.h"
 #include "Widgets/SCompoundWidget.h"
 
 class SAudioMeterWidget;
 
 namespace UE::Audio::Insights
 {
-	class IAudioMeterableEntry;
+	struct FAudioMeterInfo;
 
 	class SAudioMeterView : public SCompoundWidget
 	{
@@ -21,11 +22,13 @@ namespace UE::Audio::Insights
 			SLATE_ARGUMENT(FSlateColor, NameTextColor)
 		SLATE_END_ARGS()
 
-		void Construct(const FArguments& InArgs, const TSharedPtr<IAudioMeterableEntry>& InAudioMeterableEntry);
+		void Construct(const FArguments& InArgs, const TSharedRef<FAudioMeterInfo>& InAudioMeterInfo);
 
-		void SetAudioMeterChannelInfo(const TSharedRef<IAudioMeterableEntry>& InAudioMeterableEntry);
+		void SetAudioMeterChannelInfo(const TSharedRef<FAudioMeterInfo>& InAudioMeterInfo);
+		void SetName(const FText& InName);
 
 	protected:
 		TSharedPtr<SAudioMeterWidget> AudioMeterWidget;
+		FText NameText;
 	};
 } // namespace UE::Audio::Insights
