[Audio Insights] Added Audio Meters dashboard

diff --git a/Engine/Plugins/AudioInsights/Source/AudioInsights/Private/AudioInsightsDashboardFactory.cpp b/Engine/Plugins/AudioInsights/Source/AudioInsights/Private/AudioInsightsDashboardFactory.cpp
--- a/Engine/Plugins/AudioInsights/Source/AudioInsights/Private/AudioInsightsDashboardFactory.cpp
+++ b/Engine/Plugins/AudioInsights/Source/AudioInsights/Private/AudioInsightsDashboardFactory.cpp
@@ -191,6 +191,7 @@ namespace UE::Audio::Insights
 
 		TSharedRef<FTabManager::FStack> LogTabStack = FTabManager::NewStack();
 		TSharedRef<FTabManager::FStack> AnalysisTabStack = FTabManager::NewStack();
+		TSharedRef<FTabManager::FStack> AudioMetersTabStack = FTabManager::NewStack();
 		TSharedRef<FTabManager::FStack> PlotsTabStack = FTabManager::NewStack();
 
 		for (const auto& [FactoryName, Factory] : DashboardViewFactories)
@@ -211,6 +212,12 @@ namespace UE::Audio::Insights
 				}
 				break;
 
+				case EDefaultDashboardTabStack::AudioMeters:
+				{
+					AudioMetersTabStack->AddTab(FactoryName, ETabState::OpenedTab);
+				}
+				break;
+
 				case EDefaultDashboardTabStack::Plots:
 				{
 					PlotsTabStack->AddTab(FactoryName, ETabState::OpenedTab);
@@ -257,8 +264,21 @@ namespace UE::Audio::Insights
 						// Bottom
 						->Split
 						(
-							PlotsTabStack
-							->SetSizeCoefficient(0.42f)
+							// Top
+							FTabManager::NewSplitter()
+							->SetOrientation(Orient_Horizontal)
+							->SetSizeCoefficient(0.42f) // Row width
+							->Split
+							(
+								PlotsTabStack
+								->SetSizeCoefficient(0.5f)
+							)
+							// Bottom
+							->Split
+							(
+								AudioMetersTabStack
+								->SetSizeCoefficient(0.5f)
+							)
 						)
 					)
 				)
diff --git a/Engine/Plugins/AudioInsights/Source/AudioInsights/Private/AudioInsightsModule.cpp b/Engine/Plugins/AudioInsights/Source/AudioInsights/Private/AudioInsightsModule.cpp
--- a/Engine/Plugins/AudioInsights/Source/AudioInsights/Private/AudioInsightsModule.cpp
+++ b/Engine/Plugins/AudioInsights/Source/AudioInsights/Private/AudioInsightsModule.cpp
@@ -16,6 +16,7 @@
 #if !WITH_EDITOR
 #include "AudioInsightsComponent.h"
 #include "Views/AudioEventLogDashboardViewFactory.h"
+#include "Views/AudioMetersPanelDashboardViewFactory.h"
 #include "Views/SoundDashboardViewFactory.h"
 #include "Views/SoundPlotsDashboardViewFactory.h"
 #include "Views/SubmixDashboardViewFactory.h"
@@ -59,6 +60,7 @@ namespace UE::Audio::Insights
 			DashboardFactory->RegisterViewFactory(MakeShared<FVirtualLoopDashboardViewFactory>());
 			DashboardFactory->RegisterViewFactory(PlotsDashboard);
 			DashboardFactory->RegisterViewFactory(MakeShared<FSubmixDashboardViewFactory>());
+			DashboardFactory->RegisterViewFactory(MakeShared<FAudioMetersPanelDashboardViewFactory>());
 
 			PlotsDashboard->InitPlots(SoundsDashboard);
 #endif // !WITH_EDITOR
diff --git a/Engine/Plugins/AudioInsights/Source/AudioInsights/Private/Views/AudioMetersPanelDashboardViewFactory.cpp b/Engine/Plugins/AudioInsights/Source/AudioInsights/Private/Views/AudioMetersPanelDashboardViewFactory.cpp
new file mode 100644
--- /dev/null
+++ b/Engine/Plugins/AudioInsights/Source/AudioInsights/Private/Views/AudioMetersPanelDashboardViewFactory.cpp
@@ -0,0 +1,110 @@
+// Copyright Epic Games, Inc. All Rights Reserved.
+#include "Views/AudioMetersPanelDashboardViewFactory.h"
+
+#include "AudioInsightsStyle.h"
+#include "Providers/AudioMeterProvider.h"
+#include "Views/SAudioMeterView.h"
+#include "Widgets/Layout/SScrollBox.h"
+
+#define LOCTEXT_NAMESPACE "AudioInsights"
+
+namespace UE::Audio::Insights
+{
+	FAudioMetersPanelDashboardViewFactory::FAudioMetersPanelDashboardViewFactory()
+	{
+		FAudioMeterProvider::OnAddAudioMeter.BindRaw(this, &FAudioMetersPanelDashboardViewFactory::HandleOnAddAudioMeter);
+		FAudioMeterProvider::OnRemoveAudioMeter.BindRaw(this, &FAudioMetersPanelDashboardViewFactory::HandleOnRemoveAudioMeter);
+		FAudioMeterProvider::OnUpdateAudioMeterInfo.BindRaw(this, &FAudioMetersPanelDashboardViewFactory::HandleOnUpdateAudioMeterInfo);
+	}
+
+	FAudioMetersPanelDashboardViewFactory::~FAudioMetersPanelDashboardViewFactory()
+	{
+		FAudioMeterProvider::OnAddAudioMeter.Unbind();
+		FAudioMeterProvider::OnRemoveAudioMeter.Unbind();
+		FAudioMeterProvider::OnUpdateAudioMeterInfo.Unbind();
+	}
+
+	FName FAudioMetersPanelDashboardViewFactory::GetName() const
+	{
+		return "AudioMetersPanel";
+	}
+
+	FText FAudioMetersPanelDashboardViewFactory::GetDisplayName() const
+	{
+		return LOCTEXT("AudioDashboard_AudioMetersPanelTab_DisplayName", "Audio Meters");
+	}
+
+	EDefaultDashboardTabStack FAudioMetersPanelDashboardViewFactory::GetDefaultTabStack() const
+	{
+		return EDefaultDashboardTabStack::AudioMeters;
+	}
+
+	FSlateIcon FAudioMetersPanelDashboardViewFactory::GetIcon() const
+	{
+		return FSlateStyle::Get().CreateIcon("AudioInsights.Icon.Submix");
+	}
+
+	TSharedRef<SWidget> FAudioMetersPanelDashboardViewFactory::MakeWidget(TSharedRef<SDockTab> OwnerTab, const FSpawnTabArgs& SpawnTabArgs)
+	{
+		if (!AudioMeterWidgetsScrollBox.IsValid())
+		{
+			SAssignNew(AudioMeterWidgetsScrollBox, SScrollBox)
+			.Orientation(Orient_Horizontal)
+			+ SScrollBox::Slot()
+			[
+				SAssignNew(AudioMeterWidgetsContainer, SHorizontalBox)
+			];
+		}
+
+		return AudioMeterWidgetsScrollBox.ToSharedRef();
+	}
+
+	void FAudioMetersPanelDashboardViewFactory::HandleOnAddAudioMeter(const uint32 InEntryId, const TSharedRef<FAudioMeterInfo>& InAudioMeterInfo, const FText& InName, const FSlateColor& InNameColor)
+	{
+		const TSharedRef<SAudioMeterView>& AudioMeterView = AudioMeterViews.Emplace(InEntryId, SNew(SAudioMeterView, InAudioMeterInfo)
+			.NameText(InName)
+			.NameTextColor(InNameColor));
+
+		AudioMeterWidgetsContainer->AddSlot()
+		.AutoWidth()
+		.Padding(10.0f, 0.0f, 10.0f, 0.0f)
+		[
+			AudioMeterView
+		];
+
+		if (AudioMeterWidgetsScrollBox.IsValid())
+		{
+			AudioMeterWidgetsScrollBox->Invalidate(EInvalidateWidget::Layout);
+		}
+	}
+
+	void FAudioMetersPanelDashboardViewFactory::HandleOnRemoveAudioMeter(const uint32 InEntryId)
+	{
+		if (TSharedRef<SAudioMeterView>* FoundAudioMeterView = AudioMeterViews.Find(InEntryId))
+		{
+			AudioMeterWidgetsContainer->RemoveSlot(*FoundAudioMeterView);
+			AudioMeterViews.Remove(InEntryId);
+		}
+
+		if (AudioMeterWidgetsScrollBox.IsValid())
+		{
+			AudioMeterWidgetsScrollBox->Invalidate(EInvalidateWidget::Layout);
+		}
+	}
+
+	void FAudioMetersPanelDashboardViewFactory::HandleOnUpdateAudioMeterInfo(const uint32 InEntryId, const TSharedRef<FAudioMeterInfo>& InAudioMeterInfo, const FText& InName)
+	{
+		if (TSharedRef<SAudioMeterView>* FoundAudioMeterView = AudioMeterViews.Find(InEntryId))
+		{
+			(*FoundAudioMeterView)->SetAudioMeterChannelInfo(InAudioMeterInfo);
+			(*FoundAudioMeterView)->SetName(InName);
+
+			if (AudioMeterWidgetsScrollBox.IsValid())
+			{
+				AudioMeterWidgetsScrollBox->Invalidate(EInvalidateWidget::Layout);
+			}
+		}
+	}
+} // namespace UE::Audio::Insights
+
+#undef LOCTEXT_NAMESPACE
diff --git a/Engine/Plugins/AudioInsights/Source/AudioInsights/Public/Views/AudioMetersPanelDashboardViewFactory.h b/Engine/Plugins/AudioInsights/Source/AudioInsights/Public/Views/AudioMetersPanelDashboardViewFactory.h
new file mode 100644
--- /dev/null
+++ b/Engine/Plugins/AudioInsights/Source/AudioInsights/Public/Views/AudioMetersPanelDashboardViewFactory.h
@@ -0,0 +1,40 @@
+// Copyright Epic Games, Inc. All Rights Reserved.
+#pragma once
+
+#include "Views/DashboardViewFactory.h"
+
+#define UE_API AUDIOINSIGHTS_API
+
+class SHorizontalBox;
+class SScrollBox;
+class SWidget;
+
+namespace UE::Audio::Insights
+{
+	class SAudioMeterView;
+	struct FAudioMeterInfo;
+
+	class FAudioMetersPanelDashboardViewFactory : public IDashboardViewFactory, public TSharedFromThis<FAudioMetersPanelDashboardViewFactory>
+	{
+	public:
+		UE_API FAudioMetersPanelDashboardViewFactory();
+		UE_API virtual ~FAudioMetersPanelDashboardViewFactory();
+
+		UE_API virtual FName GetName() const override;
+		UE_API virtual FText GetDisplayName() const override;
+		UE_API virtual EDefaultDashboardTabStack GetDefaultTabStack() const override;
+		UE_API virtual FSlateIcon GetIcon() const override;
+		UE_API virtual TSharedRef<SWidget> MakeWidget(TSharedRef<SDockTab> OwnerTab, const FSpawnTabArgs& SpawnTabArgs) override;
+
+	private:
+		void HandleOnAddAudioMeter(const uint32 InEntryId, const TSharedRef<FAudioMeterInfo>& InAudioMeterInfo, const FText& InName, const FSlateColor& InNameColor);
+		void HandleOnRemoveAudioMeter(const uint32 InEntryId);
+		void HandleOnUpdateAudioMeterInfo(const uint32 InEntryId, const TSharedRef<FAudioMeterInfo>& InAudioMeterInfo, const FText& InName);
+
+		TSharedPtr<SScrollBox> AudioMeterWidgetsScrollBox;
+		TSharedPtr<SHorizontalBox> AudioMeterWidgetsContainer;
+		TMap<uint32, TSharedRef<SAudioMeterView>> AudioMeterViews;
+	};
+} // namespace UE::Audio::Insights
+
+#undef UE_API
