[Audio Insights] FSubmixDashboardViewFactory: added checkboxes and logic to enable/disable audio meters in the Audio Meters dashboard

diff --git a/Engine/Plugins/AudioInsights/Source/AudioInsights/Private/Views/SubmixDashboardViewFactory.cpp b/Engine/Plugins/AudioInsights/Source/AudioInsights/Private/Views/SubmixDashboardViewFactory.cpp
--- a/Engine/Plugins/AudioInsights/Source/AudioInsights/Private/Views/SubmixDashboardViewFactory.cpp
+++ b/Engine/Plugins/AudioInsights/Source/AudioInsights/Private/Views/SubmixDashboardViewFactory.cpp
@@ -4,7 +4,9 @@
 #include "AudioInsightsStyle.h"
 #include "IAudioInsightsModule.h"
 #include "Internationalization/Text.h"
+#include "Providers/AudioMeterProvider.h"
 #include "Providers/SubmixTraceProvider.h"
+#include "Widgets/Input/SCheckBox.h"
 
 #if WITH_EDITOR
 #include "Editor.h"
@@ -21,6 +23,22 @@ namespace UE::Audio::Insights
 		{
 			return static_cast<const FSubmixDashboardEntry&>(InData);
 		};
+
+		void SendSubmixEnvelopeFollowerCVar(const bool bInIsSubmixChecked, const FString& InSubmixName)
+		{
+#if WITH_EDITOR
+			const IAudioInsightsModule& AudioInsightsModule = IAudioInsightsModule::GetEditorChecked();
+#else
+			const IAudioInsightsModule& AudioInsightsModule = IAudioInsightsModule::GetChecked();
+#endif // WITH_EDITOR
+
+			const ::Audio::FDeviceId AudioDeviceId = AudioInsightsModule.GetDeviceId();
+
+			const FString SubmixCommandStr = FString::Printf(TEXT("au.%sSubmixEnvelopeFollower_AD%d %s"), bInIsSubmixChecked ? TEXT("Start") : TEXT("Stop"), AudioDeviceId, *InSubmixName);
+			
+			IAudioInsightsTraceModule& AudioInsightsTraceModule = IAudioInsightsModule::GetChecked().GetTraceModule();
+			AudioInsightsTraceModule.ExecuteConsoleCommand(SubmixCommandStr);
+		}
 	}
 
 	FSubmixDashboardViewFactory::FSubmixDashboardViewFactory()
@@ -109,10 +127,61 @@ namespace UE::Audio::Insights
 				return SNullWidget::NullWidget;
 			}
 
+			const FSubmixDashboardEntry& SubmixDashboardEntry = CastEntry(*InRowData);
+			const uint32 SubmixId = SubmixDashboardEntry.SubmixId;
+
+			const bool bIsSubmixChecked = SubmixCheckboxCheckedStates.FindOrAdd(SubmixId);
+
+			if (bIsSubmixChecked)
+			{
+				SendSubmixEnvelopeFollowerCVar(bIsSubmixChecked, SubmixDashboardEntry.GetDisplayName().ToString());
+				FAudioMeterProvider::OnUpdateAudioMeterInfo.ExecuteIfBound(SubmixId, SubmixDashboardEntry.AudioMeterInfo, SubmixDashboardEntry.GetDisplayName());
+			}
+
 			return SNew(SHorizontalBox)
 				.Clipping(EWidgetClipping::ClipToBounds)
 				+ SHorizontalBox::Slot()
 				.AutoWidth()
+				[
+					SNew(SCheckBox)
+					.IsChecked_Lambda([this, SubmixId]()
+					{
+						const bool* FoundSubmixCheckedState = SubmixCheckboxCheckedStates.Find(SubmixId);
+						const bool bIsSubmixChecked = FoundSubmixCheckedState && *FoundSubmixCheckedState == true;
+
+						return bIsSubmixChecked ? ECheckBoxState::Checked : ECheckBoxState::Unchecked;
+					})
+					.OnCheckStateChanged_Lambda([this, InRowData](ECheckBoxState NewState)
+					{
+						const bool bIsChecked = NewState == ECheckBoxState::Checked;
+
+						const FSubmixDashboardEntry& SubmixDashboardEntry = CastEntry(*InRowData);
+
+						if (bool* FoundSubmixCheckedState = SubmixCheckboxCheckedStates.Find(SubmixDashboardEntry.SubmixId))
+						{
+							*FoundSubmixCheckedState = bIsChecked;
+						}
+
+						SendSubmixEnvelopeFollowerCVar(bIsChecked, SubmixDashboardEntry.GetDisplayName().ToString());
+
+						if (bIsChecked)
+						{
+							FAudioMeterProvider::OnAddAudioMeter.ExecuteIfBound(SubmixDashboardEntry.SubmixId, SubmixDashboardEntry.AudioMeterInfo, SubmixDashboardEntry.GetDisplayName(), FSlateColor(FColor(30, 240, 90)));
+						}
+						else
+						{
+							FAudioMeterProvider::OnRemoveAudioMeter.ExecuteIfBound(SubmixDashboardEntry.SubmixId);
+						}
+					})
+				]
+				+ SHorizontalBox::Slot()
+				.AutoWidth()
+				[
+					SNew(SBox)
+					.MinDesiredWidth(5.0f)
+				]
+				+ SHorizontalBox::Slot()
+				.AutoWidth()
 				[
 					SNew(STextBlock)
 					.Text(ValueText)
@@ -168,6 +237,20 @@ namespace UE::Audio::Insights
 			}
 		}
 
+		// Update audio meters (if any has been checked previously)
+		for (const TSharedPtr<IDashboardDataViewEntry>& Entry : DataViewEntries)
+		{
+			const FSubmixDashboardEntry& SubmixDashboardEntry = CastEntry(*Entry);
+
+			if (const bool* bFoundSubmixCheckboxCheckedState = SubmixCheckboxCheckedStates.Find(SubmixDashboardEntry.SubmixId))
+			{
+				if (*bFoundSubmixCheckboxCheckedState)
+				{
+					FAudioMeterProvider::OnUpdateAudioMeterInfo.ExecuteIfBound(SubmixDashboardEntry.SubmixId, SubmixDashboardEntry.AudioMeterInfo, SubmixDashboardEntry.GetDisplayName());
+				}
+			}
+		}
+
 		return DashboardWidget->AsShared();
 	}
 
diff --git a/Engine/Plugins/AudioInsights/Source/AudioInsights/Public/Views/SubmixDashboardViewFactory.h b/Engine/Plugins/AudioInsights/Source/AudioInsights/Public/Views/SubmixDashboardViewFactory.h
--- a/Engine/Plugins/AudioInsights/Source/AudioInsights/Public/Views/SubmixDashboardViewFactory.h
+++ b/Engine/Plugins/AudioInsights/Source/AudioInsights/Public/Views/SubmixDashboardViewFactory.h
@@ -32,6 +32,8 @@ namespace UE::Audio::Insights
 		UE_API void RequestListRefresh();
 
 		TSharedPtr<FSubmixTraceProvider> SubmixProvider;
+
+		TMap<uint32, bool> SubmixCheckboxCheckedStates;
 	};
 } // namespace UE::Audio::Insights
 
