[Audio Insights] Create settings menu for Event Log

diff --git a/Engine/Plugins/AudioInsights/Source/AudioInsights/Private/Settings/AudioEventLogSettings.cpp b/Engine/Plugins/AudioInsights/Source/AudioInsights/Private/Settings/AudioEventLogSettings.cpp
new file mode 100644
--- /dev/null
+++ b/Engine/Plugins/AudioInsights/Source/AudioInsights/Private/Settings/AudioEventLogSettings.cpp
@@ -0,0 +1,13 @@
+// Copyright Epic Games, Inc. All Rights Reserved.
+#include "Settings/AudioEventLogSettings.h"
+
+#include UE_INLINE_GENERATED_CPP_BY_NAME(AudioEventLogSettings)
+
+#if WITH_EDITOR
+
+FAudioEventLogSettings::FOnReadEventLogSettings FAudioEventLogSettings::OnReadSettings;
+FAudioEventLogSettings::FOnWriteEventLogSettings FAudioEventLogSettings::OnWriteSettings;
+FAudioEventLogSettings::FOnRequestReadEventLogSettings FAudioEventLogSettings::OnRequestReadSettings;
+FAudioEventLogSettings::FOnRequestWriteEventLogSettings FAudioEventLogSettings::OnRequestWriteSettings;
+
+#endif // WITH_EDITOR
diff --git a/Engine/Plugins/AudioInsights/Source/AudioInsights/Private/Views/AudioEventLogDashboardViewFactory.cpp b/Engine/Plugins/AudioInsights/Source/AudioInsights/Private/Views/AudioEventLogDashboardViewFactory.cpp
--- a/Engine/Plugins/AudioInsights/Source/AudioInsights/Private/Views/AudioEventLogDashboardViewFactory.cpp
+++ b/Engine/Plugins/AudioInsights/Source/AudioInsights/Private/Views/AudioEventLogDashboardViewFactory.cpp
@@ -1,6 +1,7 @@
 // Copyright Epic Games, Inc. All Rights Reserved.
 #include "Views/AudioEventLogDashboardViewFactory.h"
 
+#include "Algo/AnyOf.h"
 #include "AudioInsightsModule.h"
 #include "AudioInsightsStyle.h"
 #include "AudioMixerTrace.h"
@@ -118,6 +119,21 @@ namespace UE::Audio::Insights
 
 		SortByColumn = TimestampColumnID;
 		SortMode = EColumnSortMode::Ascending;
+
+		CommandList = MakeShared<FUICommandList>();
+
+#if WITH_EDITOR
+		FAudioEventLogSettings::OnReadSettings.AddRaw(this, &FAudioEventLogDashboardViewFactory::OnReadEditorSettings);
+		FAudioEventLogSettings::OnWriteSettings.AddRaw(this, &FAudioEventLogDashboardViewFactory::OnWriteEditorSettings);
+#endif // WITH_EDITOR
+	}
+
+	FAudioEventLogDashboardViewFactory::~FAudioEventLogDashboardViewFactory()
+	{
+#if WITH_EDITOR
+		FAudioEventLogSettings::OnReadSettings.RemoveAll(this);
+		FAudioEventLogSettings::OnWriteSettings.RemoveAll(this);
+#endif // WITH_EDITOR
 	}
 
 	FName FAudioEventLogDashboardViewFactory::GetName() const
@@ -142,10 +158,12 @@ namespace UE::Audio::Insights
 
 	TSharedRef<SWidget> FAudioEventLogDashboardViewFactory::MakeWidget(TSharedRef<SDockTab> OwnerTab, const FSpawnTabArgs& SpawnTabArgs)
 	{
-#if !WITH_EDITOR
+#if WITH_EDITOR
+		FAudioEventLogSettings::OnRequestReadSettings.Broadcast();
+#else
 		FAudioInsightsModule& AudioInsightsModule = FAudioInsightsModule::GetChecked();
 		AudioInsightsModule.GetTimingViewExtender().OnTimingViewTimeMarkerChanged.AddSP(this, &FAudioEventLogDashboardViewFactory::OnTimingViewTimeMarkerChanged);
-#endif // !WITH_EDITOR
+#endif // WITH_EDITOR
 
 		if (!DashboardWidget.IsValid())
 		{
@@ -211,6 +229,13 @@ namespace UE::Audio::Insights
 
 		auto CreateColumnData = [this]()
 		{
+			FAudioEventLogVisibleColumns VisibleColumns;
+			if (VisibleColumnSettings.IsSet())
+			{
+				// Use the values loaded from the settings:
+				VisibleColumns = VisibleColumnSettings.GetValue();
+			}
+
 			return TMap<FName, FTraceTableDashboardViewFactory::FColumnData>
 			{
 				{
@@ -219,7 +244,7 @@ namespace UE::Audio::Insights
 						LOCTEXT("AudioEventLog_MessageIDColumnDisplayName", "ID"),
 						[](const IDashboardDataViewEntry& InData) { return FText::AsNumber(AudioEventLogPrivate::CastEntry(InData).MessageID); },
 						nullptr /* GetIconName */,
-						true /* bDefaultHidden */,
+						!VisibleColumns.GetIsVisible(MessageIDColumnID) /* bDefaultHidden */,
 						0.1f /* FillWidth */
 					}
 				},
@@ -229,7 +254,7 @@ namespace UE::Audio::Insights
 						LOCTEXT("AudioEventLog_TimestampColumnDisplayName", "Timestamp"),
 						[](const IDashboardDataViewEntry& InData) { return FText::AsNumber(AudioEventLogPrivate::CastEntry(InData).Timestamp, FSlateStyle::Get().GetTimeFormat()); },
 						nullptr /* GetIconName */,
-						false /* bDefaultHidden */,
+						!VisibleColumns.GetIsVisible(TimestampColumnID) /* bDefaultHidden */,
 						0.15f /* FillWidth */
 					}
 				},
@@ -248,7 +273,7 @@ namespace UE::Audio::Insights
 							return FText::AsNumber(AudioEventLogPrivate::CastEntry(InData).PlayOrder);
 						},
 						nullptr /* GetIconName */,
-						false /* bDefaultHidden */,
+						!VisibleColumns.GetIsVisible(PlayOrderColumnID) /* bDefaultHidden */,
 						0.15f /* FillWidth */
 					}
 				},
@@ -258,7 +283,7 @@ namespace UE::Audio::Insights
 						LOCTEXT("AudioEventLog_EventColumnDisplayName", "Event"),
 						[](const IDashboardDataViewEntry& InData) { return FText::FromString(AudioEventLogPrivate::CastEntry(InData).EventName); },
 						nullptr /* GetIconName */,
-						false /* bDefaultHidden */,
+						!VisibleColumns.GetIsVisible(EventColumnID) /* bDefaultHidden */,
 						0.2f  /* FillWidth */
 					}
 				},
@@ -290,7 +315,7 @@ namespace UE::Audio::Insights
 
 							return NAME_None;
 						},
-						false /* bDefaultHidden */,
+						!VisibleColumns.GetIsVisible(AssetNameColumnID) /* bDefaultHidden */,
 						0.3f  /* FillWidth */
 					}
 				},
@@ -303,7 +328,7 @@ namespace UE::Audio::Insights
 						{
 							return CastEntry(InData).ActorIconName;
 						},
-						false /* bDefaultHidden */,
+						!VisibleColumns.GetIsVisible(ActorColumnID) /* bDefaultHidden */,
 						0.3f  /* FillWidth */
 					}
 				},
@@ -313,7 +338,7 @@ namespace UE::Audio::Insights
 						LOCTEXT("AudioEventLog_CategoryColumnDisplayName", "Category"),
 						[](const IDashboardDataViewEntry& InData) { return FText::FromString(AudioEventLogPrivate::CastEntry(InData).CategoryName); },
 						nullptr /* GetIconName */,
-						false /* bDefaultHidden */,
+						!VisibleColumns.GetIsVisible(CategoryColumnID) /* bDefaultHidden */,
 						0.3f  /* FillWidth */
 					}
 				}
@@ -355,6 +380,22 @@ namespace UE::Audio::Insights
 		return ColumnId == TimestampColumnID;
 	}
 
+	void FAudioEventLogDashboardViewFactory::OnHiddenColumnsListChanged()
+	{
+		// Update and save settings to reload it on the next Editor sessions
+		if (VisibleColumnSettings && HeaderRowWidget)
+		{
+			for (const SHeaderRow::FColumn& Column : HeaderRowWidget->GetColumns())
+			{
+				VisibleColumnSettings->SetIsVisible(Column.ColumnId, Column.bIsVisible);
+			}
+
+#if WITH_EDITOR
+			FAudioEventLogSettings::OnRequestWriteSettings.Broadcast();
+#endif // WITH_EDITOR
+		}
+	}
+
 	TSharedPtr<SWidget> FAudioEventLogDashboardViewFactory::CreateFilterBarButtonWidget()
 	{
 		if (!FilterBarButton.IsValid())
@@ -397,6 +438,31 @@ namespace UE::Audio::Insights
 		return FilterBar.ToSharedRef();
 	}
 
+	TSharedRef<SWidget> FAudioEventLogDashboardViewFactory::OnGetSettingsMenuContent()
+	{
+		FMenuBuilder MenuBuilder(false /*bShouldCloseWindowAfterMenuSelection*/, CommandList);
+
+		MenuBuilder.BeginSection("AudioEventLogSettingActions", LOCTEXT("AudioEventLog_Settings_HeaderText", "Settings"));
+		{
+			// Visible columns
+			MenuBuilder.AddSubMenu
+			(
+				LOCTEXT("AudioEventLog_Settings_VisibleColumns", "Visible Columns"),
+				LOCTEXT("AudioEventLog_Settings_VisibleColumnsTooltip", "Show/hide columns"),
+				FNewMenuDelegate::CreateRaw(this, &FAudioEventLogDashboardViewFactory::BuildVisibleColumnsMenuContent)
+			);
+		}
+
+		MenuBuilder.EndSection();
+
+		return MenuBuilder.MakeWidget();
+	}
+
+	FText FAudioEventLogDashboardViewFactory::OnGetSettingsMenuToolTip()
+	{
+		return LOCTEXT("AudioEventLog_Settings_TooltipText", "Event log settings");
+	}
+
 	void FAudioEventLogDashboardViewFactory::OnSelectionChanged(TSharedPtr<IDashboardDataViewEntry> SelectedItem, ESelectInfo::Type SelectInfo)
 	{
 		using namespace AudioEventLogPrivate;
@@ -501,6 +567,108 @@ namespace UE::Audio::Insights
 		return false;
 	}
 
+	void FAudioEventLogDashboardViewFactory::BuildVisibleColumnsMenuContent(FMenuBuilder& MenuBuilder)
+	{
+		MenuBuilder.AddMenuEntry
+		(
+			TAttribute<FText>::CreateSP(this, &FAudioEventLogDashboardViewFactory::GetToggleAllColumnsVisibilityText),
+			FText::GetEmpty(),
+			FSlateIcon(),
+			FUIAction
+			(
+				FExecuteAction::CreateSP(this, &FAudioEventLogDashboardViewFactory::ToggleAllColumnsVisibility),
+				FCanExecuteAction::CreateLambda([]() { return true; }),
+				FGetActionCheckState::CreateSP(this, &FAudioEventLogDashboardViewFactory::GetToggleAllColumnsVisibilityCheckState)
+			),
+			NAME_None,
+			EUserInterfaceActionType::ToggleButton
+		);
+
+		for (const SHeaderRow::FColumn& Column : HeaderRowWidget->GetColumns())
+		{
+			MenuBuilder.AddMenuEntry
+			(
+				Column.DefaultText,
+				FText::GetEmpty(),
+				FSlateIcon(),
+				FUIAction
+				(
+					FExecuteAction::CreateSP(this, &FAudioEventLogDashboardViewFactory::ToggleColumnVisibility, Column.ColumnId),
+					FCanExecuteAction::CreateLambda([]() { return true; }),
+					FGetActionCheckState::CreateSP(this, &FAudioEventLogDashboardViewFactory::GetIsColumnVisibleCheckedState, Column.ColumnId)
+				),
+				NAME_None,
+				EUserInterfaceActionType::ToggleButton
+			);
+		}
+	}
+
+	void FAudioEventLogDashboardViewFactory::ToggleColumnVisibility(FName ColumnId)
+	{
+		const bool bIsColumnVisible = HeaderRowWidget->IsColumnVisible(ColumnId);
+		HeaderRowWidget->SetShowGeneratedColumn(ColumnId, !bIsColumnVisible);
+	}
+
+	ECheckBoxState FAudioEventLogDashboardViewFactory::GetIsColumnVisibleCheckedState(FName ColumnId) const
+	{
+		return HeaderRowWidget->IsColumnVisible(ColumnId) ? ECheckBoxState::Checked : ECheckBoxState::Unchecked;
+	}
+
+	void FAudioEventLogDashboardViewFactory::ToggleAllColumnsVisibility()
+	{
+		const ECheckBoxState CurrentState = GetToggleAllColumnsVisibilityCheckState();
+		const bool bVisibleIsDesired = CurrentState == ECheckBoxState::Unchecked || CurrentState == ECheckBoxState::Undetermined;
+
+		// Null our VisibleColumnSettings TOptional while updating HeaderRowWidget so that we ignore OnHiddenColumnsListChanged callbacks:
+		VisibleColumnSettings.Reset();
+
+		for (const SHeaderRow::FColumn& Column : HeaderRowWidget->GetColumns())
+		{
+			if (Column.bIsVisible != bVisibleIsDesired)
+			{
+				HeaderRowWidget->SetShowGeneratedColumn(Column.ColumnId, bVisibleIsDesired);
+			}
+		}
+
+		// Reinitialize the VisibleColumnSettings and perform a single OnHiddenColumnsListChanged call to update it and persist the settings:
+		VisibleColumnSettings.Emplace();
+		OnHiddenColumnsListChanged();
+	}
+
+	ECheckBoxState FAudioEventLogDashboardViewFactory::GetToggleAllColumnsVisibilityCheckState() const
+	{
+		const TIndirectArray<SHeaderRow::FColumn>& Columns = HeaderRowWidget->GetColumns();
+		const bool bAnyColumnsVisible = Algo::AnyOf(Columns, [](const SHeaderRow::FColumn& Column) { return Column.bIsVisible; });
+		const bool bAnyColumnsInvisible = Algo::AnyOf(Columns, [](const SHeaderRow::FColumn& Column) { return !Column.bIsVisible; });
+
+		if (bAnyColumnsVisible && bAnyColumnsInvisible)
+		{
+			return ECheckBoxState::Undetermined;
+		}
+		else if (bAnyColumnsVisible)
+		{
+			return ECheckBoxState::Checked;
+		}
+		else
+		{
+			return ECheckBoxState::Unchecked;
+		}
+	}
+
+	FText FAudioEventLogDashboardViewFactory::GetToggleAllColumnsVisibilityText() const
+	{
+		switch (GetToggleAllColumnsVisibilityCheckState())
+		{
+			case ECheckBoxState::Unchecked:
+			case ECheckBoxState::Undetermined:
+				return LOCTEXT("AudioEventLog_VisibleColumns_SelectAll", "Select: All");
+			case ECheckBoxState::Checked:
+				return LOCTEXT("AudioEventLog_VisibleColumns_SelectNone", "Select: None");
+			default:
+				return FText::GetEmpty();
+		}
+	}
+
 	void FAudioEventLogDashboardViewFactory::SortByPredicate(TFunctionRef<bool(const FAudioEventLogDashboardEntry&, const FAudioEventLogDashboardEntry&)> Predicate)
 	{
 		using namespace AudioEventLogPrivate;
@@ -526,7 +694,38 @@ namespace UE::Audio::Insights
 		}
 	}
 
-	#if !WITH_EDITOR
+#if WITH_EDITOR
+	void FAudioEventLogDashboardViewFactory::OnReadEditorSettings(const FAudioEventLogSettings& InSettings)
+	{
+		// If HeaderRowWidget is already created, update it with the desired visible columns:
+		if (HeaderRowWidget)
+		{
+			// Null our VisibleColumnSettings TOptional while updating HeaderRowWidget so that we ignore OnHiddenColumnsListChanged callbacks:
+			VisibleColumnSettings.Reset();
+
+			for (const SHeaderRow::FColumn& Column : HeaderRowWidget->GetColumns())
+			{
+				const bool bShowColumn = InSettings.VisibleColumns.GetIsVisible(Column.ColumnId);
+				if (Column.bIsVisible != bShowColumn)
+				{
+					HeaderRowWidget->SetShowGeneratedColumn(Column.ColumnId, bShowColumn);
+				}
+			}
+		}
+
+		VisibleColumnSettings = InSettings.VisibleColumns;
+	}
+
+	void FAudioEventLogDashboardViewFactory::OnWriteEditorSettings(FAudioEventLogSettings& OutSettings)
+	{
+		if (VisibleColumnSettings)
+		{
+			OutSettings.VisibleColumns = VisibleColumnSettings.GetValue();
+		}
+	}
+#endif // WITH_EDITOR
+
+#if !WITH_EDITOR
 	void FAudioEventLogDashboardViewFactory::OnTimingViewTimeMarkerChanged(double InTimeMarker)
 	{
 		using namespace AudioEventLogPrivate;
@@ -573,7 +772,7 @@ namespace UE::Audio::Insights
 			}
 		}
 	}
-	#endif // !WITH_EDITOR
+#endif // !WITH_EDITOR
 
 } // namespace UE::Audio::Insights
 
diff --git a/Engine/Plugins/AudioInsights/Source/AudioInsights/Private/Views/TableDashboardViewFactory.cpp b/Engine/Plugins/AudioInsights/Source/AudioInsights/Private/Views/TableDashboardViewFactory.cpp
--- a/Engine/Plugins/AudioInsights/Source/AudioInsights/Private/Views/TableDashboardViewFactory.cpp
+++ b/Engine/Plugins/AudioInsights/Source/AudioInsights/Private/Views/TableDashboardViewFactory.cpp
@@ -7,6 +7,7 @@
 #include "Containers/Array.h"
 #include "Framework/Commands/UICommandList.h"
 #include "Templates/SharedPointer.h"
+#include "Widgets/Input/SComboButton.h"
 #include "Widgets/Layout/SScrollBox.h"
 #include "Widgets/SBoxPanel.h"
 
@@ -115,6 +116,36 @@ namespace UE::Audio::Insights
 		RequestSort();
 	}
 
+	TSharedRef<SWidget> FTraceTableDashboardViewFactory::CreateSettingsButtonWidget()
+	{
+		// Override OnGetSettingsMenuContent in derived class to enable the settings menu widget
+		if (OnGetSettingsMenuContent() == SNullWidget::NullWidget)
+		{
+			return SNullWidget::NullWidget;
+		}
+
+		return SNew(SComboButton)
+			.ComboButtonStyle(&FAppStyle::Get().GetWidgetStyle<FComboButtonStyle>("SimpleComboButton"))
+			.OnGetMenuContent(this, &FTraceTableDashboardViewFactory::OnGetSettingsMenuContent)
+			.MenuPlacement(EMenuPlacement::MenuPlacement_ComboBoxRight)
+			.HasDownArrow(false)
+			.ToolTipText(OnGetSettingsMenuToolTip())
+			.ButtonContent()
+			[
+				SNew(SHorizontalBox)
+
+				+ SHorizontalBox::Slot()
+				.AutoWidth()
+				.VAlign(VAlign_Center)
+				.Padding(5.0, 0.0f)
+				[
+					SNew(SImage)
+					.ColorAndOpacity(FSlateColor::UseForeground())
+					.Image(FSlateStyle::Get().GetBrush("AudioInsights.Icon.Settings"))
+				]
+			];
+	}
+
 	TSharedPtr<SWidget> FTraceTableDashboardViewFactory::OnConstructContextMenu()
 	{
 		return SNullWidget::NullWidget;
@@ -138,7 +169,8 @@ namespace UE::Audio::Insights
 			[](const TPair<FName, FColumnData>& ColumnInfo) { return ColumnInfo.Key; });
 
 		SAssignNew(HeaderRowWidget, SHeaderRow)
-		.CanSelectGeneratedColumn(true); // Allows for showing/hiding columns
+		.CanSelectGeneratedColumn(true) // Allows for showing/hiding columns
+		.OnHiddenColumnsListChanged(this, &FTraceTableDashboardViewFactory::OnHiddenColumnsListChanged);
 
 		// This only works if header row columns are added with slots and not programmatically
 		// check in SHeaderRow::Construct: for ( FColumn* const Column : InArgs.Slots ) for more info
@@ -229,6 +261,18 @@ namespace UE::Audio::Insights
 				[
 					CreateFilterBarWidget()
 				]
+
+				// Empty Spacing
+				+ SHorizontalBox::Slot()
+				.FillWidth(1.0f)
+
+				// Optional : Settings button
+				+ SHorizontalBox::Slot()
+				.HAlign(HAlign_Right)
+				.AutoWidth()
+				[
+					CreateSettingsButtonWidget()
+				]
 			]
 			+ SVerticalBox::Slot()
 			.FillHeight(1.0f)
diff --git a/Engine/Plugins/AudioInsights/Source/AudioInsights/Public/Settings/AudioEventLogSettings.h b/Engine/Plugins/AudioInsights/Source/AudioInsights/Public/Settings/AudioEventLogSettings.h
new file mode 100644
--- /dev/null
+++ b/Engine/Plugins/AudioInsights/Source/AudioInsights/Public/Settings/AudioEventLogSettings.h
@@ -0,0 +1,93 @@
+// Copyright Epic Games, Inc. All Rights Reserved.
+
+#pragma once
+
+#include "Delegates/DelegateCombinations.h"
+#include "Engine/DeveloperSettings.h"
+
+#include "AudioEventLogSettings.generated.h"
+
+#define LOCTEXT_NAMESPACE "AudioInsights"
+
+/** FAudioEventLogVisibleColumns
+ *
+ * Control which columns are visible in the Audio Event Log tab.
+ */
+USTRUCT()
+struct FAudioEventLogVisibleColumns
+{
+	GENERATED_BODY()
+
+	UPROPERTY(EditAnywhere, config, Category = EventLog)
+	bool bMessageID = false;
+
+	UPROPERTY(EditAnywhere, config, Category = EventLog)
+	bool bTimestamp = true;
+
+	UPROPERTY(EditAnywhere, config, Category = EventLog)
+	bool bPlayOrder = true;
+
+	UPROPERTY(EditAnywhere, config, Category = EventLog)
+	bool bEvent = true;
+
+	UPROPERTY(EditAnywhere, config, Category = EventLog)
+	bool bAsset = true;
+
+	UPROPERTY(EditAnywhere, config, Category = EventLog)
+	bool bActor = true;
+
+	UPROPERTY(EditAnywhere, config, Category = EventLog)
+	bool bCategory = true;
+
+	bool GetIsVisible(FName ColumnId) const
+	{
+		const FName PropertyName(TEXT("b") + ColumnId.ToString());
+		if (const FProperty* Property = StaticStruct()->FindPropertyByName(PropertyName))
+		{
+			if (const FBoolProperty* BoolProperty = CastField<FBoolProperty>(Property))
+			{
+				return BoolProperty->GetPropertyValue_InContainer(this);
+			}
+		}
+
+		ensureMsgf(false, TEXT("No VisibleColumns setting for ColumnId: %s"), *ColumnId.ToString());
+		return false;
+	}
+
+	void SetIsVisible(FName ColumnId, bool bIsVisible)
+	{
+		const FName PropertyName(TEXT("b") + ColumnId.ToString());
+		if (const FProperty* Property = StaticStruct()->FindPropertyByName(PropertyName))
+		{
+			if (const FBoolProperty* BoolProperty = CastField<FBoolProperty>(Property))
+			{
+				BoolProperty->SetPropertyValue_InContainer(this, bIsVisible);
+			}
+		}
+	}
+};
+
+USTRUCT()
+struct FAudioEventLogSettings
+{
+	GENERATED_BODY()
+
+	UPROPERTY(EditAnywhere, config, Category = EventLog)
+	FAudioEventLogVisibleColumns VisibleColumns;
+
+#if WITH_EDITOR
+	DECLARE_MULTICAST_DELEGATE_OneParam(FOnReadEventLogSettings, const FAudioEventLogSettings&);
+	static AUDIOINSIGHTS_API FOnReadEventLogSettings OnReadSettings;
+
+	DECLARE_MULTICAST_DELEGATE_OneParam(FOnWriteEventLogSettings, FAudioEventLogSettings&);
+	static AUDIOINSIGHTS_API FOnWriteEventLogSettings OnWriteSettings;
+
+	DECLARE_MULTICAST_DELEGATE(FOnRequestReadEventLogSettings);
+	static AUDIOINSIGHTS_API FOnRequestReadEventLogSettings OnRequestReadSettings;
+
+	DECLARE_MULTICAST_DELEGATE(FOnRequestWriteEventLogSettings);
+	static AUDIOINSIGHTS_API FOnRequestWriteEventLogSettings OnRequestWriteSettings;
+#endif // WITH_EDITOR
+};
+
+#undef LOCTEXT_NAMESPACE
\ No newline at end of file
diff --git a/Engine/Plugins/AudioInsights/Source/AudioInsights/Public/Views/AudioEventLogDashboardViewFactory.h b/Engine/Plugins/AudioInsights/Source/AudioInsights/Public/Views/AudioEventLogDashboardViewFactory.h
--- a/Engine/Plugins/AudioInsights/Source/AudioInsights/Public/Views/AudioEventLogDashboardViewFactory.h
+++ b/Engine/Plugins/AudioInsights/Source/AudioInsights/Public/Views/AudioEventLogDashboardViewFactory.h
@@ -2,6 +2,7 @@
 #pragma once
 
 #include "Filters/GenericFilter.h"
+#include "Settings/AudioEventLogSettings.h"
 #include "Templates/SharedPointer.h"
 #include "Views/SAudioFilterBar.h"
 #include "Views/TableDashboardViewFactory.h"
@@ -48,7 +49,7 @@ namespace UE::Audio::Insights
 	{
 	public:
 		UE_API FAudioEventLogDashboardViewFactory();
-		virtual ~FAudioEventLogDashboardViewFactory() = default;
+		UE_API virtual ~FAudioEventLogDashboardViewFactory();
 
 		UE_API virtual FName GetName() const override;
 		UE_API virtual FText GetDisplayName() const override;
@@ -64,10 +65,14 @@ namespace UE::Audio::Insights
 
 		UE_API virtual void SortTable() override;
 		UE_API virtual bool IsColumnSortable(const FName& ColumnId) const override;
+		UE_API void OnHiddenColumnsListChanged() override;
 
 		UE_API virtual TSharedPtr<SWidget> CreateFilterBarButtonWidget() override;
 		UE_API virtual TSharedRef<SWidget> CreateFilterBarWidget() override;
 
+		UE_API virtual TSharedRef<SWidget> OnGetSettingsMenuContent() override;
+		UE_API virtual FText OnGetSettingsMenuToolTip() override;
+
 		UE_API virtual void OnSelectionChanged(TSharedPtr<IDashboardDataViewEntry> SelectedItem, ESelectInfo::Type SelectInfo) override;
 		virtual bool ClearSelectionOnClick() const override { return true; }
 		UE_API virtual const FTableRowStyle* GetRowStyle() const override;
@@ -81,8 +86,20 @@ namespace UE::Audio::Insights
 		TSharedRef<FAudioEventLogFilter> CreateNewEventFilterType(const TSharedPtr<FFilterCategory>& FilterCategory, const FString& EventType);
 		bool PassesEventTypeFilter(const FEventLogFilterID EventLogID) const;
 
+		void BuildVisibleColumnsMenuContent(FMenuBuilder& MenuBuilder);
+		void ToggleColumnVisibility(FName ColumnId);
+		ECheckBoxState GetIsColumnVisibleCheckedState(FName ColumnId) const;
+		void ToggleAllColumnsVisibility();
+		ECheckBoxState GetToggleAllColumnsVisibilityCheckState() const;
+		FText GetToggleAllColumnsVisibilityText() const;
+
 		void SortByPredicate(TFunctionRef<bool(const class FAudioEventLogDashboardEntry&, const class FAudioEventLogDashboardEntry&)> Predicate);
 
+#if WITH_EDITOR
+		void OnReadEditorSettings(const FAudioEventLogSettings& InSettings);
+		void OnWriteEditorSettings(FAudioEventLogSettings& OutSettings);
+#endif // WITH_EDITOR
+
 #if !WITH_EDITOR
 		void OnTimingViewTimeMarkerChanged(double InTimeMarker);
 #endif // !WITH_EDITOR
@@ -92,6 +109,9 @@ namespace UE::Audio::Insights
 
 		TMap<FString, TSharedPtr<FFilterCategory>> FilterCategories;
 		TMap<FString, FEventLogFilterID> EventFilterTypes;
+
+		TSharedPtr<FUICommandList> CommandList;
+		TOptional<FAudioEventLogVisibleColumns> VisibleColumnSettings;
 	};
 } // namespace UE::Audio::Insights
 
diff --git a/Engine/Plugins/AudioInsights/Source/AudioInsights/Public/Views/TableDashboardViewFactory.h b/Engine/Plugins/AudioInsights/Source/AudioInsights/Public/Views/TableDashboardViewFactory.h
--- a/Engine/Plugins/AudioInsights/Source/AudioInsights/Public/Views/TableDashboardViewFactory.h
+++ b/Engine/Plugins/AudioInsights/Source/AudioInsights/Public/Views/TableDashboardViewFactory.h
@@ -146,10 +146,15 @@ namespace UE::Audio::Insights
 		virtual void SortTable() = 0;
 		virtual bool IsColumnSortable(const FName& ColumnId) const { return true; }
 
+		virtual void OnHiddenColumnsListChanged() { }
 
 		virtual TSharedPtr<SWidget> CreateFilterBarButtonWidget() { return nullptr; }
 		virtual TSharedRef<SWidget> CreateFilterBarWidget() { return SNullWidget::NullWidget; }
 
+		UE_API virtual TSharedRef<SWidget> CreateSettingsButtonWidget();
+		virtual TSharedRef<SWidget> OnGetSettingsMenuContent() { return SNullWidget::NullWidget; }
+		virtual FText OnGetSettingsMenuToolTip() { return FText::GetEmpty(); }
+
 		UE_API virtual TSharedPtr<SWidget> OnConstructContextMenu();
 		UE_API virtual void OnSelectionChanged(TSharedPtr<IDashboardDataViewEntry> SelectedItem, ESelectInfo::Type SelectInfo);
 		virtual bool ClearSelectionOnClick() const { return false; }
diff --git a/Engine/Plugins/AudioInsights/Source/AudioInsightsEditor/Private/AudioInsightsEditorSettings.cpp b/Engine/Plugins/AudioInsights/Source/AudioInsightsEditor/Private/AudioInsightsEditorSettings.cpp
--- a/Engine/Plugins/AudioInsights/Source/AudioInsightsEditor/Private/AudioInsightsEditorSettings.cpp
+++ b/Engine/Plugins/AudioInsights/Source/AudioInsightsEditor/Private/AudioInsightsEditorSettings.cpp
@@ -25,20 +25,30 @@ void UAudioInsightsEditorSettings::PostInitProperties()
 {
 	Super::PostInitProperties();
 
-	OnRequestReadSettingsHandle = FSoundDashboardSettings::OnRequestReadSettings.AddLambda([this]() { FSoundDashboardSettings::OnReadSettings.Broadcast(SoundDashboardSettings); });
-	OnRequestWriteSettingsHandle = FSoundDashboardSettings::OnRequestWriteSettings.AddLambda([this]()
-		{
-			FSoundDashboardSettings::OnWriteSettings.Broadcast(SoundDashboardSettings);
-			SaveConfig();
-		});
+	OnRequestReadEventLogSettingsHandle = FAudioEventLogSettings::OnRequestReadSettings.AddLambda([this](){ FAudioEventLogSettings::OnReadSettings.Broadcast(EventLogSettings); });
+	OnRequestWriteEventLogSettingsHandle = FAudioEventLogSettings::OnRequestWriteSettings.AddLambda([this]()
+	{
+		FAudioEventLogSettings::OnWriteSettings.Broadcast(EventLogSettings);
+		SaveConfig();
+	});
+
+	OnRequestReadSoundDashboardSettingsHandle = FSoundDashboardSettings::OnRequestReadSettings.AddLambda([this](){ FSoundDashboardSettings::OnReadSettings.Broadcast(SoundDashboardSettings); });
+	OnRequestWriteSoundDashboardSettingsHandle = FSoundDashboardSettings::OnRequestWriteSettings.AddLambda([this]()
+	{
+		FSoundDashboardSettings::OnWriteSettings.Broadcast(SoundDashboardSettings);
+		SaveConfig();
+	});
 }
 
 void UAudioInsightsEditorSettings::BeginDestroy()
 {
 	Super::BeginDestroy();
 
-	FSoundDashboardSettings::OnRequestReadSettings.Remove(OnRequestReadSettingsHandle);
-	FSoundDashboardSettings::OnRequestWriteSettings.Remove(OnRequestWriteSettingsHandle);
+	FAudioEventLogSettings::OnRequestReadSettings.Remove(OnRequestReadEventLogSettingsHandle);
+	FAudioEventLogSettings::OnRequestWriteSettings.Remove(OnRequestWriteEventLogSettingsHandle);
+
+	FSoundDashboardSettings::OnRequestReadSettings.Remove(OnRequestReadSoundDashboardSettingsHandle);
+	FSoundDashboardSettings::OnRequestWriteSettings.Remove(OnRequestWriteSoundDashboardSettingsHandle);
 }
 
 void UAudioInsightsEditorSettings::PostEditChangeChainProperty(FPropertyChangedChainEvent& PropertyChangedEvent)
@@ -57,6 +67,11 @@ void UAudioInsightsEditorSettings::PostEditChangeChainProperty(FPropertyChangedC
 		return;
 	}
 
+	if (ActiveMemberNode->GetValue()->GetFName() == GET_MEMBER_NAME_CHECKED(UAudioInsightsEditorSettings, EventLogSettings))
+	{
+		FAudioEventLogSettings::OnReadSettings.Broadcast(EventLogSettings);
+	}
+
 	if (ActiveMemberNode->GetValue()->GetFName() == GET_MEMBER_NAME_CHECKED(UAudioInsightsEditorSettings, SoundDashboardSettings))
 	{
 		FSoundDashboardSettings::OnReadSettings.Broadcast(SoundDashboardSettings);
diff --git a/Engine/Plugins/AudioInsights/Source/AudioInsightsEditor/Public/AudioInsightsEditorSettings.h b/Engine/Plugins/AudioInsights/Source/AudioInsightsEditor/Public/AudioInsightsEditorSettings.h
--- a/Engine/Plugins/AudioInsights/Source/AudioInsightsEditor/Public/AudioInsightsEditorSettings.h
+++ b/Engine/Plugins/AudioInsights/Source/AudioInsightsEditor/Public/AudioInsightsEditorSettings.h
@@ -5,6 +5,7 @@
 #include "AudioSpectrogram.h"
 #include "AudioSpectrumAnalyzer.h"
 #include "Engine/DeveloperSettings.h"
+#include "Settings/AudioEventLogSettings.h"
 #include "Settings/SoundDashboardSettings.h"
 
 #include "AudioInsightsEditorSettings.generated.h"
@@ -37,13 +38,20 @@ public:
 	UPROPERTY(EditAnywhere, config, Category = SpectrumAnalyzer, meta = (ShowOnlyInnerProperties))
 	FSpectrumAnalyzerRackUnitSettings SpectrumAnalyzerSettings;
 
+	/** Settings for the Event Log */
+	UPROPERTY(EditAnywhere, config, Category = EventLog, meta = (ShowOnlyInnerProperties))
+	FAudioEventLogSettings EventLogSettings;
+
 	/** Settings for the Sound Dashboard */
 	UPROPERTY(EditAnywhere, config, Category = SoundDashboard, meta = (ShowOnlyInnerProperties))
 	FSoundDashboardSettings SoundDashboardSettings;
 
 private:
-	FDelegateHandle OnRequestReadSettingsHandle;
-	FDelegateHandle OnRequestWriteSettingsHandle;
+	FDelegateHandle OnRequestReadEventLogSettingsHandle;
+	FDelegateHandle OnRequestWriteEventLogSettingsHandle;
+
+	FDelegateHandle OnRequestReadSoundDashboardSettingsHandle;
+	FDelegateHandle OnRequestWriteSoundDashboardSettingsHandle;
 };
 
 #undef UE_API
