[Audio Insights] Crash fix when moving timemarker inside event log of an old trace file - Incorrect variable being null checked inside SoundTraceProvider

diff --git a/Engine/Plugins/AudioInsights/Source/AudioInsights/Private/Providers/SoundTraceProvider.cpp b/Engine/Plugins/AudioInsights/Source/AudioInsights/Private/Providers/SoundTraceProvider.cpp
--- a/Engine/Plugins/AudioInsights/Source/AudioInsights/Private/Providers/SoundTraceProvider.cpp
+++ b/Engine/Plugins/AudioInsights/Source/AudioInsights/Private/Providers/SoundTraceProvider.cpp
@@ -1196,7 +1196,7 @@ namespace UE::Audio::Insights
 
 			GetActiveSoundEntryFromIDs(SoundWaveIsAlivePingCachedMessage.ActiveSoundPlayOrder, ParentSoundIsAlivePingEntryReturn);
 
-			if (ParentSoundEntryReturn == nullptr || !ParentSoundIsAlivePingEntryReturn->IsValid() || !FSoundTraceProviderPrivate::EntryTypeHasSoundWaveEntries(ParentSoundIsAlivePingEntryReturn->EntryType))
+			if (ParentSoundIsAlivePingEntryReturn == nullptr || !ParentSoundIsAlivePingEntryReturn->IsValid() || !FSoundTraceProviderPrivate::EntryTypeHasSoundWaveEntries(ParentSoundIsAlivePingEntryReturn->EntryType))
 			{
 				continue;
 			}
