[Audio Insights] Create helper class for visible column settings to avoid repetition across dashboards

diff --git a/Engine/Plugins/AudioInsights/Source/AudioInsights/Private/Settings/VisibleColumnsSettings.cpp b/Engine/Plugins/AudioInsights/Source/AudioInsights/Private/Settings/VisibleColumnsSettings.cpp
new file mode 100644
--- /dev/null
+++ b/Engine/Plugins/AudioInsights/Source/AudioInsights/Private/Settings/VisibleColumnsSettings.cpp
@@ -0,0 +1,29 @@
+// Copyright Epic Games, Inc. All Rights Reserved.
+#include "Settings/VisibleColumnsSettings.h"
+
+bool FVisibleColumnsSettings::GetIsVisible(const FName& ColumnId) const
+{
+	const FName PropertyName(TEXT("b") + ColumnId.ToString());
+	if (const FProperty* Property = FindProperty(PropertyName))
+	{
+		if (const FBoolProperty* BoolProperty = CastField<FBoolProperty>(Property))
+		{
+			return BoolProperty->GetPropertyValue_InContainer(this);
+		}
+	}
+
+	ensureMsgf(false, TEXT("No VisibleColumns setting for ColumnId: %s"), *ColumnId.ToString());
+	return false;
+}
+
+void FVisibleColumnsSettings::SetIsVisible(const FName& ColumnId, const bool bIsVisible)
+{
+	const FName PropertyName(TEXT("b") + ColumnId.ToString());
+	if (const FProperty* Property = FindProperty(PropertyName))
+	{
+		if (const FBoolProperty* BoolProperty = CastField<FBoolProperty>(Property))
+		{
+			BoolProperty->SetPropertyValue_InContainer(this, bIsVisible);
+		}
+	}
+}
\ No newline at end of file
diff --git a/Engine/Plugins/AudioInsights/Source/AudioInsights/Private/Views/AudioEventLogDashboardViewFactory.cpp b/Engine/Plugins/AudioInsights/Source/AudioInsights/Private/Views/AudioEventLogDashboardViewFactory.cpp
--- a/Engine/Plugins/AudioInsights/Source/AudioInsights/Private/Views/AudioEventLogDashboardViewFactory.cpp
+++ b/Engine/Plugins/AudioInsights/Source/AudioInsights/Private/Views/AudioEventLogDashboardViewFactory.cpp
@@ -1,7 +1,6 @@
 // Copyright Epic Games, Inc. All Rights Reserved.
 #include "Views/AudioEventLogDashboardViewFactory.h"
 
-#include "Algo/AnyOf.h"
 #include "AudioInsightsModule.h"
 #include "AudioInsightsStyle.h"
 #include "AudioMixerTrace.h"
@@ -173,6 +172,19 @@ namespace UE::Audio::Insights
 			{
 				FilteredEntriesListView->SetSelectionMode(ESelectionMode::SingleToggle);
 			}
+
+			if (HeaderRowWidget.IsValid())
+			{
+				const FAudioEventLogVisibleColumns VisibleColumns = InitVisibleColumnSettings.IsSet() ? InitVisibleColumnSettings.GetValue() : FAudioEventLogVisibleColumns();
+				VisibleColumnsSettingsMenu = MakeShared<FVisibleColumnsSettingsMenu<FAudioEventLogVisibleColumns>>(HeaderRowWidget.ToSharedRef(), VisibleColumns);
+
+				VisibleColumnsSettingsMenu->OnVisibleColumnsSettingsUpdated.AddSPLambda(this, []()
+				{
+#if WITH_EDITOR
+					FAudioEventLogSettings::OnRequestWriteSettings.Broadcast();
+#endif // WITH_EDITOR	
+				});
+			}
 		}
 
 		return DashboardWidget->AsShared();
@@ -229,12 +241,7 @@ namespace UE::Audio::Insights
 
 		auto CreateColumnData = [this]()
 		{
-			FAudioEventLogVisibleColumns VisibleColumns;
-			if (VisibleColumnSettings.IsSet())
-			{
-				// Use the values loaded from the settings:
-				VisibleColumns = VisibleColumnSettings.GetValue();
-			}
+			FAudioEventLogVisibleColumns VisibleColumns = InitVisibleColumnSettings.IsSet() ? InitVisibleColumnSettings.GetValue() : FAudioEventLogVisibleColumns();
 
 			return TMap<FName, FTraceTableDashboardViewFactory::FColumnData>
 			{
@@ -382,17 +389,9 @@ namespace UE::Audio::Insights
 
 	void FAudioEventLogDashboardViewFactory::OnHiddenColumnsListChanged()
 	{
-		// Update and save settings to reload it on the next Editor sessions
-		if (VisibleColumnSettings && HeaderRowWidget)
+		if (VisibleColumnsSettingsMenu.IsValid())
 		{
-			for (const SHeaderRow::FColumn& Column : HeaderRowWidget->GetColumns())
-			{
-				VisibleColumnSettings->SetIsVisible(Column.ColumnId, Column.bIsVisible);
-			}
-
-#if WITH_EDITOR
-			FAudioEventLogSettings::OnRequestWriteSettings.Broadcast();
-#endif // WITH_EDITOR
+			VisibleColumnsSettingsMenu->OnHiddenColumnsListChanged();
 		}
 	}
 
@@ -569,103 +568,9 @@ namespace UE::Audio::Insights
 
 	void FAudioEventLogDashboardViewFactory::BuildVisibleColumnsMenuContent(FMenuBuilder& MenuBuilder)
 	{
-		MenuBuilder.AddMenuEntry
-		(
-			TAttribute<FText>::CreateSP(this, &FAudioEventLogDashboardViewFactory::GetToggleAllColumnsVisibilityText),
-			FText::GetEmpty(),
-			FSlateIcon(),
-			FUIAction
-			(
-				FExecuteAction::CreateSP(this, &FAudioEventLogDashboardViewFactory::ToggleAllColumnsVisibility),
-				FCanExecuteAction::CreateLambda([]() { return true; }),
-				FGetActionCheckState::CreateSP(this, &FAudioEventLogDashboardViewFactory::GetToggleAllColumnsVisibilityCheckState)
-			),
-			NAME_None,
-			EUserInterfaceActionType::ToggleButton
-		);
-
-		for (const SHeaderRow::FColumn& Column : HeaderRowWidget->GetColumns())
+		if (VisibleColumnsSettingsMenu.IsValid())
 		{
-			MenuBuilder.AddMenuEntry
-			(
-				Column.DefaultText,
-				FText::GetEmpty(),
-				FSlateIcon(),
-				FUIAction
-				(
-					FExecuteAction::CreateSP(this, &FAudioEventLogDashboardViewFactory::ToggleColumnVisibility, Column.ColumnId),
-					FCanExecuteAction::CreateLambda([]() { return true; }),
-					FGetActionCheckState::CreateSP(this, &FAudioEventLogDashboardViewFactory::GetIsColumnVisibleCheckedState, Column.ColumnId)
-				),
-				NAME_None,
-				EUserInterfaceActionType::ToggleButton
-			);
-		}
-	}
-
-	void FAudioEventLogDashboardViewFactory::ToggleColumnVisibility(FName ColumnId)
-	{
-		const bool bIsColumnVisible = HeaderRowWidget->IsColumnVisible(ColumnId);
-		HeaderRowWidget->SetShowGeneratedColumn(ColumnId, !bIsColumnVisible);
-	}
-
-	ECheckBoxState FAudioEventLogDashboardViewFactory::GetIsColumnVisibleCheckedState(FName ColumnId) const
-	{
-		return HeaderRowWidget->IsColumnVisible(ColumnId) ? ECheckBoxState::Checked : ECheckBoxState::Unchecked;
-	}
-
-	void FAudioEventLogDashboardViewFactory::ToggleAllColumnsVisibility()
-	{
-		const ECheckBoxState CurrentState = GetToggleAllColumnsVisibilityCheckState();
-		const bool bVisibleIsDesired = CurrentState == ECheckBoxState::Unchecked || CurrentState == ECheckBoxState::Undetermined;
-
-		// Null our VisibleColumnSettings TOptional while updating HeaderRowWidget so that we ignore OnHiddenColumnsListChanged callbacks:
-		VisibleColumnSettings.Reset();
-
-		for (const SHeaderRow::FColumn& Column : HeaderRowWidget->GetColumns())
-		{
-			if (Column.bIsVisible != bVisibleIsDesired)
-			{
-				HeaderRowWidget->SetShowGeneratedColumn(Column.ColumnId, bVisibleIsDesired);
-			}
-		}
-
-		// Reinitialize the VisibleColumnSettings and perform a single OnHiddenColumnsListChanged call to update it and persist the settings:
-		VisibleColumnSettings.Emplace();
-		OnHiddenColumnsListChanged();
-	}
-
-	ECheckBoxState FAudioEventLogDashboardViewFactory::GetToggleAllColumnsVisibilityCheckState() const
-	{
-		const TIndirectArray<SHeaderRow::FColumn>& Columns = HeaderRowWidget->GetColumns();
-		const bool bAnyColumnsVisible = Algo::AnyOf(Columns, [](const SHeaderRow::FColumn& Column) { return Column.bIsVisible; });
-		const bool bAnyColumnsInvisible = Algo::AnyOf(Columns, [](const SHeaderRow::FColumn& Column) { return !Column.bIsVisible; });
-
-		if (bAnyColumnsVisible && bAnyColumnsInvisible)
-		{
-			return ECheckBoxState::Undetermined;
-		}
-		else if (bAnyColumnsVisible)
-		{
-			return ECheckBoxState::Checked;
-		}
-		else
-		{
-			return ECheckBoxState::Unchecked;
-		}
-	}
-
-	FText FAudioEventLogDashboardViewFactory::GetToggleAllColumnsVisibilityText() const
-	{
-		switch (GetToggleAllColumnsVisibilityCheckState())
-		{
-			case ECheckBoxState::Unchecked:
-			case ECheckBoxState::Undetermined:
-				return LOCTEXT("AudioEventLog_VisibleColumns_SelectAll", "Select: All");
-			case ECheckBoxState::Checked:
-				return LOCTEXT("AudioEventLog_VisibleColumns_SelectNone", "Select: None");
-			default:
-				return FText::GetEmpty();
+			VisibleColumnsSettingsMenu->BuildVisibleColumnsMenuContent(MenuBuilder);
 		}
 	}
 
@@ -697,30 +602,22 @@ namespace UE::Audio::Insights
 #if WITH_EDITOR
 	void FAudioEventLogDashboardViewFactory::OnReadEditorSettings(const FAudioEventLogSettings& InSettings)
 	{
-		// If HeaderRowWidget is already created, update it with the desired visible columns:
-		if (HeaderRowWidget)
+		if (VisibleColumnsSettingsMenu.IsValid())
 		{
-			// Null our VisibleColumnSettings TOptional while updating HeaderRowWidget so that we ignore OnHiddenColumnsListChanged callbacks:
-			VisibleColumnSettings.Reset();
-
-			for (const SHeaderRow::FColumn& Column : HeaderRowWidget->GetColumns())
-			{
-				const bool bShowColumn = InSettings.VisibleColumns.GetIsVisible(Column.ColumnId);
-				if (Column.bIsVisible != bShowColumn)
-				{
-					HeaderRowWidget->SetShowGeneratedColumn(Column.ColumnId, bShowColumn);
-				}
-			}
+			VisibleColumnsSettingsMenu->ReadFromSettings(InSettings.VisibleColumns);
+		}
+		
+		if (!InitVisibleColumnSettings.IsSet())
+		{
+			InitVisibleColumnSettings = InSettings.VisibleColumns;
 		}
-
-		VisibleColumnSettings = InSettings.VisibleColumns;
 	}
 
 	void FAudioEventLogDashboardViewFactory::OnWriteEditorSettings(FAudioEventLogSettings& OutSettings)
 	{
-		if (VisibleColumnSettings)
+		if (VisibleColumnsSettingsMenu.IsValid())
 		{
-			OutSettings.VisibleColumns = VisibleColumnSettings.GetValue();
+			VisibleColumnsSettingsMenu->WriteToSettings(OutSettings.VisibleColumns);
 		}
 	}
 #endif // WITH_EDITOR
diff --git a/Engine/Plugins/AudioInsights/Source/AudioInsights/Private/Views/SoundDashboardViewFactory.cpp b/Engine/Plugins/AudioInsights/Source/AudioInsights/Private/Views/SoundDashboardViewFactory.cpp
--- a/Engine/Plugins/AudioInsights/Source/AudioInsights/Private/Views/SoundDashboardViewFactory.cpp
+++ b/Engine/Plugins/AudioInsights/Source/AudioInsights/Private/Views/SoundDashboardViewFactory.cpp
@@ -1094,7 +1094,7 @@ namespace UE::Audio::Insights
 		AudioInsightsModule.GetTimingViewExtender().OnTimingViewTimeMarkerChanged.AddSP(this, &FSoundDashboardViewFactory::OnTimingViewTimeMarkerChanged);
 #endif // WITH_EDITOR
 
-		return SNew(SVerticalBox)
+		const TSharedRef<SWidget> SoundDashboardWidget = SNew(SVerticalBox)
 #if WITH_EDITOR
 			+ SVerticalBox::Slot()
 			.AutoHeight()
@@ -1155,6 +1155,21 @@ namespace UE::Audio::Insights
 			[
 				FTraceTreeDashboardViewFactory::MakeWidget(OwnerTab, SpawnTabArgs)
 			];
+
+		if (HeaderRowWidget.IsValid())
+		{
+			const FSoundDashboardVisibleColumns VisibleColumns = InitVisibleColumnSettings.IsSet() ? InitVisibleColumnSettings.GetValue() : FSoundDashboardVisibleColumns();
+			VisibleColumnsSettingsMenu = MakeShared<FVisibleColumnsSettingsMenu<FSoundDashboardVisibleColumns>>(HeaderRowWidget.ToSharedRef(), VisibleColumns);
+
+			VisibleColumnsSettingsMenu->OnVisibleColumnsSettingsUpdated.AddSPLambda(this, []()
+			{
+#if WITH_EDITOR
+				FSoundDashboardSettings::OnRequestWriteSettings.Broadcast();
+#endif // WITH_EDITOR	
+			});
+		}
+
+		return SoundDashboardWidget->AsShared();
 	}
 
 	TSharedRef<SWidget> FSoundDashboardViewFactory::GenerateWidgetForRootColumn(const TSharedRef<FTraceTreeDashboardViewFactory::SRowWidget>& InRowWidget, const TSharedRef<IDashboardDataTreeViewEntry>& InRowData, const FName& InColumn, const FText& InValueText)
@@ -1585,97 +1600,9 @@ namespace UE::Audio::Insights
 
 	void FSoundDashboardViewFactory::BuildVisibleColumnsMenuContent(FMenuBuilder& MenuBuilder)
 	{
-		MenuBuilder.AddMenuEntry(
-			TAttribute<FText>::CreateSP(this, &FSoundDashboardViewFactory::GetToggleAllColumnsVisibilityText),
-			FText::GetEmpty(),
-			FSlateIcon(),
-			FUIAction(
-				FExecuteAction::CreateSP(this, &FSoundDashboardViewFactory::ToggleAllColumnsVisibility),
-				FCanExecuteAction::CreateLambda([]() { return true; }),
-				FGetActionCheckState::CreateSP(this, &FSoundDashboardViewFactory::GetToggleAllColumnsVisibilityCheckState)
-			),
-			NAME_None,
-			EUserInterfaceActionType::ToggleButton);
-
-		for (const SHeaderRow::FColumn& Column : HeaderRowWidget->GetColumns())
-		{
-			MenuBuilder.AddMenuEntry(
-				Column.DefaultText,
-				FText::GetEmpty(),
-				FSlateIcon(),
-				FUIAction(
-					FExecuteAction::CreateSP(this, &FSoundDashboardViewFactory::ToggleColumnVisibility, Column.ColumnId),
-					FCanExecuteAction::CreateLambda([]() { return true; }),
-					FGetActionCheckState::CreateSP(this, &FSoundDashboardViewFactory::GetIsColumnVisibleCheckedState, Column.ColumnId)
-				),
-				NAME_None,
-				EUserInterfaceActionType::ToggleButton);
-		}
-	}
-
-	void FSoundDashboardViewFactory::ToggleColumnVisibility(FName ColumnId)
-	{
-		const bool bIsColumnVisible = HeaderRowWidget->IsColumnVisible(ColumnId);
-		HeaderRowWidget->SetShowGeneratedColumn(ColumnId, !bIsColumnVisible);
-	}
-
-	ECheckBoxState FSoundDashboardViewFactory::GetIsColumnVisibleCheckedState(FName ColumnId) const
-	{
-		return HeaderRowWidget->IsColumnVisible(ColumnId) ? ECheckBoxState::Checked : ECheckBoxState::Unchecked;
-	}
-
-	void FSoundDashboardViewFactory::ToggleAllColumnsVisibility()
-	{
-		const ECheckBoxState CurrentState = GetToggleAllColumnsVisibilityCheckState();
-		const bool bVisibleIsDesired = CurrentState == ECheckBoxState::Unchecked || CurrentState == ECheckBoxState::Undetermined;
-
-		// Null our VisibleColumnSettings TOptional while updating HeaderRowWidget so that we ignore OnHiddenColumnsListChanged callbacks:
-		VisibleColumnSettings.Reset();
-
-		for (const SHeaderRow::FColumn& Column : HeaderRowWidget->GetColumns())
+		if (VisibleColumnsSettingsMenu.IsValid())
 		{
-			if (Column.bIsVisible != bVisibleIsDesired)
-			{
-				HeaderRowWidget->SetShowGeneratedColumn(Column.ColumnId, bVisibleIsDesired);
-			}
-		}
-
-		// Reinitialize the VisibleColumnSettings and perform a single OnHiddenColumnsListChanged call to update it and persist the settings:
-		VisibleColumnSettings.Emplace();
-		OnHiddenColumnsListChanged();
-	}
-
-	ECheckBoxState FSoundDashboardViewFactory::GetToggleAllColumnsVisibilityCheckState() const
-	{
-		const TIndirectArray<SHeaderRow::FColumn>& Columns = HeaderRowWidget->GetColumns();
-		const bool bAnyColumnsVisible = Algo::AnyOf(Columns, [](const SHeaderRow::FColumn& Column) { return Column.bIsVisible; });
-		const bool bAnyColumnsInvisible = Algo::AnyOf(Columns, [](const SHeaderRow::FColumn& Column) { return !Column.bIsVisible; });
-
-		if (bAnyColumnsVisible && bAnyColumnsInvisible)
-		{
-			return ECheckBoxState::Undetermined;
-		}
-		else if (bAnyColumnsVisible)
-		{
-			return ECheckBoxState::Checked;
-		}
-		else
-		{
-			return ECheckBoxState::Unchecked;
-		}
-	}
-
-	FText FSoundDashboardViewFactory::GetToggleAllColumnsVisibilityText() const
-	{
-		switch (GetToggleAllColumnsVisibilityCheckState())
-		{
-		case ECheckBoxState::Unchecked:
-		case ECheckBoxState::Undetermined:
-			return LOCTEXT("SoundsDashboard_VisibleColumns_SelectAll", "Select: All");
-		case ECheckBoxState::Checked:
-			return LOCTEXT("SoundsDashboard_VisibleColumns_SelectNone", "Select: None");
-		default:
-			return FText::GetEmpty();
+			VisibleColumnsSettingsMenu->BuildVisibleColumnsMenuContent(MenuBuilder);
 		}
 	}
 
@@ -1783,12 +1710,7 @@ namespace UE::Audio::Insights
 	{
 		using namespace FSoundDashboardViewFactoryPrivate;
 
-		FSoundDashboardVisibleColumns VisibleColumns;
-		if (VisibleColumnSettings.IsSet())
-		{
-			// Use the values loaded from the settings:
-			VisibleColumns = VisibleColumnSettings.GetValue();
-		}
+		FSoundDashboardVisibleColumns VisibleColumns = InitVisibleColumnSettings.IsSet() ? InitVisibleColumnSettings.GetValue() : FSoundDashboardVisibleColumns();
 
 		static const TMap<FName, FTraceTreeDashboardViewFactory::FHeaderRowColumnData> HeaderRowColumnData =
 		{
@@ -1995,17 +1917,9 @@ namespace UE::Audio::Insights
 
 	void FSoundDashboardViewFactory::OnHiddenColumnsListChanged()
 	{
-		// Update and save settings to reload it on the next Editor sessions
-		if (VisibleColumnSettings && HeaderRowWidget)
+		if (VisibleColumnsSettingsMenu.IsValid())
 		{
-			for (const SHeaderRow::FColumn& Column : HeaderRowWidget->GetColumns())
-			{
-				VisibleColumnSettings->SetIsVisible(Column.ColumnId, Column.bIsVisible);
-			}
-
-#if WITH_EDITOR
-			FSoundDashboardSettings::OnRequestWriteSettings.Broadcast();
-#endif // WITH_EDITOR
+			VisibleColumnsSettingsMenu->OnHiddenColumnsListChanged();
 		}
 	}
 
@@ -2813,23 +2727,16 @@ namespace UE::Audio::Insights
 		bShowRecentlyStoppedSounds = InSettings.bShowStoppedSounds;
 		bDisplayAmpPeakInDb = InSettings.AmplitudeDisplayMode == EAudioAmplitudeDisplayMode::Decibels;
 
-		// If HeaderRowWidget is already created, update it with the desired visible columns:
-		if (HeaderRowWidget)
+		// If VisibleColumnsSettingsMenu is already created, update it with the desired visible columns:
+		if (VisibleColumnsSettingsMenu.IsValid())
 		{
-			// Null our VisibleColumnSettings TOptional while updating HeaderRowWidget so that we ignore OnHiddenColumnsListChanged callbacks:
-			VisibleColumnSettings.Reset();
-
-			for (const SHeaderRow::FColumn& Column : HeaderRowWidget->GetColumns())
-			{
-				const bool bShowColumn = InSettings.VisibleColumns.GetIsVisible(Column.ColumnId);
-				if (Column.bIsVisible != bShowColumn)
-				{
-					HeaderRowWidget->SetShowGeneratedColumn(Column.ColumnId, bShowColumn);
-				}
-			}
+			VisibleColumnsSettingsMenu->ReadFromSettings(InSettings.VisibleColumns);
 		}
 
-		VisibleColumnSettings = InSettings.VisibleColumns;
+		if (!InitVisibleColumnSettings.IsSet())
+		{
+			InitVisibleColumnSettings = InSettings.VisibleColumns;
+		}
 
 		if (RecentlyStoppedSoundTimeout != InSettings.StoppedSoundTimeoutTime)
 		{
@@ -2846,9 +2753,9 @@ namespace UE::Audio::Insights
 		OutSettings.bShowStoppedSounds = bShowRecentlyStoppedSounds;
 		OutSettings.AmplitudeDisplayMode = bDisplayAmpPeakInDb ? EAudioAmplitudeDisplayMode::Decibels : EAudioAmplitudeDisplayMode::Linear;
 
-		if (VisibleColumnSettings)
+		if (VisibleColumnsSettingsMenu.IsValid())
 		{
-			OutSettings.VisibleColumns = VisibleColumnSettings.GetValue();
+			VisibleColumnsSettingsMenu->WriteToSettings(OutSettings.VisibleColumns);
 		}
 
 		if (OutSettings.StoppedSoundTimeoutTime != RecentlyStoppedSoundTimeout)
diff --git a/Engine/Plugins/AudioInsights/Source/AudioInsights/Public/Settings/AudioEventLogSettings.h b/Engine/Plugins/AudioInsights/Source/AudioInsights/Public/Settings/AudioEventLogSettings.h
--- a/Engine/Plugins/AudioInsights/Source/AudioInsights/Public/Settings/AudioEventLogSettings.h
+++ b/Engine/Plugins/AudioInsights/Source/AudioInsights/Public/Settings/AudioEventLogSettings.h
@@ -4,6 +4,7 @@
 
 #include "Delegates/DelegateCombinations.h"
 #include "Engine/DeveloperSettings.h"
+#include "Settings/VisibleColumnsSettings.h"
 
 #include "AudioEventLogSettings.generated.h"
 
@@ -14,7 +15,7 @@
  * Control which columns are visible in the Audio Event Log tab.
  */
 USTRUCT()
-struct FAudioEventLogVisibleColumns
+struct FAudioEventLogVisibleColumns : public FVisibleColumnsSettings
 {
 	GENERATED_BODY()
 
@@ -39,31 +40,9 @@ struct FAudioEventLogVisibleColumns
 	UPROPERTY(EditAnywhere, config, Category = EventLog)
 	bool bCategory = true;
 
-	bool GetIsVisible(FName ColumnId) const
+	virtual const FProperty* FindProperty(const FName& PropertyName) const override
 	{
-		const FName PropertyName(TEXT("b") + ColumnId.ToString());
-		if (const FProperty* Property = StaticStruct()->FindPropertyByName(PropertyName))
-		{
-			if (const FBoolProperty* BoolProperty = CastField<FBoolProperty>(Property))
-			{
-				return BoolProperty->GetPropertyValue_InContainer(this);
-			}
-		}
-
-		ensureMsgf(false, TEXT("No VisibleColumns setting for ColumnId: %s"), *ColumnId.ToString());
-		return false;
-	}
-
-	void SetIsVisible(FName ColumnId, bool bIsVisible)
-	{
-		const FName PropertyName(TEXT("b") + ColumnId.ToString());
-		if (const FProperty* Property = StaticStruct()->FindPropertyByName(PropertyName))
-		{
-			if (const FBoolProperty* BoolProperty = CastField<FBoolProperty>(Property))
-			{
-				BoolProperty->SetPropertyValue_InContainer(this, bIsVisible);
-			}
-		}
+		return StaticStruct()->FindPropertyByName(PropertyName);
 	}
 };
 
diff --git a/Engine/Plugins/AudioInsights/Source/AudioInsights/Public/Settings/SoundDashboardSettings.h b/Engine/Plugins/AudioInsights/Source/AudioInsights/Public/Settings/SoundDashboardSettings.h
--- a/Engine/Plugins/AudioInsights/Source/AudioInsights/Public/Settings/SoundDashboardSettings.h
+++ b/Engine/Plugins/AudioInsights/Source/AudioInsights/Public/Settings/SoundDashboardSettings.h
@@ -4,6 +4,7 @@
 
 #include "Delegates/DelegateCombinations.h"
 #include "Engine/DeveloperSettings.h"
+#include "Settings/VisibleColumnsSettingsMenu.h"
 
 #include "SoundDashboardSettings.generated.h"
 
@@ -58,7 +59,7 @@ enum class ESoundDashboardAutoExpandOptions : uint8
 };
 
 USTRUCT()
-struct FSoundDashboardVisibleColumns
+struct FSoundDashboardVisibleColumns : public FVisibleColumnsSettings
 {
 	GENERATED_BODY()
 
@@ -110,31 +111,9 @@ struct FSoundDashboardVisibleColumns
 	UPROPERTY(EditAnywhere, config, Category = SoundDashboard)
 	bool bCategory = false;
 
-	bool GetIsVisible(FName ColumnId) const
+	virtual const FProperty* FindProperty(const FName& PropertyName) const override
 	{
-		const FName PropertyName(TEXT("b") + ColumnId.ToString());
-		if (const FProperty* Property = StaticStruct()->FindPropertyByName(PropertyName))
-		{
-			if (const FBoolProperty* BoolProperty = CastField<FBoolProperty>(Property))
-			{
-				return BoolProperty->GetPropertyValue_InContainer(this);
-			}
-		}
-
-		ensureMsgf(false, TEXT("No VisibleColumns setting for ColumnId: %s"), *ColumnId.ToString());
-		return false;
-	}
-
-	void SetIsVisible(FName ColumnId, bool bIsVisible)
-	{
-		const FName PropertyName(TEXT("b") + ColumnId.ToString());
-		if (const FProperty* Property = StaticStruct()->FindPropertyByName(PropertyName))
-		{
-			if (const FBoolProperty* BoolProperty = CastField<FBoolProperty>(Property))
-			{
-				BoolProperty->SetPropertyValue_InContainer(this, bIsVisible);
-			}
-		}
+		return StaticStruct()->FindPropertyByName(PropertyName);
 	}
 };
 
diff --git a/Engine/Plugins/AudioInsights/Source/AudioInsights/Public/Settings/VisibleColumnsSettings.h b/Engine/Plugins/AudioInsights/Source/AudioInsights/Public/Settings/VisibleColumnsSettings.h
new file mode 100644
--- /dev/null
+++ b/Engine/Plugins/AudioInsights/Source/AudioInsights/Public/Settings/VisibleColumnsSettings.h
@@ -0,0 +1,34 @@
+// Copyright Epic Games, Inc. All Rights Reserved.
+
+#pragma once
+
+#include "VisibleColumnsSettings.generated.h"
+
+#define UE_API AUDIOINSIGHTS_API
+
+/** FVisibleColumnsSettings
+*
+* Helper struct for adding visible column settings to an Audio Insights dashboard
+* Used in conjunction with FVisibleColumnsSettingsMenu to populate a menu with enable/disable checkboxes for the columns in a table/tree
+* Inherit from this in a USTRUCT that contains boolean UPROPERTIES for each column in the table, with names matching the column IDs with a 'b' prefix
+* 
+* e.g.
+* Table : | Column01 | Column02 | Column03 |
+* USTRUCT : bColumn01, bColumn02, bColumn03
+*/
+USTRUCT()
+struct FVisibleColumnsSettings
+{
+	GENERATED_BODY()
+
+	virtual ~FVisibleColumnsSettings() = default;
+
+	UE_API virtual bool GetIsVisible(const FName& ColumnId) const;
+	UE_API virtual void SetIsVisible(const FName& ColumnId, const bool bIsVisible);
+
+	// Override in derived class to enable finding properties by name in the derived USTRUCT
+	// E.g. return StaticStruct()->FindPropertyByName(PropertyName);
+	UE_API virtual const FProperty* FindProperty(const FName& PropertyName) const { return nullptr; }
+};
+
+#undef UE_API
\ No newline at end of file
diff --git a/Engine/Plugins/AudioInsights/Source/AudioInsights/Public/Settings/VisibleColumnsSettingsMenu.h b/Engine/Plugins/AudioInsights/Source/AudioInsights/Public/Settings/VisibleColumnsSettingsMenu.h
new file mode 100644
--- /dev/null
+++ b/Engine/Plugins/AudioInsights/Source/AudioInsights/Public/Settings/VisibleColumnsSettingsMenu.h
@@ -0,0 +1,68 @@
+// Copyright Epic Games, Inc. All Rights Reserved.
+
+#pragma once
+
+#include "Delegates/DelegateCombinations.h"
+#include "Settings/VisibleColumnsSettings.h"
+#include "Styling/SlateTypes.h"
+
+class FMenuBuilder;
+class SHeaderRow;
+
+template<typename T>
+concept TIsVisibleColumnsSettingsType = std::is_base_of_v<FVisibleColumnsSettings, T>;
+
+namespace UE::Audio::Insights
+{
+	/** FVisibleColumnsSettingsMenu
+	*
+	* Helper class for creating a settings sub-menu for the visibility of columns in a table
+	* Requires a settings USTRUCT that inherits from FVisibleColumnsSettings to be passed in on construction.
+	*/
+	template<TIsVisibleColumnsSettingsType T>
+	class FVisibleColumnsSettingsMenu : public TSharedFromThis<FVisibleColumnsSettingsMenu<T>>
+	{
+	public:
+		/*
+		* Constructor
+		* @param 	InHeaderRowWidget			Header Row widget from the table (SListView, STreeView, etc.)
+		* @param 	InVisibleColumnSettings		A USTRUCT that inherits from FVisibleColumnsSettings, typically to be stored in the Editor Preferences
+		*/
+		FVisibleColumnsSettingsMenu(TSharedRef<SHeaderRow> InHeaderRowWidget, const T& InVisibleColumnSettings) 
+			: HeaderRowWidget(InHeaderRowWidget)
+			, VisibleColumnSettings(InVisibleColumnSettings)
+		{
+		}
+
+		virtual ~FVisibleColumnsSettingsMenu() = default;
+		
+		/* Pass in a FMenuBuilder object to populate a sub-menu with visible columns settings */
+		void BuildVisibleColumnsMenuContent(FMenuBuilder& OutMenuBuilder);
+		
+		/* Call from SHeaderRow::OnHiddenColumnsListChanged callback to update the settings when the column visibility is changed from the header row  */
+		void OnHiddenColumnsListChanged();
+
+		/* Call to cache/write the current column visibility settings from/to editor preferences */
+		void ReadFromSettings(const T& InVisibleColumnSettings);
+		void WriteToSettings(T& OutVisibleColumnSettings);
+
+		/* Delegate that notifies when the cached VisibleColumnSettings have been updated */
+		DECLARE_MULTICAST_DELEGATE(FOnVisibleColumnsSettingsUpdated);
+		FOnVisibleColumnsSettingsUpdated OnVisibleColumnsSettingsUpdated;
+
+	private:
+		void ToggleColumnVisibility(FName ColumnId);
+		void ToggleAllColumnsVisibility();
+
+		void UpdateColumnVisabilityFromSettings(const FVisibleColumnsSettings& InVisibleColumnSettings);
+
+		ECheckBoxState GetIsColumnVisibleCheckedState(FName ColumnId) const;
+		ECheckBoxState GetToggleAllColumnsVisibilityCheckState() const;
+		FText GetToggleAllColumnsVisibilityText() const;
+
+		TWeakPtr<SHeaderRow> HeaderRowWidget;
+		TOptional<T> VisibleColumnSettings;
+	};
+}
+
+#include "Settings/VisibleColumnsSettingsMenu.inl"
\ No newline at end of file
diff --git a/Engine/Plugins/AudioInsights/Source/AudioInsights/Public/Settings/VisibleColumnsSettingsMenu.inl b/Engine/Plugins/AudioInsights/Source/AudioInsights/Public/Settings/VisibleColumnsSettingsMenu.inl
new file mode 100644
--- /dev/null
+++ b/Engine/Plugins/AudioInsights/Source/AudioInsights/Public/Settings/VisibleColumnsSettingsMenu.inl
@@ -0,0 +1,215 @@
+// Copyright Epic Games, Inc. All Rights Reserved.
+
+#pragma once
+
+#include "Algo/AnyOf.h"
+#include "Framework/MultiBox/MultiBoxBuilder.h"
+
+namespace UE::Audio::Insights
+{
+	template<TIsVisibleColumnsSettingsType T>
+	inline void FVisibleColumnsSettingsMenu<T>::BuildVisibleColumnsMenuContent(FMenuBuilder& OutMenuBuilder)
+	{
+		const TSharedPtr<SHeaderRow> HeaderRowWidgetShared = HeaderRowWidget.Pin();
+		if (!HeaderRowWidgetShared.IsValid())
+		{
+			return;
+		}
+
+		OutMenuBuilder.AddMenuEntry
+		(
+			TAttribute<FText>::CreateSP(this, &FVisibleColumnsSettingsMenu::GetToggleAllColumnsVisibilityText),
+			FText::GetEmpty(),
+			FSlateIcon(),
+			FUIAction
+			(
+				FExecuteAction::CreateSP(this, &FVisibleColumnsSettingsMenu::ToggleAllColumnsVisibility),
+				FCanExecuteAction::CreateLambda([]() { return true; }),
+				FGetActionCheckState::CreateSP(this, &FVisibleColumnsSettingsMenu::GetToggleAllColumnsVisibilityCheckState)
+			),
+			NAME_None,
+			EUserInterfaceActionType::ToggleButton
+		);
+
+		for (const SHeaderRow::FColumn& Column : HeaderRowWidgetShared->GetColumns())
+		{
+			OutMenuBuilder.AddMenuEntry
+			(
+				Column.DefaultText,
+				FText::GetEmpty(),
+				FSlateIcon(),
+				FUIAction
+				(
+					FExecuteAction::CreateSP(this, &FVisibleColumnsSettingsMenu::ToggleColumnVisibility, Column.ColumnId),
+					FCanExecuteAction::CreateLambda([]() { return true; }),
+					FGetActionCheckState::CreateSP(this, &FVisibleColumnsSettingsMenu::GetIsColumnVisibleCheckedState, Column.ColumnId)
+				),
+				NAME_None,
+				EUserInterfaceActionType::ToggleButton
+			);
+		}
+	}
+
+	template<TIsVisibleColumnsSettingsType T>
+	inline void FVisibleColumnsSettingsMenu<T>::OnHiddenColumnsListChanged()
+	{
+		const TSharedPtr<SHeaderRow> HeaderRowWidgetShared = HeaderRowWidget.Pin();
+		if (!HeaderRowWidgetShared.IsValid() || !VisibleColumnSettings)
+		{
+			return;
+		}
+
+		// Update and save settings to reload it on the next Editor sessions
+		for (const SHeaderRow::FColumn& Column : HeaderRowWidgetShared->GetColumns())
+		{
+			VisibleColumnSettings->SetIsVisible(Column.ColumnId, Column.bIsVisible);
+		}
+
+		OnVisibleColumnsSettingsUpdated.Broadcast();
+	}
+
+	template<TIsVisibleColumnsSettingsType T>
+	void FVisibleColumnsSettingsMenu<T>::ReadFromSettings(const T& InVisibleColumnSettings)
+	{
+		const UScriptStruct* VisibleColumnSettingsStruct = FVisibleColumnsSettings::StaticStruct();
+		const UScriptStruct* SettingsStruct = T::StaticStruct();
+
+		if (SettingsStruct->IsChildOf(VisibleColumnSettingsStruct))
+		{
+			// Null our VisibleColumnSettings TOptional while updating HeaderRowWidget so that we ignore OnHiddenColumnsListChanged callbacks:
+			VisibleColumnSettings.Reset();
+
+			UpdateColumnVisabilityFromSettings(InVisibleColumnSettings);
+
+			VisibleColumnSettings.Emplace();
+
+			T& CastStruct = static_cast<T&>(VisibleColumnSettings.GetValue());
+			CastStruct = InVisibleColumnSettings;
+		}
+	}
+
+	template<TIsVisibleColumnsSettingsType T>
+	void FVisibleColumnsSettingsMenu<T>::WriteToSettings(T& OutVisibleColumnSettings)
+	{
+		if (VisibleColumnSettings)
+		{
+			const UScriptStruct* VisibleColumnSettingsStruct = FVisibleColumnsSettings::StaticStruct();
+			const UScriptStruct* SettingsStruct = T::StaticStruct();
+
+			if (SettingsStruct->IsChildOf(VisibleColumnSettingsStruct))
+			{
+				T& CastStruct = static_cast<T&>(VisibleColumnSettings.GetValue());
+				OutVisibleColumnSettings = CastStruct;
+			}
+		}
+	}
+
+	template<TIsVisibleColumnsSettingsType T>
+	inline void FVisibleColumnsSettingsMenu<T>::ToggleColumnVisibility(FName ColumnId)
+	{
+		const TSharedPtr<SHeaderRow> HeaderRowWidgetShared = HeaderRowWidget.Pin();
+		if (!HeaderRowWidgetShared.IsValid())
+		{
+			return;
+		}
+
+		const bool bIsColumnVisible = HeaderRowWidgetShared->IsColumnVisible(ColumnId);
+		HeaderRowWidgetShared->SetShowGeneratedColumn(ColumnId, !bIsColumnVisible);
+	}
+
+	template<TIsVisibleColumnsSettingsType T>
+	inline void FVisibleColumnsSettingsMenu<T>::ToggleAllColumnsVisibility()
+	{
+		const TSharedPtr<SHeaderRow> HeaderRowWidgetShared = HeaderRowWidget.Pin();
+		if (!HeaderRowWidgetShared.IsValid())
+		{
+			return;
+		}
+
+		const ECheckBoxState CurrentState = GetToggleAllColumnsVisibilityCheckState();
+		const bool bVisibleIsDesired = CurrentState == ECheckBoxState::Unchecked || CurrentState == ECheckBoxState::Undetermined;
+
+		// Null our VisibleColumnSettings TOptional while updating HeaderRowWidget so that we ignore OnHiddenColumnsListChanged callbacks:
+		VisibleColumnSettings.Reset();
+
+		for (const SHeaderRow::FColumn& Column : HeaderRowWidgetShared->GetColumns())
+		{
+			if (Column.bIsVisible != bVisibleIsDesired)
+			{
+				HeaderRowWidgetShared->SetShowGeneratedColumn(Column.ColumnId, bVisibleIsDesired);
+			}
+		}
+
+		// Reinitialize the VisibleColumnSettings and perform a single OnHiddenColumnsListChanged call to update it and persist the settings:
+		VisibleColumnSettings.Emplace();
+		OnHiddenColumnsListChanged();
+	}
+
+	template<TIsVisibleColumnsSettingsType T>
+	inline void FVisibleColumnsSettingsMenu<T>::UpdateColumnVisabilityFromSettings(const FVisibleColumnsSettings& InVisibleColumnSettings)
+	{
+		const TSharedPtr<SHeaderRow> HeaderRowWidgetShared = HeaderRowWidget.Pin();
+		if (!HeaderRowWidgetShared.IsValid())
+		{
+			return;
+		}
+
+		for (const SHeaderRow::FColumn& Column : HeaderRowWidgetShared->GetColumns())
+		{
+			const bool bShowColumn = InVisibleColumnSettings.GetIsVisible(Column.ColumnId);
+			if (Column.bIsVisible != bShowColumn)
+			{
+				HeaderRowWidgetShared->SetShowGeneratedColumn(Column.ColumnId, bShowColumn);
+			}
+		}
+	}
+
+	template<TIsVisibleColumnsSettingsType T>
+	inline ECheckBoxState FVisibleColumnsSettingsMenu<T>::GetIsColumnVisibleCheckedState(FName ColumnId) const
+	{
+		return HeaderRowWidget.Pin().IsValid() ? HeaderRowWidget.Pin()->IsColumnVisible(ColumnId) ? ECheckBoxState::Checked : ECheckBoxState::Unchecked
+											   : ECheckBoxState::Undetermined;
+	}
+
+	template<TIsVisibleColumnsSettingsType T>
+	inline ECheckBoxState FVisibleColumnsSettingsMenu<T>::GetToggleAllColumnsVisibilityCheckState() const
+	{
+		const TSharedPtr<SHeaderRow> HeaderRowWidgetShared = HeaderRowWidget.Pin();
+		if (!HeaderRowWidgetShared.IsValid())
+		{
+			return ECheckBoxState::Undetermined;
+		}
+
+		const TIndirectArray<SHeaderRow::FColumn>& Columns = HeaderRowWidgetShared->GetColumns();
+		const bool bAnyColumnsVisible = Algo::AnyOf(Columns, [](const SHeaderRow::FColumn& Column) { return Column.bIsVisible; });
+		const bool bAnyColumnsInvisible = Algo::AnyOf(Columns, [](const SHeaderRow::FColumn& Column) { return !Column.bIsVisible; });
+
+		if (bAnyColumnsVisible && bAnyColumnsInvisible)
+		{
+			return ECheckBoxState::Undetermined;
+		}
+		else if (bAnyColumnsVisible)
+		{
+			return ECheckBoxState::Checked;
+		}
+		else
+		{
+			return ECheckBoxState::Unchecked;
+		}
+	}
+
+	template<TIsVisibleColumnsSettingsType T>
+	inline FText FVisibleColumnsSettingsMenu<T>::GetToggleAllColumnsVisibilityText() const
+	{
+		switch (GetToggleAllColumnsVisibilityCheckState())
+		{
+			case ECheckBoxState::Unchecked:
+			case ECheckBoxState::Undetermined:
+				return NSLOCTEXT("AudioInsights", "VisibleColumnsSettings_SelectAll", "Select: All");
+			case ECheckBoxState::Checked:
+				return NSLOCTEXT("AudioInsights", "VisibleColumnsSettings_SelectNone", "Select: None");
+			default:
+				return FText::GetEmpty();
+		}
+	}
+}
\ No newline at end of file
diff --git a/Engine/Plugins/AudioInsights/Source/AudioInsights/Public/Views/AudioEventLogDashboardViewFactory.h b/Engine/Plugins/AudioInsights/Source/AudioInsights/Public/Views/AudioEventLogDashboardViewFactory.h
--- a/Engine/Plugins/AudioInsights/Source/AudioInsights/Public/Views/AudioEventLogDashboardViewFactory.h
+++ b/Engine/Plugins/AudioInsights/Source/AudioInsights/Public/Views/AudioEventLogDashboardViewFactory.h
@@ -3,6 +3,7 @@
 
 #include "Filters/GenericFilter.h"
 #include "Settings/AudioEventLogSettings.h"
+#include "Settings/VisibleColumnsSettingsMenu.h"
 #include "Templates/SharedPointer.h"
 #include "Views/SAudioFilterBar.h"
 #include "Views/TableDashboardViewFactory.h"
@@ -87,11 +88,6 @@ namespace UE::Audio::Insights
 		bool PassesEventTypeFilter(const FEventLogFilterID EventLogID) const;
 
 		void BuildVisibleColumnsMenuContent(FMenuBuilder& MenuBuilder);
-		void ToggleColumnVisibility(FName ColumnId);
-		ECheckBoxState GetIsColumnVisibleCheckedState(FName ColumnId) const;
-		void ToggleAllColumnsVisibility();
-		ECheckBoxState GetToggleAllColumnsVisibilityCheckState() const;
-		FText GetToggleAllColumnsVisibilityText() const;
 
 		void SortByPredicate(TFunctionRef<bool(const class FAudioEventLogDashboardEntry&, const class FAudioEventLogDashboardEntry&)> Predicate);
 
@@ -111,7 +107,8 @@ namespace UE::Audio::Insights
 		TMap<FString, FEventLogFilterID> EventFilterTypes;
 
 		TSharedPtr<FUICommandList> CommandList;
-		TOptional<FAudioEventLogVisibleColumns> VisibleColumnSettings;
+		TSharedPtr<FVisibleColumnsSettingsMenu<FAudioEventLogVisibleColumns>> VisibleColumnsSettingsMenu;
+		TOptional<FAudioEventLogVisibleColumns> InitVisibleColumnSettings;
 	};
 } // namespace UE::Audio::Insights
 
diff --git a/Engine/Plugins/AudioInsights/Source/AudioInsights/Public/Views/SoundDashboardViewFactory.h b/Engine/Plugins/AudioInsights/Source/AudioInsights/Public/Views/SoundDashboardViewFactory.h
--- a/Engine/Plugins/AudioInsights/Source/AudioInsights/Public/Views/SoundDashboardViewFactory.h
+++ b/Engine/Plugins/AudioInsights/Source/AudioInsights/Public/Views/SoundDashboardViewFactory.h
@@ -3,6 +3,7 @@
 
 #include "Delegates/Delegate.h"
 #include "Settings/SoundDashboardSettings.h"
+#include "Settings/VisibleColumnsSettingsMenu.h"
 #include "Views/SoundPlotsWidgetView.h"
 #include "Views/TreeDashboardViewFactory.h"
 
@@ -160,11 +161,6 @@ namespace UE::Audio::Insights
 		UE_API void RefreshExpansionChecks(const TSharedPtr<IDashboardDataTreeViewEntry>& Entry);
 		UE_API void BuildAutoExpandMenuContent(FMenuBuilder& MenuBuilder);
 		UE_API void BuildVisibleColumnsMenuContent(FMenuBuilder& MenuBuilder);
-		UE_API void ToggleColumnVisibility(FName ColumnId);
-		UE_API ECheckBoxState GetIsColumnVisibleCheckedState(FName ColumnId) const;
-		UE_API void ToggleAllColumnsVisibility();
-		UE_API ECheckBoxState GetToggleAllColumnsVisibilityCheckState() const;
-		UE_API FText GetToggleAllColumnsVisibilityText() const;
 
 		UE_API void FilterText();
 		UE_API void RecursiveSort(TArray<TSharedPtr<IDashboardDataTreeViewEntry>>& OutTree, TFunctionRef<bool(const FSoundDashboardEntry&, const FSoundDashboardEntry&)> Predicate);
@@ -258,7 +254,8 @@ namespace UE::Audio::Insights
 		bool bIsMuteFilteredMode = false;
 		bool bIsSoloFilteredMode = false;
 
-		TOptional<FSoundDashboardVisibleColumns> VisibleColumnSettings;
+		TSharedPtr<FVisibleColumnsSettingsMenu<FSoundDashboardVisibleColumns>> VisibleColumnsSettingsMenu;
+		TOptional<FSoundDashboardVisibleColumns> InitVisibleColumnSettings;
 	};
 } // namespace UE::Audio::Insights
 
