[Audio Insights] Enabling AudioInsightsTimingViewExtender in the editor

diff --git a/Engine/Plugins/AudioInsights/Source/AudioInsights/Private/AudioInsightsModule.h b/Engine/Plugins/AudioInsights/Source/AudioInsights/Private/AudioInsightsModule.h
--- a/Engine/Plugins/AudioInsights/Source/AudioInsights/Private/AudioInsightsModule.h
+++ b/Engine/Plugins/AudioInsights/Source/AudioInsights/Private/AudioInsightsModule.h
@@ -2,6 +2,7 @@
 #pragma once 
 
 #include "AudioInsightsDashboardFactory.h"
+#include "AudioInsightsTimingViewExtender.h"
 #include "AudioInsightsTraceModule.h"
 #include "Framework/Docking/TabManager.h"
 #include "IAudioInsightsModule.h"
@@ -11,9 +12,6 @@
 #include "Views/DashboardViewFactory.h"
 #include "Widgets/Docking/SDockTab.h"
 
-#if !WITH_EDITOR
-#include "AudioInsightsTimingViewExtender.h"
-#endif // !WITH_EDITOR
 
 namespace UE::Audio::Insights
 {
@@ -44,9 +42,10 @@ namespace UE::Audio::Insights
 
 #if !WITH_EDITOR
 		TSharedPtr<FAudioInsightsComponent> GetAudioInsightsComponent() { return AudioInsightsComponent; };
-		FAudioInsightsTimingViewExtender& GetTimingViewExtender() { return AudioInsightsTimingViewExtender; };
 #endif // !WITH_EDITOR
 
+		FAudioInsightsTimingViewExtender& GetTimingViewExtender() { return AudioInsightsTimingViewExtender; };
+
 		TSharedRef<FDashboardFactory> GetDashboardFactory();
 		const TSharedRef<FDashboardFactory> GetDashboardFactory() const;
 
@@ -59,7 +58,8 @@ namespace UE::Audio::Insights
 
 #if !WITH_EDITOR
 		TSharedPtr<FAudioInsightsComponent> AudioInsightsComponent;
-		FAudioInsightsTimingViewExtender AudioInsightsTimingViewExtender;
 #endif // !WITH_EDITOR
+
+		FAudioInsightsTimingViewExtender AudioInsightsTimingViewExtender;
 	};
 } // namespace UE::Audio::Insights
diff --git a/Engine/Plugins/AudioInsights/Source/AudioInsights/Private/AudioInsightsTimingViewExtender.cpp b/Engine/Plugins/AudioInsights/Source/AudioInsights/Private/AudioInsightsTimingViewExtender.cpp
--- a/Engine/Plugins/AudioInsights/Source/AudioInsights/Private/AudioInsightsTimingViewExtender.cpp
+++ b/Engine/Plugins/AudioInsights/Source/AudioInsights/Private/AudioInsightsTimingViewExtender.cpp
@@ -40,12 +40,15 @@ namespace UE::Audio::Insights
 
 	void FAudioInsightsTimingViewExtender::SetTimeMarker(const double Timestamp)
 	{
-		if (TimingView == nullptr)
+		// If you set the time marker on the TimingView, it will Broadcast OnTimingViewTimeMarkerChanged. Otherwise we gotta do it ourselves.
+		if (TimingView != nullptr)
 		{
-			return;
+			TimingView->SetAndCenterOnTimeMarker(Timestamp);
+		}
+		else
+		{
+			OnTimingViewTimeMarkerChanged.Broadcast(Timestamp);
 		}
-
-		TimingView->SetAndCenterOnTimeMarker(Timestamp);
 	}
 
 	void FAudioInsightsTimingViewExtender::OnTimeMarkerChanged(UE::Insights::Timing::ETimeChangedFlags InFlags, double InTimeMarker)
diff --git a/Engine/Plugins/AudioInsights/Source/AudioInsights/Private/AudioInsightsTraceProviderBase.cpp b/Engine/Plugins/AudioInsights/Source/AudioInsights/Private/AudioInsightsTraceProviderBase.cpp
--- a/Engine/Plugins/AudioInsights/Source/AudioInsights/Private/AudioInsightsTraceProviderBase.cpp
+++ b/Engine/Plugins/AudioInsights/Source/AudioInsights/Private/AudioInsightsTraceProviderBase.cpp
@@ -2,34 +2,28 @@
 
 #include "AudioInsightsTraceProviderBase.h"
 
-#if !WITH_EDITOR
 #include "AudioInsightsModule.h"
 #include "AudioInsightsTimingViewExtender.h"
 #include "Modules/ModuleManager.h"
-#endif // !WITH_EDITOR
 
 namespace UE::Audio::Insights
 {
 	FTraceProviderBase::FTraceProviderBase(FName InName)
 		: Name(InName)
 	{
-#if !WITH_EDITOR
 		FAudioInsightsModule& AudioInsightsModule = FAudioInsightsModule::GetChecked();
 		FAudioInsightsTimingViewExtender& AudioInsightsTimingViewExtender = AudioInsightsModule.GetTimingViewExtender();
 		AudioInsightsTimingViewExtender.OnTimingViewTimeMarkerChanged.AddRaw(this, &FTraceProviderBase::OnTimingViewTimeMarkerChanged);
-#endif // !WITH_EDITOR
 	}
 
 	FTraceProviderBase::~FTraceProviderBase()
 	{
-#if !WITH_EDITOR
 		if (FModuleManager::Get().IsModuleLoaded("AudioInsights"))
 		{
 			FAudioInsightsModule& AudioInsightsModule = FAudioInsightsModule::GetChecked();
 			FAudioInsightsTimingViewExtender& AudioInsightsTimingViewExtender = AudioInsightsModule.GetTimingViewExtender();
 			AudioInsightsTimingViewExtender.OnTimingViewTimeMarkerChanged.RemoveAll(this);
 		}
-#endif // !WITH_EDITOR
 	}
 
 	FName FTraceProviderBase::GetName() const
diff --git a/Engine/Plugins/AudioInsights/Source/AudioInsights/Private/Providers/SoundTraceProvider.cpp b/Engine/Plugins/AudioInsights/Source/AudioInsights/Private/Providers/SoundTraceProvider.cpp
--- a/Engine/Plugins/AudioInsights/Source/AudioInsights/Private/Providers/SoundTraceProvider.cpp
+++ b/Engine/Plugins/AudioInsights/Source/AudioInsights/Private/Providers/SoundTraceProvider.cpp
@@ -1060,9 +1060,9 @@ namespace UE::Audio::Insights
 		AddSingleValueToBuffer(ActiveSoundEntry.RelativeRenderCostDataPoints, CumulativeRelativeRenderCost.TotalAccumulationDataPoint);
 	}
 
-#if !WITH_EDITOR
 	void FSoundTraceProvider::OnTimingViewTimeMarkerChanged(double TimeMarker)
 	{
+#if !WITH_EDITOR
 		using namespace FSoundTraceProviderPrivate;
 
 		if (!SessionCachedMessages.IsValid())
@@ -1389,10 +1389,13 @@ namespace UE::Audio::Insights
 			}
 		}
 
+#endif //!WITH_EDITOR
+
 		// Call parent method to update LastMessageId
 		FTraceProviderBase::OnTimingViewTimeMarkerChanged(TimeMarker);
 	}
 
+#if !WITH_EDITOR
 	void FSoundTraceProvider::CacheIsPlottingFlagRecursive(const TSharedPtr<IDashboardDataTreeViewEntry>& Entry)
 	{
 		using namespace FSoundTraceProviderPrivate;
diff --git a/Engine/Plugins/AudioInsights/Source/AudioInsights/Private/Providers/SubmixTraceProvider.cpp b/Engine/Plugins/AudioInsights/Source/AudioInsights/Private/Providers/SubmixTraceProvider.cpp
--- a/Engine/Plugins/AudioInsights/Source/AudioInsights/Private/Providers/SubmixTraceProvider.cpp
+++ b/Engine/Plugins/AudioInsights/Source/AudioInsights/Private/Providers/SubmixTraceProvider.cpp
@@ -334,9 +334,11 @@ namespace UE::Audio::Insights
 	{
 		SessionCachedMessages = MakeUnique<FSubmixSessionCachedMessages>(InSession);
 	}
+#endif // WITH_EDITOR
 
 	void FSubmixTraceProvider::OnTimingViewTimeMarkerChanged(double TimeMarker)
 	{
+#if !WITH_EDITOR
 		using namespace FSubmixTraceProviderPrivate;
 
 		if (!SessionCachedMessages.IsValid())
@@ -437,9 +439,9 @@ namespace UE::Audio::Insights
 			}
 		}
 
+#endif // !WITH_EDITOR
+
 		// Call parent method to update LastMessageId
 		FTraceProviderBase::OnTimingViewTimeMarkerChanged(TimeMarker);
 	}
-#endif // WITH_EDITOR
-
 } // namespace UE::Audio::Insights
diff --git a/Engine/Plugins/AudioInsights/Source/AudioInsights/Private/Views/AudioEventLogDashboardViewFactory.cpp b/Engine/Plugins/AudioInsights/Source/AudioInsights/Private/Views/AudioEventLogDashboardViewFactory.cpp
--- a/Engine/Plugins/AudioInsights/Source/AudioInsights/Private/Views/AudioEventLogDashboardViewFactory.cpp
+++ b/Engine/Plugins/AudioInsights/Source/AudioInsights/Private/Views/AudioEventLogDashboardViewFactory.cpp
@@ -159,10 +159,9 @@ namespace UE::Audio::Insights
 	{
 #if WITH_EDITOR
 		FAudioEventLogSettings::OnRequestReadSettings.Broadcast();
-#else
+#endif // WITH_EDITOR
 		FAudioInsightsModule& AudioInsightsModule = FAudioInsightsModule::GetChecked();
 		AudioInsightsModule.GetTimingViewExtender().OnTimingViewTimeMarkerChanged.AddSP(this, &FAudioEventLogDashboardViewFactory::OnTimingViewTimeMarkerChanged);
-#endif // WITH_EDITOR
 
 		if (!DashboardWidget.IsValid())
 		{
@@ -497,12 +496,8 @@ namespace UE::Audio::Insights
 
 		if (SelectInfo != ESelectInfo::Type::Direct)
 		{
-#if WITH_EDITOR
-			// TODO UE-301957: Set timemarker via delegate
-#else
 			FAudioInsightsModule& AudioInsightsModule = FAudioInsightsModule::GetChecked();
 			AudioInsightsModule.GetTimingViewExtender().SetTimeMarker(SelectedEventLogEntry.Timestamp);
-#endif // WITH_EDITOR
 		}
 	}
 
@@ -622,7 +617,6 @@ namespace UE::Audio::Insights
 	}
 #endif // WITH_EDITOR
 
-#if !WITH_EDITOR
 	void FAudioEventLogDashboardViewFactory::OnTimingViewTimeMarkerChanged(double InTimeMarker)
 	{
 		using namespace AudioEventLogPrivate;
@@ -632,6 +626,7 @@ namespace UE::Audio::Insights
 			return;
 		}
 
+#if !WITH_EDITOR
 		// We assume the data entries are in timestamp order (earliest to latest)
 		for (int32 Index = 0; Index < DataViewEntries.Num(); ++Index)
 		{
@@ -668,8 +663,8 @@ namespace UE::Audio::Insights
 				FilteredEntriesListView->SetSelection(Entry);
 			}
 		}
-	}
 #endif // !WITH_EDITOR
+	}
 
 } // namespace UE::Audio::Insights
 
diff --git a/Engine/Plugins/AudioInsights/Source/AudioInsights/Private/Views/SoundDashboardViewFactory.cpp b/Engine/Plugins/AudioInsights/Source/AudioInsights/Private/Views/SoundDashboardViewFactory.cpp
--- a/Engine/Plugins/AudioInsights/Source/AudioInsights/Private/Views/SoundDashboardViewFactory.cpp
+++ b/Engine/Plugins/AudioInsights/Source/AudioInsights/Private/Views/SoundDashboardViewFactory.cpp
@@ -1089,10 +1089,10 @@ namespace UE::Audio::Insights
 	{
 #if WITH_EDITOR
 		FSoundDashboardSettings::OnRequestReadSettings.Broadcast();
-#else
+#endif // WITH_EDITOR
+
 		FAudioInsightsModule& AudioInsightsModule = FAudioInsightsModule::GetChecked();
 		AudioInsightsModule.GetTimingViewExtender().OnTimingViewTimeMarkerChanged.AddSP(this, &FSoundDashboardViewFactory::OnTimingViewTimeMarkerChanged);
-#endif // WITH_EDITOR
 
 		const TSharedRef<SWidget> SoundDashboardWidget = SNew(SVerticalBox)
 #if WITH_EDITOR
@@ -2985,12 +2985,12 @@ namespace UE::Audio::Insights
 
 		return Objects;
 	}
-#else
+#endif // WITH_EDITOR
+
 	void FSoundDashboardViewFactory::OnTimingViewTimeMarkerChanged(double InTimeMarker)
 	{
 		OnUpdatePlotVisibility.Broadcast(DataViewEntries);
 	}
-#endif // WITH_EDITOR
 
 	TSharedRef<SWidget> FSoundDashboardViewFactory::MakeShowPlotWidget()
 	{
@@ -3125,9 +3125,7 @@ namespace UE::Audio::Insights
 			}
 		}
 
-#if !WITH_EDITOR
 		OnUpdatePlotVisibility.Broadcast(DataViewEntries);
-#endif // !WITH_EDITOR
 	}
 
 	void FSoundDashboardViewFactory::SetShowPlotRecursive(const TSharedPtr<IDashboardDataTreeViewEntry>& Entry, const bool bShowPlot)
diff --git a/Engine/Plugins/AudioInsights/Source/AudioInsights/Public/Views/AudioEventLogDashboardViewFactory.h b/Engine/Plugins/AudioInsights/Source/AudioInsights/Public/Views/AudioEventLogDashboardViewFactory.h
--- a/Engine/Plugins/AudioInsights/Source/AudioInsights/Public/Views/AudioEventLogDashboardViewFactory.h
+++ b/Engine/Plugins/AudioInsights/Source/AudioInsights/Public/Views/AudioEventLogDashboardViewFactory.h
@@ -96,9 +96,7 @@ namespace UE::Audio::Insights
 		void OnWriteEditorSettings(FAudioEventLogSettings& OutSettings);
 #endif // WITH_EDITOR
 
-#if !WITH_EDITOR
 		void OnTimingViewTimeMarkerChanged(double InTimeMarker);
-#endif // !WITH_EDITOR
 
 		TSharedPtr<SWidget> DashboardWidget;
 		TSharedPtr<SAudioFilterBar<FEventLogFilterID>> FilterBar;
diff --git a/Engine/Plugins/AudioInsights/Source/AudioInsights/Public/Views/SoundDashboardViewFactory.h b/Engine/Plugins/AudioInsights/Source/AudioInsights/Public/Views/SoundDashboardViewFactory.h
--- a/Engine/Plugins/AudioInsights/Source/AudioInsights/Public/Views/SoundDashboardViewFactory.h
+++ b/Engine/Plugins/AudioInsights/Source/AudioInsights/Public/Views/SoundDashboardViewFactory.h
@@ -117,13 +117,11 @@ namespace UE::Audio::Insights
 		DECLARE_MULTICAST_DELEGATE_TwoParams(FOnProcessPlotData, const TArray<TSharedPtr<IDashboardDataTreeViewEntry>>& /*DataViewEntries*/, const TArray<TSharedPtr<IDashboardDataTreeViewEntry>>& /*SelectedEntries*/);
 		FOnProcessPlotData OnProcessPlotData;
 
-#if !WITH_EDITOR
 		DECLARE_MULTICAST_DELEGATE_OneParam(FOnUpdatePlotVisibility, const TArray<TSharedPtr<IDashboardDataTreeViewEntry>>& /*DataViewEntries*/);
 		FOnUpdatePlotVisibility OnUpdatePlotVisibility;
 
 		DECLARE_MULTICAST_DELEGATE_OneParam(FOnUpdatePlotSelection, const TArray<TSharedPtr<IDashboardDataTreeViewEntry>>& /*SelectedEntries*/);
 		FOnUpdatePlotSelection OnUpdatePlotSelection;
-#endif // !WITH_EDITOR
 
 	private:
 		UE_API void BindCommands();
@@ -187,10 +185,11 @@ namespace UE::Audio::Insights
 		UE_API void MuteSoloFilteredEntries();
 
 		UE_API TArray<TObjectPtr<UObject>> GetSelectedEditableAssets() const;
-#else
-		void OnTimingViewTimeMarkerChanged(double InTimeMarker);
 #endif // WITH_EDITOR
 
+		void OnTimingViewTimeMarkerChanged(double InTimeMarker);
+
+
 		UE_API TSharedRef<SWidget> MakeShowPlotWidget();
 
 		UE_API TSharedRef<SWidget> CreateShowPlotColumnButton(const TSharedRef<FTraceTreeDashboardViewFactory::SRowWidget>& InRowWidget,
diff --git a/Engine/Plugins/AudioInsights/Source/AudioInsights/Public/Views/SoundPlotsWidgetView.h b/Engine/Plugins/AudioInsights/Source/AudioInsights/Public/Views/SoundPlotsWidgetView.h
--- a/Engine/Plugins/AudioInsights/Source/AudioInsights/Public/Views/SoundPlotsWidgetView.h
+++ b/Engine/Plugins/AudioInsights/Source/AudioInsights/Public/Views/SoundPlotsWidgetView.h
@@ -54,9 +54,10 @@ namespace UE::Audio::Insights
 #else
 		void OnAudioInsightsComponentTabSpawn();
 		void OnSessionAnalysisCompleted();
-		void OnTimingViewTimeMarkerChanged(double InTimeMarker);
 #endif // WITH_EDITOR
 
+		void OnTimingViewTimeMarkerChanged(double InTimeMarker);
+
 		TSharedRef<SWidget> MakePlotsWidget();
 		void ConstructPlotsWidgetContents();
 
@@ -106,10 +107,8 @@ namespace UE::Audio::Insights
 			Running,
 			Stopped,
 			Paused,
-#if !WITH_EDITOR
 			LoadingTraceFile,
 			ControlledByTimemarker
-#endif
 		};
 
 		EGameState GameState = EGameState::Stopped;
@@ -121,9 +120,9 @@ namespace UE::Audio::Insights
 #if !WITH_EDITOR
 		double PreviousTime = 0.0;
 		double CurrentRangeUpperBound = 0.0;
+#endif // !WITH_EDITOR
 
 		bool bTimeMarkerMoved = false;
-#endif // !WITH_EDITOR
 	};
 } // namespace UE::Audio::Insights
 
