[Audio Insights] Improve error reporting from VisibleColumnSettings to distinguish between properties not existing and properties existing but being the wrong type (bug hawk suggestion)

diff --git a/Engine/Plugins/AudioInsights/Source/AudioInsights/Private/Settings/VisibleColumnsSettings.cpp b/Engine/Plugins/AudioInsights/Source/AudioInsights/Private/Settings/VisibleColumnsSettings.cpp
--- a/Engine/Plugins/AudioInsights/Source/AudioInsights/Private/Settings/VisibleColumnsSettings.cpp
+++ b/Engine/Plugins/AudioInsights/Source/AudioInsights/Private/Settings/VisibleColumnsSettings.cpp
@@ -10,9 +10,16 @@ bool FVisibleColumnsSettings::GetIsVisible(const FName& ColumnId) const
 		{
 			return BoolProperty->GetPropertyValue_InContainer(this);
 		}
+		else
+		{
+			ensureMsgf(false, TEXT("VisibleColumns setting for ColumnId %s found, but is incorrectly typed (should be a boolean)."), *ColumnId.ToString());
+		}
 	}
-
-	ensureMsgf(false, TEXT("No VisibleColumns setting for ColumnId: %s"), *ColumnId.ToString());
+	else
+	{
+		ensureMsgf(false, TEXT("No VisibleColumns setting for ColumnId: %s"), *ColumnId.ToString());
+	}
+	
 	return false;
 }
 
@@ -25,5 +32,13 @@ void FVisibleColumnsSettings::SetIsVisible(const FName& ColumnId, const bool bIs
 		{
 			BoolProperty->SetPropertyValue_InContainer(this, bIsVisible);
 		}
+		else
+		{
+			ensureMsgf(false, TEXT("VisibleColumns setting for ColumnId %s found, but is incorrectly typed (should be a boolean)."), *ColumnId.ToString());
+		}
+	}
+	else
+	{
+		ensureMsgf(false, TEXT("No VisibleColumns setting for ColumnId: %s"), *ColumnId.ToString());
 	}
 }
\ No newline at end of file
