Ensuring SHeaderRow is defined in inline file, fixing non-unity build error.

diff --git a/Engine/Plugins/AudioInsights/Source/AudioInsights/Public/Settings/VisibleColumnsSettingsMenu.inl b/Engine/Plugins/AudioInsights/Source/AudioInsights/Public/Settings/VisibleColumnsSettingsMenu.inl
--- a/Engine/Plugins/AudioInsights/Source/AudioInsights/Public/Settings/VisibleColumnsSettingsMenu.inl
+++ b/Engine/Plugins/AudioInsights/Source/AudioInsights/Public/Settings/VisibleColumnsSettingsMenu.inl
@@ -4,6 +4,7 @@
 
 #include "Algo/AnyOf.h"
 #include "Framework/MultiBox/MultiBoxBuilder.h"
+#include "Widgets/Views/SHeaderRow.h"
 
 namespace UE::Audio::Insights
 {
