[Audio Insights] Make which event types are being filtered by persistent in the event log - these are hidden in the editor settings and are only controllable in the event log itself. Also moved OnRequestReadSettings to after the widget has been created - this avoids us having to cache the initial settings.

diff --git a/Engine/Plugins/AudioInsights/Source/AudioInsights/Private/Views/AudioEventLogDashboardViewFactory.cpp b/Engine/Plugins/AudioInsights/Source/AudioInsights/Private/Views/AudioEventLogDashboardViewFactory.cpp
--- a/Engine/Plugins/AudioInsights/Source/AudioInsights/Private/Views/AudioEventLogDashboardViewFactory.cpp
+++ b/Engine/Plugins/AudioInsights/Source/AudioInsights/Private/Views/AudioEventLogDashboardViewFactory.cpp
@@ -157,9 +157,6 @@ namespace UE::Audio::Insights
 
 	TSharedRef<SWidget> FAudioEventLogDashboardViewFactory::MakeWidget(TSharedRef<SDockTab> OwnerTab, const FSpawnTabArgs& SpawnTabArgs)
 	{
-#if WITH_EDITOR
-		FAudioEventLogSettings::OnRequestReadSettings.Broadcast();
-#endif // WITH_EDITOR
 		FAudioInsightsModule& AudioInsightsModule = FAudioInsightsModule::GetChecked();
 		AudioInsightsModule.GetTimingViewExtender().OnTimingViewTimeMarkerChanged.AddSP(this, &FAudioEventLogDashboardViewFactory::OnTimingViewTimeMarkerChanged);
 
@@ -174,8 +171,7 @@ namespace UE::Audio::Insights
 
 			if (HeaderRowWidget.IsValid())
 			{
-				const FAudioEventLogVisibleColumns VisibleColumns = InitVisibleColumnSettings.IsSet() ? InitVisibleColumnSettings.GetValue() : FAudioEventLogVisibleColumns();
-				VisibleColumnsSettingsMenu = MakeShared<FVisibleColumnsSettingsMenu<FAudioEventLogVisibleColumns>>(HeaderRowWidget.ToSharedRef(), VisibleColumns);
+				VisibleColumnsSettingsMenu = MakeShared<FVisibleColumnsSettingsMenu<FAudioEventLogVisibleColumns>>(HeaderRowWidget.ToSharedRef(), FAudioEventLogVisibleColumns());
 
 				VisibleColumnsSettingsMenu->OnVisibleColumnsSettingsUpdated.AddSPLambda(this, []()
 				{
@@ -186,6 +182,11 @@ namespace UE::Audio::Insights
 			}
 		}
 
+#if WITH_EDITOR
+		// Read the editor settings after the widget has been created
+		FAudioEventLogSettings::OnRequestReadSettings.Broadcast();
+#endif // WITH_EDITOR
+
 		return DashboardWidget->AsShared();
 	}
 
@@ -232,6 +233,13 @@ namespace UE::Audio::Insights
 
 			return false;
 		});
+
+#if WITH_EDITOR
+		if (Reason == FTraceTableDashboardViewFactory::EProcessReason::FilterUpdated)
+		{
+			FAudioEventLogSettings::OnRequestWriteSettings.Broadcast();
+		}
+#endif // WITH_EDITOR
 	}
 
 	const TMap<FName, FTraceTableDashboardViewFactory::FColumnData>& FAudioEventLogDashboardViewFactory::GetColumns() const
@@ -240,7 +248,11 @@ namespace UE::Audio::Insights
 
 		auto CreateColumnData = [this]()
 		{
-			FAudioEventLogVisibleColumns VisibleColumns = InitVisibleColumnSettings.IsSet() ? InitVisibleColumnSettings.GetValue() : FAudioEventLogVisibleColumns();
+			FAudioEventLogVisibleColumns VisibleColumns;
+			if (VisibleColumnsSettingsMenu.IsValid())
+			{
+				VisibleColumnsSettingsMenu->WriteToSettings(VisibleColumns);
+			}
 
 			return TMap<FName, FTraceTableDashboardViewFactory::FColumnData>
 			{
@@ -597,19 +609,54 @@ namespace UE::Audio::Insights
 #if WITH_EDITOR
 	void FAudioEventLogDashboardViewFactory::OnReadEditorSettings(const FAudioEventLogSettings& InSettings)
 	{
-		if (VisibleColumnsSettingsMenu.IsValid())
+		if (FilterBar.IsValid())
 		{
-			VisibleColumnsSettingsMenu->ReadFromSettings(InSettings.VisibleColumns);
+			if (InSettings.EventFilters.IsEmpty())
+			{
+				FilterBar->RemoveAllFilters();
+			}
+			else
+			{
+				for (const auto& [EventName, FilterID] : EventFilterTypes)
+				{
+					const bool bFilterIsActive = InSettings.EventFilters.Contains(EventName);
+					FilterBar->SetFilterCheckState(FilterBar->GetFilter(EventName), bFilterIsActive ? ECheckBoxState::Checked : ECheckBoxState::Unchecked);
+				}
+			}
 		}
-		
-		if (!InitVisibleColumnSettings.IsSet())
+
+		if (VisibleColumnsSettingsMenu.IsValid())
 		{
-			InitVisibleColumnSettings = InSettings.VisibleColumns;
+			VisibleColumnsSettingsMenu->ReadFromSettings(InSettings.VisibleColumns);
 		}
 	}
 
 	void FAudioEventLogDashboardViewFactory::OnWriteEditorSettings(FAudioEventLogSettings& OutSettings)
 	{
+		if (FilterBar.IsValid())
+		{
+			TSharedPtr<TFilterCollection<FEventLogFilterID>> ActiveFilters = FilterBar->GetAllActiveFilters();
+
+			if (!ActiveFilters.IsValid() || ActiveFilters->Num() == 0)
+			{
+				OutSettings.EventFilters.Empty();
+			}
+			else
+			{
+				for (const auto& [EventName, FilterID] : EventFilterTypes)
+				{
+					if (PassesEventTypeFilter(FilterID))
+					{
+						OutSettings.EventFilters.Add(EventName);
+					}
+					else
+					{
+						OutSettings.EventFilters.Remove(EventName);
+					}
+				}
+			}
+		}
+
 		if (VisibleColumnsSettingsMenu.IsValid())
 		{
 			VisibleColumnsSettingsMenu->WriteToSettings(OutSettings.VisibleColumns);
diff --git a/Engine/Plugins/AudioInsights/Source/AudioInsights/Public/Settings/AudioEventLogSettings.h b/Engine/Plugins/AudioInsights/Source/AudioInsights/Public/Settings/AudioEventLogSettings.h
--- a/Engine/Plugins/AudioInsights/Source/AudioInsights/Public/Settings/AudioEventLogSettings.h
+++ b/Engine/Plugins/AudioInsights/Source/AudioInsights/Public/Settings/AudioEventLogSettings.h
@@ -51,6 +51,11 @@ struct FAudioEventLogSettings
 {
 	GENERATED_BODY()
 
+	/** Which event types are being filtered by in the Event Log - hidden in the Editor Preferences */
+	UPROPERTY(config, meta = (HideInDetailPanel))
+	TSet<FString> EventFilters;
+
+	/** Visibility of columns in the Event Log tab */
 	UPROPERTY(EditAnywhere, config, Category = EventLog)
 	FAudioEventLogVisibleColumns VisibleColumns;
 
diff --git a/Engine/Plugins/AudioInsights/Source/AudioInsights/Public/Views/AudioEventLogDashboardViewFactory.h b/Engine/Plugins/AudioInsights/Source/AudioInsights/Public/Views/AudioEventLogDashboardViewFactory.h
--- a/Engine/Plugins/AudioInsights/Source/AudioInsights/Public/Views/AudioEventLogDashboardViewFactory.h
+++ b/Engine/Plugins/AudioInsights/Source/AudioInsights/Public/Views/AudioEventLogDashboardViewFactory.h
@@ -106,7 +106,6 @@ namespace UE::Audio::Insights
 
 		TSharedPtr<FUICommandList> CommandList;
 		TSharedPtr<FVisibleColumnsSettingsMenu<FAudioEventLogVisibleColumns>> VisibleColumnsSettingsMenu;
-		TOptional<FAudioEventLogVisibleColumns> InitVisibleColumnSettings;
 	};
 } // namespace UE::Audio::Insights
 
