[Audio Insights] Add user defined Y-Axis range controls to the Plots widget

diff --git a/Engine/Plugins/AudioInsights/Source/AudioInsights/Private/Views/SAudioCurveView.cpp b/Engine/Plugins/AudioInsights/Source/AudioInsights/Private/Views/SAudioCurveView.cpp
--- a/Engine/Plugins/AudioInsights/Source/AudioInsights/Private/Views/SAudioCurveView.cpp
+++ b/Engine/Plugins/AudioInsights/Source/AudioInsights/Private/Views/SAudioCurveView.cpp
@@ -17,13 +17,14 @@
 void SAudioCurveView::Construct( const SAudioCurveView::FArguments& InArgs )
 {
 	// A lot of this should go into a style 
-    ViewRange = InArgs._ViewRange;
-    GridLineColor = InArgs._GridLineColor;
-    AxesLabelColor = InArgs._AxesLabelColor;
+	ViewRange = InArgs._ViewRange;
+	AutoRangeYAxis = InArgs._AutoRangeYAxis;
+	GridLineColor = InArgs._GridLineColor;
+	AxesLabelColor = InArgs._AxesLabelColor;
 	YMargin = FMath::Clamp(InArgs._YMargin.Get(), 0.0f, 0.5f);
 	HorizontalAxisIncrement = InArgs._HorizontalAxisIncrement;
 	DesiredSize = InArgs._DesiredSize;
-	YDataRange = FVector2f(0.0f, 1.0f);
+	YDataRange = FFloatInterval(0.0f, 1.0f);
 	NumHorizontalGridLines = 10;
 
 	// set clipping on by default, since the OnPaint function is drawing outside the bounds
@@ -198,12 +199,12 @@ void SAudioCurveView::UpdateYDataRange()
 
 	if (MinValue < TNumericLimits<float>::Max())
 	{
-		YDataRange.X = MinValue;
+		YDataRange.Min = MinValue;
 	}
 	
 	if (MaxValue > TNumericLimits<float>::Lowest())
 	{
-		YDataRange.Y = MaxValue;
+		YDataRange.Max = MaxValue;
 	}
 }
 
@@ -260,15 +261,19 @@ void SAudioCurveView::UpdateYDataRangeFromTimestampRange(const double InLowerBou
 		MaxValue = MinValue + Epsilon;
 	}
 
-	YDataRange.X = MinValue;
-	YDataRange.Y = MaxValue;
+	YDataRange.Min = MinValue;
+	YDataRange.Max = MaxValue;
 }
 #endif // !WITH_EDITOR
 
 void SAudioCurveView::SetCurvesPointData(TSharedPtr<TMap<int32, TArray<FCurvePoint>>> InPointDataPerCurve)
 {
 	PointDataPerCurve = InPointDataPerCurve;
-	UpdateYDataRange();
+
+	if (AutoRangeYAxis.Get())
+	{
+		UpdateYDataRange();
+	}
 }
 
 void SAudioCurveView::SetCurvesMetadata(TSharedPtr<TMap<int32, FCurveMetadata>> InMetadataPerCurve)
@@ -281,7 +286,7 @@ void SAudioCurveView::SetYValueFormattingOptions(const FNumberFormattingOptions
 	YValueFormattingOptions = InYValueFormattingOptions;
 }
 
-void SAudioCurveView::SetInitYAxisRange(const FVector2f& InRange)
+void SAudioCurveView::SetYAxisRange(const FFloatInterval& InRange)
 {
 	YDataRange = InRange;
 }
@@ -293,12 +298,12 @@ float SAudioCurveView::ValueToLocalY(const FVector2f AllottedLocalSize, const fl
 	const float MarginTop = AllottedLocalSize.Y - MarginBase;
 	
 	// Special case to add padding based on YMargin if YDataRange min/max are the same 
-	if (FMath::IsNearlyEqual(YDataRange.X, YDataRange.Y))
+	if (FMath::IsNearlyEqual(YDataRange.Min, YDataRange.Max))
 	{
-		const FVector2f PaddedYDataRange = FVector2f(YDataRange.Y * (1.0f - YMargin.Get()), YDataRange.X * (1.0f + YMargin.Get()));
+		const FVector2f PaddedYDataRange = FVector2f(YDataRange.Max * (1.0f - YMargin.Get()), YDataRange.Min * (1.0f + YMargin.Get()));
 		return FMath::GetMappedRangeValueUnclamped(PaddedYDataRange, FVector2f(MarginTop, MarginBase), Value);
 	}
-	return FMath::GetMappedRangeValueUnclamped(YDataRange, FVector2f(MarginTop, MarginBase), Value);
+	return FMath::GetMappedRangeValueUnclamped(FVector2f(YDataRange.Min, YDataRange.Max), FVector2f(MarginTop, MarginBase), Value);
 }
 
 float SAudioCurveView::LocalYToValue(const FVector2f AllottedLocalSize, const float LocalY) const
@@ -308,12 +313,12 @@ float SAudioCurveView::LocalYToValue(const FVector2f AllottedLocalSize, const fl
 	const float MarginTop = AllottedLocalSize.Y - MarginBase;
 
 	// Special case to add padding based on YMargin if YDataRange min/max are the same 
-	if (FMath::IsNearlyEqual(YDataRange.X, YDataRange.Y))
+	if (FMath::IsNearlyEqual(YDataRange.Min, YDataRange.Max))
 	{
-		const FVector2f PaddedYDataRange = FVector2f(YDataRange.Y * (1.0f - YMargin.Get()), YDataRange.X * (1.0f + YMargin.Get()));
+		const FVector2f PaddedYDataRange = FVector2f(YDataRange.Max * (1.0f - YMargin.Get()), YDataRange.Min * (1.0f + YMargin.Get()));
 		return FMath::GetMappedRangeValueUnclamped(FVector2f(MarginTop, MarginBase), PaddedYDataRange, LocalY);
 	}
-	return FMath::GetMappedRangeValueUnclamped(FVector2f(MarginTop, MarginBase), YDataRange, LocalY);
+	return FMath::GetMappedRangeValueUnclamped(FVector2f(MarginTop, MarginBase), FVector2f(YDataRange.Min, YDataRange.Max), LocalY);
 }
 
 int32 SAudioCurveView::PaintGridLines(const FGeometry& AllottedGeometry, const FSlateRect& MyCullingRect, FSlateWindowElementList& OutDrawElements, int32 LayerId, const FWidgetStyle& InWidgetStyle, bool bParentEnabled, const SSimpleTimeSlider::FScrubRangeToScreen& RangeToScreen) const
diff --git a/Engine/Plugins/AudioInsights/Source/AudioInsights/Private/Views/SoundDashboardViewFactory.cpp b/Engine/Plugins/AudioInsights/Source/AudioInsights/Private/Views/SoundDashboardViewFactory.cpp
--- a/Engine/Plugins/AudioInsights/Source/AudioInsights/Private/Views/SoundDashboardViewFactory.cpp
+++ b/Engine/Plugins/AudioInsights/Source/AudioInsights/Private/Views/SoundDashboardViewFactory.cpp
@@ -39,6 +39,16 @@ namespace UE::Audio::Insights
 			return static_cast<FSoundDashboardEntry&>(InData);
 		};
 
+		void OnDefaultPlotRangesChanged(bool bRequestWriteSettings = true)
+		{
+#if WITH_EDITOR
+			if (bRequestWriteSettings)
+			{
+				FSoundDashboardSettings::OnRequestWriteSettings.Broadcast();
+			}
+#endif // WITH_EDITOR
+		}
+
 		float GetLastEntryArrayValue(const ::Audio::TCircularAudioBuffer<FDataPoint>& InDataPoints)
 		{
 			if (InDataPoints.Num() > 0)
@@ -2738,6 +2748,8 @@ namespace UE::Audio::Insights
 			InitVisibleColumnSettings = InSettings.VisibleColumns;
 		}
 
+		DefaultPlotRanges = InSettings.DefaultPlotRanges;
+
 		if (RecentlyStoppedSoundTimeout != InSettings.StoppedSoundTimeoutTime)
 		{
 			RecentlyStoppedSoundTimeout = InSettings.StoppedSoundTimeoutTime;
@@ -2758,6 +2770,8 @@ namespace UE::Audio::Insights
 			VisibleColumnsSettingsMenu->WriteToSettings(OutSettings.VisibleColumns);
 		}
 
+		OutSettings.DefaultPlotRanges = DefaultPlotRanges;
+
 		if (OutSettings.StoppedSoundTimeoutTime != RecentlyStoppedSoundTimeout)
 		{
 			OutSettings.StoppedSoundTimeoutTime = RecentlyStoppedSoundTimeout;
@@ -3192,7 +3206,43 @@ namespace UE::Audio::Insights
 						[](const IDashboardDataTreeViewEntry& InEntry) -> const ::Audio::TCircularAudioBuffer<FDataPoint>& { return CastEntry(InEntry).AmplitudeDataPoints; },
 						[this](const float InValue) { return bDisplayAmpPeakInDb ? ::Audio::ConvertToDecibels(InValue) : InValue; },
 						[this]() { return bDisplayAmpPeakInDb ? FSlateStyle::Get().GetAmpFloatFormat() : FSlateStyle::Get().GetLinearVolumeFloatFormat(); },
-						[this]() { return bDisplayAmpPeakInDb ? FVector2f(-160.0f, 0.0f) : FVector2f(0.0f, 1.0f); }
+						[this]()
+						{
+							if (bDisplayAmpPeakInDb)
+							{
+								return FSoundPlotsDashboardPlotRanges::StaticStruct()->FindPropertyByName(GET_MEMBER_NAME_CHECKED(FSoundPlotsDashboardPlotRanges, Amplitude_dB));
+							}
+							else
+							{
+								return FSoundPlotsDashboardPlotRanges::StaticStruct()->FindPropertyByName(GET_MEMBER_NAME_CHECKED(FSoundPlotsDashboardPlotRanges, Amplitude_Linear));
+							}
+						},
+						[this]() { return bDisplayAmpPeakInDb ? DefaultPlotRanges.bUseCustomAmplitude_dBRange : DefaultPlotRanges.bUseCustomAmplitude_LinearRange; },
+						[this](bool bUseCustomYRange)
+						{
+							if (bDisplayAmpPeakInDb)
+							{
+								DefaultPlotRanges.bUseCustomAmplitude_dBRange = bUseCustomYRange;
+							}
+							else
+							{
+								DefaultPlotRanges.bUseCustomAmplitude_LinearRange = bUseCustomYRange;
+							}
+							OnDefaultPlotRangesChanged();
+						},
+						[this]() { return bDisplayAmpPeakInDb ? DefaultPlotRanges.Amplitude_dB : DefaultPlotRanges.Amplitude_Linear; },
+						[this](const FFloatInterval& Range, TOptional<ETextCommit::Type> CommitType)
+						{
+							if (bDisplayAmpPeakInDb)
+							{
+								DefaultPlotRanges.Amplitude_dB = Range;
+							}
+							else
+							{
+								DefaultPlotRanges.Amplitude_Linear = Range;
+							}
+							OnDefaultPlotRangesChanged(CommitType.IsSet());
+						}
 					}
 				},
 				{
@@ -3202,7 +3252,11 @@ namespace UE::Audio::Insights
 						[](const IDashboardDataTreeViewEntry& InEntry) -> const ::Audio::TCircularAudioBuffer<FDataPoint>& { return CastEntry(InEntry).VolumeDataPoints; },
 						nullptr,
 						[this]() { return FSlateStyle::Get().GetLinearVolumeFloatFormat(); },
-						[this]() { return FVector2f(0.0f, 1.0f); }
+						[]() { return FSoundPlotsDashboardPlotRanges::StaticStruct()->FindPropertyByName(GET_MEMBER_NAME_CHECKED(FSoundPlotsDashboardPlotRanges, Volume)); },
+						[this]() { return DefaultPlotRanges.bUseCustomVolumeRange; },
+						[this](bool bUseCustomYRange) { DefaultPlotRanges.bUseCustomVolumeRange = bUseCustomYRange; OnDefaultPlotRangesChanged(); },
+						[this]() { return DefaultPlotRanges.Volume; },
+						[this](const FFloatInterval& Range, TOptional<ETextCommit::Type> CommitType) { DefaultPlotRanges.Volume = Range; OnDefaultPlotRangesChanged(CommitType.IsSet()); }
 					}
 				},
 				{
@@ -3212,7 +3266,11 @@ namespace UE::Audio::Insights
 						[](const IDashboardDataTreeViewEntry& InEntry) -> const ::Audio::TCircularAudioBuffer<FDataPoint>& { return CastEntry(InEntry).DistanceDataPoints; },
 						nullptr,
 						[this]() { return FSlateStyle::Get().GetFreqFloatFormat(); },
-						[this]() { return FVector2f(10000.0f, 0.0f); }
+						[]() { return FSoundPlotsDashboardPlotRanges::StaticStruct()->FindPropertyByName(GET_MEMBER_NAME_CHECKED(FSoundPlotsDashboardPlotRanges, Distance)); },
+						[this]() { return DefaultPlotRanges.bUseCustomDistanceRange; },
+						[this](bool bUseCustomYRange) { DefaultPlotRanges.bUseCustomDistanceRange = bUseCustomYRange; OnDefaultPlotRangesChanged(); },
+						[this]() { return DefaultPlotRanges.Distance; },
+						[this](const FFloatInterval& Range, TOptional<ETextCommit::Type> CommitType) { DefaultPlotRanges.Distance = Range; OnDefaultPlotRangesChanged(CommitType.IsSet()); }
 					}
 				},
 				{
@@ -3222,7 +3280,11 @@ namespace UE::Audio::Insights
 						[](const IDashboardDataTreeViewEntry& InEntry) -> const ::Audio::TCircularAudioBuffer<FDataPoint>& { return CastEntry(InEntry).DistanceAttenuationDataPoints; },
 						nullptr,
 						[this]() { return FSlateStyle::Get().GetLinearVolumeFloatFormat(); },
-						[this]() { return FVector2f(0.0f, 1.0f); }
+						[]() { return FSoundPlotsDashboardPlotRanges::StaticStruct()->FindPropertyByName(GET_MEMBER_NAME_CHECKED(FSoundPlotsDashboardPlotRanges, DistanceAttenuation)); },
+						[this]() { return DefaultPlotRanges.bUseCustomDistanceAttenuationRange; },
+						[this](bool bUseCustomYRange) { DefaultPlotRanges.bUseCustomDistanceAttenuationRange = bUseCustomYRange; OnDefaultPlotRangesChanged(); },
+						[this]() { return DefaultPlotRanges.DistanceAttenuation; },
+						[this](const FFloatInterval& Range, TOptional<ETextCommit::Type> CommitType) { DefaultPlotRanges.DistanceAttenuation = Range; OnDefaultPlotRangesChanged(CommitType.IsSet()); }
 					}
 				},
 				{
@@ -3232,7 +3294,11 @@ namespace UE::Audio::Insights
 						[](const IDashboardDataTreeViewEntry& InEntry) -> const ::Audio::TCircularAudioBuffer<FDataPoint>& { return CastEntry(InEntry).PitchDataPoints; },
 						nullptr,
 						[this]() { return FSlateStyle::Get().GetPitchFloatFormat(); },
-						[this]() { return FVector2f(0.0f, 1.0f); }
+						[]() { return FSoundPlotsDashboardPlotRanges::StaticStruct()->FindPropertyByName(GET_MEMBER_NAME_CHECKED(FSoundPlotsDashboardPlotRanges, Pitch)); },
+						[this]() { return DefaultPlotRanges.bUseCustomPitchRange; },
+						[this](bool bUseCustomYRange) { DefaultPlotRanges.bUseCustomPitchRange = bUseCustomYRange; OnDefaultPlotRangesChanged(); },
+						[this]() { return DefaultPlotRanges.Pitch; },
+						[this](const FFloatInterval& Range, TOptional<ETextCommit::Type> CommitType) { DefaultPlotRanges.Pitch = Range; OnDefaultPlotRangesChanged(CommitType.IsSet()); }
 					}
 				},
 				{
@@ -3242,7 +3308,11 @@ namespace UE::Audio::Insights
 						[](const IDashboardDataTreeViewEntry& InEntry) -> const ::Audio::TCircularAudioBuffer<FDataPoint>& { return CastEntry(InEntry).PriorityDataPoints; },
 						nullptr,
 						[this]() { return FSlateStyle::Get().GetLinearVolumeFloatFormat(); },
-						[this]() { return FVector2f(0.0f, 1.0f); }
+						[]() { return FSoundPlotsDashboardPlotRanges::StaticStruct()->FindPropertyByName(GET_MEMBER_NAME_CHECKED(FSoundPlotsDashboardPlotRanges, Priority)); },
+						[this]() { return DefaultPlotRanges.bUseCustomPriorityRange; },
+						[this](bool bUseCustomYRange) { DefaultPlotRanges.bUseCustomPriorityRange = bUseCustomYRange; OnDefaultPlotRangesChanged(); },
+						[this]() { return DefaultPlotRanges.Priority; },
+						[this](const FFloatInterval& Range, TOptional<ETextCommit::Type> CommitType) { DefaultPlotRanges.Priority = Range; OnDefaultPlotRangesChanged(CommitType.IsSet()); }
 					}
 				},
 				{
@@ -3252,7 +3322,11 @@ namespace UE::Audio::Insights
 						[](const IDashboardDataTreeViewEntry& InEntry) -> const ::Audio::TCircularAudioBuffer<FDataPoint>& { return CastEntry(InEntry).LPFFreqDataPoints; },
 						nullptr,
 						[this]() { return FSlateStyle::Get().GetFreqFloatFormat(); },
-						[this]() { return FVector2f(20.0f, 20000.0f); }
+						[]() { return FSoundPlotsDashboardPlotRanges::StaticStruct()->FindPropertyByName(GET_MEMBER_NAME_CHECKED(FSoundPlotsDashboardPlotRanges, LPFFreq)); },
+						[this]() { return DefaultPlotRanges.bUseCustomLPFFreqRange; },
+						[this](bool bUseCustomYRange) { DefaultPlotRanges.bUseCustomLPFFreqRange = bUseCustomYRange; OnDefaultPlotRangesChanged(); },
+						[this]() { return DefaultPlotRanges.LPFFreq; },
+						[this](const FFloatInterval& Range, TOptional<ETextCommit::Type> CommitType) { DefaultPlotRanges.LPFFreq = Range; OnDefaultPlotRangesChanged(CommitType.IsSet()); }
 					}
 				},
 				{
@@ -3262,7 +3336,11 @@ namespace UE::Audio::Insights
 						[](const IDashboardDataTreeViewEntry& InEntry) -> const ::Audio::TCircularAudioBuffer<FDataPoint>& { return CastEntry(InEntry).HPFFreqDataPoints; },
 						nullptr,
 						[this]() { return FSlateStyle::Get().GetFreqFloatFormat(); },
-						[this]() { return FVector2f(20.0f, 20000.0f); }
+						[]() { return FSoundPlotsDashboardPlotRanges::StaticStruct()->FindPropertyByName(GET_MEMBER_NAME_CHECKED(FSoundPlotsDashboardPlotRanges, HPFFreq)); },
+						[this]() { return DefaultPlotRanges.bUseCustomHPFFreqRange; },
+						[this](bool bUseCustomYRange) { DefaultPlotRanges.bUseCustomHPFFreqRange = bUseCustomYRange; OnDefaultPlotRangesChanged(); },
+						[this]() { return DefaultPlotRanges.HPFFreq; },
+						[this](const FFloatInterval& Range, TOptional<ETextCommit::Type> CommitType) { DefaultPlotRanges.HPFFreq = Range; OnDefaultPlotRangesChanged(CommitType.IsSet()); }
 					}
 				}
 			};
diff --git a/Engine/Plugins/AudioInsights/Source/AudioInsights/Private/Views/SoundPlotsWidgetView.cpp b/Engine/Plugins/AudioInsights/Source/AudioInsights/Private/Views/SoundPlotsWidgetView.cpp
--- a/Engine/Plugins/AudioInsights/Source/AudioInsights/Private/Views/SoundPlotsWidgetView.cpp
+++ b/Engine/Plugins/AudioInsights/Source/AudioInsights/Private/Views/SoundPlotsWidgetView.cpp
@@ -5,6 +5,8 @@
 #include "AudioInsightsStyle.h"
 #include "Features/IModularFeatures.h"
 #include "Widgets/Input/SComboBox.h"
+#include "Widgets/Input/SNumericEntryBox.h"
+#include "Widgets/Input/SSegmentedControl.h"
 
 #if WITH_EDITOR
 #include "Editor.h"
@@ -168,6 +170,10 @@ namespace UE::Audio::Insights
 			{
 				PlotWidget->SetCurvesPointData(CurveDataMapPtr);
 				PlotWidget->SetYValueFormattingOptions(*GetPlotColumnNumberFormat(SelectedPlotColumnName));
+				if (TOptional<FFloatInterval> CustomizedYRange = GetPlotColumnCustomizedYRange(SelectedPlotColumnName))
+				{
+					PlotWidget->SetYAxisRange(*CustomizedYRange);
+				}
 			}
 		}
 
@@ -305,6 +311,11 @@ namespace UE::Audio::Insights
 			ConstructPlotsWidgetContents();
 		}
 
+		static FSpinBoxStyle TransparentFillSpinBoxStyle = FAppStyle::Get().GetWidgetStyle<FSpinBoxStyle>("NumericEntrySpinBox");
+		TransparentFillSpinBoxStyle.ActiveFillBrush = FSlateColorBrush(FLinearColor::Transparent);
+		TransparentFillSpinBoxStyle.HoveredFillBrush = FSlateColorBrush(FLinearColor::Transparent);
+		TransparentFillSpinBoxStyle.InactiveFillBrush = FSlateColorBrush(FLinearColor::Transparent);
+
 		return SNew(SVerticalBox)
 			.Clipping(EWidgetClipping::ClipToBounds)
 			+ SVerticalBox::Slot()
@@ -321,7 +332,87 @@ namespace UE::Audio::Insights
 			.AutoHeight()
 			.HAlign(HAlign_Fill)
 			[
-				ComboBox.ToSharedRef()
+				SNew(SHorizontalBox)
+				+ SHorizontalBox::Slot()
+				[
+					ComboBox.ToSharedRef()
+				]
+
+				+ SHorizontalBox::Slot()
+				.AutoWidth()
+				.HAlign(EHorizontalAlignment::HAlign_Right)
+				[
+					SNew(STextBlock)
+						.Margin(FMargin(10.0, 4.0, 5.0, 0.0))
+						.Text(LOCTEXT("SoundPlotsWidgetView_YAxisRange", "Y Axis Range:"))
+				]
+				+ SHorizontalBox::Slot()
+				.Padding(4.0f, 0.0f, 0.0f, 0.0f)
+				.AutoWidth()
+				.HAlign(EHorizontalAlignment::HAlign_Right)
+				[
+					SNew(SSegmentedControl<bool>)
+						.OnValueChanged_Lambda([this](bool bInValue) { SetPlotColumnUseCustomYRange(SelectedPlotColumnName, bInValue); })
+						.Value_Lambda([this] { return GetPlotColumnUseCustomYRange(SelectedPlotColumnName); })
+						+ SSegmentedControl<bool>::Slot(false)
+						.Text(LOCTEXT("SoundPlotsWidgetView_YAxisRange_Auto_Text", "Auto"))
+						.ToolTip(LOCTEXT("SoundPlotsWidgetView_YAxisRange_Auto_ToolTipText", "Auto range Y axis."))
+						+ SSegmentedControl<bool>::Slot(true)
+						.Text(LOCTEXT("SoundPlotsWidgetView_YAxisRange_Custom_Text", "Custom"))
+						.ToolTip(LOCTEXT("SoundPlotsWidgetView_YAxisRange_Custom_ToolTipText", "Custom range Y axis."))
+				]
+
+				+ SHorizontalBox::Slot()
+				.Padding(4.0f, 0.0f, 0.0f, 0.0f)
+				.AutoWidth()
+				.HAlign(EHorizontalAlignment::HAlign_Right)
+				[
+					SNew(SNumericEntryBox<float>)
+					.MinDesiredValueWidth(50.0f)
+					.MinSliderValue_Lambda([this]() { return GetPlotColumnYRangeFloatMetaData(SelectedPlotColumnName, TEXT("UIMin")); })
+					.MaxSliderValue_Lambda([this]() { return GetPlotColumnYRangeFloatMetaData(SelectedPlotColumnName, TEXT("UIMax")); })
+					.MaxFractionalDigits(2)
+					.AllowSpin(true)
+					.SpinBoxStyle(&TransparentFillSpinBoxStyle)
+					.IsEnabled_Lambda([this]() { return GetPlotColumnUseCustomYRange(SelectedPlotColumnName); })
+					.Value_Lambda([this]()
+					{
+						const FFloatInterval DefaultRange = PlotWidget->GetYAxisRange();
+						return GetPlotColumnCustomizedYRange(SelectedPlotColumnName).Get(DefaultRange).Min;
+					})
+					.OnValueChanged_Lambda([this](float Value) { SetPlotColumnCustomYRangeMin(SelectedPlotColumnName, Value); })
+					.OnValueCommitted_Lambda([this](float Value, ETextCommit::Type CommitType) { SetPlotColumnCustomYRangeMin(SelectedPlotColumnName, Value, CommitType); })
+					.Label()
+					[
+						SNumericEntryBox<float>::BuildLabel(LOCTEXT("SoundPlotsWidgetView_YAxisRangeMinLabel", "Min"), FLinearColor::White, FLinearColor::Transparent)
+					]
+				]
+
+				+ SHorizontalBox::Slot()
+				.Padding(4.0f, 0.0f, 0.0f, 0.0f)
+				.AutoWidth()
+				.HAlign(EHorizontalAlignment::HAlign_Right)
+				[
+					SNew(SNumericEntryBox<float>)
+					.MinDesiredValueWidth(50.0f)
+					.MinSliderValue_Lambda([this]() { return GetPlotColumnYRangeFloatMetaData(SelectedPlotColumnName, TEXT("UIMin")); })
+					.MaxSliderValue_Lambda([this]() { return GetPlotColumnYRangeFloatMetaData(SelectedPlotColumnName, TEXT("UIMax")); })
+					.MaxFractionalDigits(2)
+					.AllowSpin(true)
+					.SpinBoxStyle(&TransparentFillSpinBoxStyle)
+					.IsEnabled_Lambda([this]() { return GetPlotColumnUseCustomYRange(SelectedPlotColumnName); })
+					.Value_Lambda([this]()
+					{
+						const FFloatInterval DefaultRange = PlotWidget->GetYAxisRange();
+						return GetPlotColumnCustomizedYRange(SelectedPlotColumnName).Get(DefaultRange).Max;
+					})
+					.OnValueChanged_Lambda([this](float Value) { SetPlotColumnCustomYRangeMax(SelectedPlotColumnName, Value); })
+					.OnValueCommitted_Lambda([this](float Value, ETextCommit::Type CommitType) { SetPlotColumnCustomYRangeMax(SelectedPlotColumnName, Value, CommitType); })
+					.Label()
+					[
+						SNumericEntryBox<float>::BuildLabel(LOCTEXT("SoundPlotsWidgetView_YAxisRangeMaxLabel", "Max"), FLinearColor::White, FLinearColor::Transparent)
+					]
+				]
 			]
 			+ SVerticalBox::Slot()
 			.HAlign(HAlign_Fill)
@@ -355,11 +446,12 @@ namespace UE::Audio::Insights
 		}
 
 		SAssignNew(PlotWidget, SAudioCurveView)
+			.AutoRangeYAxis_Lambda([this]() { return !GetPlotColumnUseCustomYRange(SelectedPlotColumnName); })
 			.ViewRange_Lambda([this]() { return GetViewRange(); })
 			.PixelSnappingMethod(EWidgetPixelSnapping::Disabled);
 
 		PlotWidget->SetYValueFormattingOptions(*GetPlotColumnNumberFormat(SelectedPlotColumnName));
-		PlotWidget->SetInitYAxisRange(GetPlotColumnDefaultYRange(SelectedPlotColumnName));
+		PlotWidget->SetYAxisRange(GetPlotColumnCustomizedOrDefaultYRange(SelectedPlotColumnName));
 
 		SAssignNew(ComboBox, SComboBox<FName>)
 			.ToolTipText(LOCTEXT("SoundPlotsWidgetView_SelectPlotColumnDescription", "Select a column from the table to plot."))
@@ -377,7 +469,7 @@ namespace UE::Audio::Insights
 				{
 					PlotWidget->SetCurvesPointData(*DataMap);
 					PlotWidget->SetYValueFormattingOptions(*GetPlotColumnNumberFormat(NewColumnName));
-					PlotWidget->SetInitYAxisRange(GetPlotColumnDefaultYRange(NewColumnName));
+					PlotWidget->SetYAxisRange(GetPlotColumnCustomizedOrDefaultYRange(NewColumnName));
 
 #if !WITH_EDITOR
 					bTimeMarkerMoved = true;
@@ -565,25 +657,99 @@ namespace UE::Audio::Insights
 		return nullptr;
 	}
 
-	const FVector2f FSoundPlotsWidgetView::GetPlotColumnDefaultYRange(const FName& ColumnName)
+	FFloatInterval FSoundPlotsWidgetView::GetPlotColumnCustomizedOrDefaultYRange(const FName& ColumnName)
+	{
+		// First get values from the range property metadata if available, or 0 to 1 if not available:
+		const float DefaultMin = GetPlotColumnYRangeFloatMetaData(ColumnName, TEXT("UIMin")).Get(0.0f);
+		const float DefaultMax = GetPlotColumnYRangeFloatMetaData(ColumnName, TEXT("UIMax")).Get(1.0f);
+
+		// Use the customized Y range, or the default values otherwise:
+		return GetPlotColumnCustomizedYRange(ColumnName).Get({ DefaultMin, DefaultMax });
+	}
+
+	const FText FSoundPlotsWidgetView::GetPlotColumnDisplayName(const FName& ColumnName)
 	{
 		if (const FSoundPlotsWidgetView::FPlotColumnInfo* PlotColumnInfo = ColumnInfo.Find(ColumnName))
 		{
-			if (PlotColumnInfo->GetDefaultYDataRangeFunc.IsSet())
+			return PlotColumnInfo->ColumnDisplayName;
+		}
+		return FText::GetEmpty();
+	}
+
+	TOptional<float> FSoundPlotsWidgetView::GetPlotColumnYRangeFloatMetaData(const FName& ColumnName, const FName& MetaDataKey)
+	{
+		if (const FSoundPlotsWidgetView::FPlotColumnInfo* PlotColumnInfo = ColumnInfo.Find(ColumnName))
+		{
+			if (PlotColumnInfo->GetYRangePropertyFunc.IsSet())
 			{
-				return PlotColumnInfo->GetDefaultYDataRangeFunc();
+				if (const FProperty* YRangeProperty = PlotColumnInfo->GetYRangePropertyFunc())
+				{
+					if (YRangeProperty->HasMetaData(MetaDataKey))
+					{
+						return YRangeProperty->GetFloatMetaData(MetaDataKey);
+					}
+				}
 			}
 		}
-		return FVector2f(0.0f, 1.0f);
+		return NullOpt;
 	}
 
-	const FText FSoundPlotsWidgetView::GetPlotColumnDisplayName(const FName& ColumnName)
+	TOptional<FFloatInterval> FSoundPlotsWidgetView::GetPlotColumnCustomizedYRange(const FName& ColumnName)
 	{
 		if (const FSoundPlotsWidgetView::FPlotColumnInfo* PlotColumnInfo = ColumnInfo.Find(ColumnName))
 		{
-			return PlotColumnInfo->ColumnDisplayName;
+			if (PlotColumnInfo->GetCustomYRangeFunc.IsSet() && GetPlotColumnUseCustomYRange(ColumnName))
+			{
+				return PlotColumnInfo->GetCustomYRangeFunc();
+			}
+		}
+		return NullOpt;
+	}
+
+	bool FSoundPlotsWidgetView::GetPlotColumnUseCustomYRange(const FName& ColumnName)
+	{
+		if (const FSoundPlotsWidgetView::FPlotColumnInfo* PlotColumnInfo = ColumnInfo.Find(ColumnName))
+		{
+			return (PlotColumnInfo->GetUseCustomYRangeFunc.IsSet() && PlotColumnInfo->GetUseCustomYRangeFunc());
+		}
+		return true;
+	}
+
+	void FSoundPlotsWidgetView::SetPlotColumnUseCustomYRange(const FName& ColumnName, bool bUseCustomYRange)
+	{
+		if (const FSoundPlotsWidgetView::FPlotColumnInfo* PlotColumnInfo = ColumnInfo.Find(ColumnName))
+		{
+			if (PlotColumnInfo->SetUseCustomYRangeFunc.IsSet())
+			{
+				PlotColumnInfo->SetUseCustomYRangeFunc(bUseCustomYRange);
+			}
+		}
+	}
+
+	void FSoundPlotsWidgetView::SetPlotColumnCustomYRangeMin(const FName& ColumnName, float Value, TOptional<ETextCommit::Type> CommitType)
+	{
+		if (const FSoundPlotsWidgetView::FPlotColumnInfo* PlotColumnInfo = ColumnInfo.Find(ColumnName))
+		{
+			if (PlotColumnInfo->SetCustomYRangeFunc.IsSet() && PlotColumnInfo->GetCustomYRangeFunc.IsSet())
+			{
+				FFloatInterval Range = PlotColumnInfo->GetCustomYRangeFunc();
+				Range.Min = Value;
+				PlotColumnInfo->SetCustomYRangeFunc(Range, CommitType);
+			}
+		}
+	}
+
+	void FSoundPlotsWidgetView::SetPlotColumnCustomYRangeMax(const FName& ColumnName, float Value, TOptional<ETextCommit::Type> CommitType)
+	{
+		if (const FSoundPlotsWidgetView::FPlotColumnInfo* PlotColumnInfo = ColumnInfo.Find(ColumnName))
+		{
+			if (PlotColumnInfo->SetCustomYRangeFunc.IsSet() && PlotColumnInfo->GetCustomYRangeFunc.IsSet())
+			{
+				FFloatInterval Range = PlotColumnInfo->GetCustomYRangeFunc();
+				Range.Max = Value;
+				PlotColumnInfo->SetCustomYRangeFunc(Range, CommitType);
+			}
 		}
-		return FText::GetEmpty();
 	}
 
 	TRange<double> FSoundPlotsWidgetView::GetViewRange()
diff --git a/Engine/Plugins/AudioInsights/Source/AudioInsights/Public/Settings/SoundDashboardSettings.h b/Engine/Plugins/AudioInsights/Source/AudioInsights/Public/Settings/SoundDashboardSettings.h
--- a/Engine/Plugins/AudioInsights/Source/AudioInsights/Public/Settings/SoundDashboardSettings.h
+++ b/Engine/Plugins/AudioInsights/Source/AudioInsights/Public/Settings/SoundDashboardSettings.h
@@ -117,6 +117,75 @@ struct FSoundDashboardVisibleColumns : public FVisibleColumnsSettings
 	}
 };
 
+USTRUCT()
+struct FSoundPlotsDashboardPlotRanges
+{
+	GENERATED_BODY()
+
+
+	UPROPERTY(EditAnywhere, config, Category = "SoundDashboard|CustomPlotRanges", meta = (InlineEditConditionToggle))
+	bool bUseCustomAmplitude_dBRange = false;
+
+	UPROPERTY(EditAnywhere, config, Category = "SoundDashboard|CustomPlotRanges", meta = (EditCondition = "bUseCustomAmplitude_dBRange", UIMin = "-160", UIMax = "0", ClampMin = "-160", ClampMax = "24", DisplayName = "Amplitude (dB)"))
+	FFloatInterval Amplitude_dB = { -160.0f, 0.0f };
+
+
+	UPROPERTY(EditAnywhere, config, Category = "SoundDashboard|CustomPlotRanges", meta = (InlineEditConditionToggle))
+	bool bUseCustomAmplitude_LinearRange = false;
+
+	UPROPERTY(EditAnywhere, config, Category = "SoundDashboard|CustomPlotRanges", meta = (EditCondition = "bUseCustomAmplitude_LinearRange", UIMin = "0", UIMax = "1", ClampMin = "-16", ClampMax = "16", DisplayName = "Amplitude (Linear)"))
+	FFloatInterval Amplitude_Linear = { 0.0f, 1.0f };
+
+
+	UPROPERTY(EditAnywhere, config, Category = "SoundDashboard|CustomPlotRanges", meta = (InlineEditConditionToggle))
+	bool bUseCustomVolumeRange = false;
+
+	UPROPERTY(EditAnywhere, config, Category = "SoundDashboard|CustomPlotRanges", meta = (EditCondition = "bUseCustomVolumeRange", UIMin = "0", UIMax = "1", ClampMin = "0", ClampMax = "1"))
+	FFloatInterval Volume = { 0.0f, 1.0f };
+
+
+	UPROPERTY(EditAnywhere, config, Category = "SoundDashboard|CustomPlotRanges", meta = (InlineEditConditionToggle))
+	bool bUseCustomDistanceRange = false;
+
+	UPROPERTY(EditAnywhere, config, Category = "SoundDashboard|CustomPlotRanges", meta = (EditCondition = "bUseCustomDistanceRange", UIMin = "0", UIMax = "10000", ClampMin = "0", ClampMax = "100000"))
+	FFloatInterval Distance = { 0.0f, 10000.0f };
+
+
+	UPROPERTY(EditAnywhere, config, Category = "SoundDashboard|CustomPlotRanges", meta = (InlineEditConditionToggle))
+	bool bUseCustomDistanceAttenuationRange = false;
+
+	UPROPERTY(EditAnywhere, config, Category = "SoundDashboard|CustomPlotRanges", meta = (EditCondition = "bUseCustomDistanceAttenuationRange", UIMin = "0", UIMax = "1", ClampMin = "0", ClampMax = "1"))
+	FFloatInterval DistanceAttenuation = { 0.0f, 1.0f };
+
+
+	UPROPERTY(EditAnywhere, config, Category = "SoundDashboard|CustomPlotRanges", meta = (InlineEditConditionToggle))
+	bool bUseCustomPitchRange = false;
+
+	UPROPERTY(EditAnywhere, config, Category = "SoundDashboard|CustomPlotRanges", meta = (EditCondition = "bUseCustomPitchRange", UIMin = "0.25", UIMax = "4", ClampMin = "0", ClampMax = "32", Units = "Multiplier"))
+	FFloatInterval Pitch = { 0.25f, 4.0f };
+
+
+	UPROPERTY(EditAnywhere, config, Category = "SoundDashboard|CustomPlotRanges", meta = (InlineEditConditionToggle))
+	bool bUseCustomPriorityRange = false;
+
+	UPROPERTY(EditAnywhere, config, Category = "SoundDashboard|CustomPlotRanges", meta = (EditCondition = "bUseCustomPriorityRange", UIMin = "0", UIMax = "255", ClampMin = "0", ClampMax = "255"))
+	FFloatInterval Priority = { 0.0f, 255.0f };
+
+
+	UPROPERTY(EditAnywhere, config, Category = "SoundDashboard|CustomPlotRanges", meta = (InlineEditConditionToggle))
+	bool bUseCustomLPFFreqRange = false;
+
+	UPROPERTY(EditAnywhere, config, Category = "SoundDashboard|CustomPlotRanges", meta = (EditCondition = "bUseCustomLPFFreqRange", UIMin = "20", UIMax = "20000", DisplayName = "LPF Freq", ClampMin = "0", ClampMax = "48000", Units = "Hertz"))
+	FFloatInterval LPFFreq = { 20.0f, 20000.0f };
+
+
+	UPROPERTY(EditAnywhere, config, Category = "SoundDashboard|CustomPlotRanges", meta = (InlineEditConditionToggle))
+	bool bUseCustomHPFFreqRange = false;
+
+	UPROPERTY(EditAnywhere, config, Category = "SoundDashboard|CustomPlotRanges", meta = (EditCondition = "bUseCustomHPFFreqRange", UIMin = "20", UIMax = "20000", DisplayName = "HPF Freq", ClampMin = "0", ClampMax = "48000", Units = "Hertz"))
+	FFloatInterval HPFFreq = { 20.0f, 20000.0f };
+};
+
 USTRUCT()
 struct FSoundDashboardSettings
 {
@@ -140,6 +209,9 @@ struct FSoundDashboardSettings
 	UPROPERTY(EditAnywhere, config, Category = SoundDashboard)
 	FSoundDashboardVisibleColumns VisibleColumns;
 
+	UPROPERTY(EditAnywhere, config, Category = SoundDashboard, meta = (ShowOnlyInnerProperties))
+	FSoundPlotsDashboardPlotRanges DefaultPlotRanges;
+
 #if WITH_EDITOR
 	DECLARE_MULTICAST_DELEGATE_OneParam(FOnReadSoundDashboardSettings, const FSoundDashboardSettings&);
 	static AUDIOINSIGHTS_API FOnReadSoundDashboardSettings OnReadSettings;
diff --git a/Engine/Plugins/AudioInsights/Source/AudioInsights/Public/Views/SAudioCurveView.h b/Engine/Plugins/AudioInsights/Source/AudioInsights/Public/Views/SAudioCurveView.h
--- a/Engine/Plugins/AudioInsights/Source/AudioInsights/Public/Views/SAudioCurveView.h
+++ b/Engine/Plugins/AudioInsights/Source/AudioInsights/Public/Views/SAudioCurveView.h
@@ -27,6 +27,7 @@ public:
 
 	SLATE_BEGIN_ARGS(SAudioCurveView)
 		: _ViewRange(TRange<double>(0,5))
+		, _AutoRangeYAxis(false)
 		, _YMargin(0.05f)
 		, _HorizontalAxisIncrement(0.5)
 		, _GridLineColor(FLinearColor(0.5f, 0.5f, 0.5f, 0.25f))
@@ -36,6 +37,8 @@ public:
     	/** View X axis range (in value space) */
     	SLATE_ATTRIBUTE(TRange<double>, ViewRange);
 
+		SLATE_ATTRIBUTE(bool, AutoRangeYAxis);
+
 		/** Margin for Y axis, as a 0 - 0.5f proportion, for the space each of above and below the data range. (ex. 0.05 means a 5% margin on the top and bottom, with 90% of the widget's vertical size corresponding to the data range). */
 		SLATE_ATTRIBUTE(float, YMargin);
 
@@ -59,10 +62,12 @@ public:
 	void Construct( const FArguments& InArgs );
 
 	void SetYValueFormattingOptions(const FNumberFormattingOptions InValueFormattingOptions);
-	void SetInitYAxisRange(const FVector2f& InRange);
+	void SetYAxisRange(const FFloatInterval& InRange);
 	void SetCurvesPointData(TSharedPtr<TMap<int32, TArray<FCurvePoint>>> InCurvesPointData);
 	void SetCurvesMetadata(TSharedPtr<TMap<int32, FCurveMetadata>> InMetadataPerCurve);
 
+	FFloatInterval GetYAxisRange() const { return YDataRange; }
+
 	FText GetCurveToolTipXValueText() const { return CurveToolTipXValueText; }
 	FText GetCurveToolTipYValueText() const { return CurveToolTipYValueText; }
 	FText GetCurveToolTipDisplayNameText() const { return CurveToolTipDisplayNameText; }
@@ -91,6 +96,7 @@ protected:
 	TSharedRef<SToolTip> CreateCurveTooltip();
 
 	TAttribute<TRange<double>> ViewRange;
+	TAttribute<bool> AutoRangeYAxis;
 	TAttribute<float> YMargin;
 	TAttribute<double> HorizontalAxisIncrement;
 	TAttribute<FLinearColor> GridLineColor;
@@ -101,7 +107,7 @@ protected:
 	TSharedPtr<TMap<int32, TArray<FCurvePoint>>> PointDataPerCurve;
 	TSharedPtr<TMap<int32, FCurveMetadata>> MetadataPerCurve; 
 	// Y axis data range in value space
-	FVector2f YDataRange;
+	FFloatInterval YDataRange;
 
 	// Tooltip text
 	FText CurveToolTipXValueText;		
diff --git a/Engine/Plugins/AudioInsights/Source/AudioInsights/Public/Views/SoundDashboardViewFactory.h b/Engine/Plugins/AudioInsights/Source/AudioInsights/Public/Views/SoundDashboardViewFactory.h
--- a/Engine/Plugins/AudioInsights/Source/AudioInsights/Public/Views/SoundDashboardViewFactory.h
+++ b/Engine/Plugins/AudioInsights/Source/AudioInsights/Public/Views/SoundDashboardViewFactory.h
@@ -255,6 +255,7 @@ namespace UE::Audio::Insights
 
 		TSharedPtr<FVisibleColumnsSettingsMenu<FSoundDashboardVisibleColumns>> VisibleColumnsSettingsMenu;
 		TOptional<FSoundDashboardVisibleColumns> InitVisibleColumnSettings;
+		FSoundPlotsDashboardPlotRanges DefaultPlotRanges;
 	};
 } // namespace UE::Audio::Insights
 
diff --git a/Engine/Plugins/AudioInsights/Source/AudioInsights/Public/Views/SoundPlotsWidgetView.h b/Engine/Plugins/AudioInsights/Source/AudioInsights/Public/Views/SoundPlotsWidgetView.h
--- a/Engine/Plugins/AudioInsights/Source/AudioInsights/Public/Views/SoundPlotsWidgetView.h
+++ b/Engine/Plugins/AudioInsights/Source/AudioInsights/Public/Views/SoundPlotsWidgetView.h
@@ -23,7 +23,11 @@ namespace UE::Audio::Insights
 			const TFunction<const ::Audio::TCircularAudioBuffer<FDataPoint>& (const IDashboardDataTreeViewEntry& InData)> DataFunc;
 			const TFunction<float(const float)> UnitConversionFunc;
 			const TFunction<const FNumberFormattingOptions*()> GetFormatOptionsFunc;
-			const TFunction<FVector2f()> GetDefaultYDataRangeFunc;
+			const TFunction<const FProperty*()> GetYRangePropertyFunc;
+			const TFunction<bool()> GetUseCustomYRangeFunc;
+			const TFunction<void(bool)> SetUseCustomYRangeFunc;
+			const TFunction<FFloatInterval()> GetCustomYRangeFunc;
+			const TFunction<void(const FFloatInterval&, TOptional<ETextCommit::Type>)> SetCustomYRangeFunc;
 		};
 
 		UE_API FSoundPlotsWidgetView(TMap<FName, FSoundPlotsWidgetView::FPlotColumnInfo>&& InColumnInfo, const TFunction<bool(const IDashboardDataTreeViewEntry& InEntry)> InIsPlotEnabledForEntryFunc);
@@ -70,9 +74,16 @@ namespace UE::Audio::Insights
 		const TFunction<const ::Audio::TCircularAudioBuffer<FDataPoint>& (const IDashboardDataTreeViewEntry& InEntry)> GetPlotColumnDataFunc(const FName& ColumnName);
 		const TFunction<float (const float)> GetPlotColumnUnitConversionFunc(const FName& ColumnName);
 		const FNumberFormattingOptions* GetPlotColumnNumberFormat(const FName& ColumnName);
-		const FVector2f GetPlotColumnDefaultYRange(const FName& ColumnName);
+		FFloatInterval GetPlotColumnCustomizedOrDefaultYRange(const FName& ColumnName);
 		const FText GetPlotColumnDisplayName(const FName& ColumnName);
 
+		TOptional<float> GetPlotColumnYRangeFloatMetaData(const FName& ColumnName, const FName& MetaDataKey);
+		TOptional<FFloatInterval> GetPlotColumnCustomizedYRange(const FName& ColumnName);
+		bool GetPlotColumnUseCustomYRange(const FName& ColumnName);
+		void SetPlotColumnUseCustomYRange(const FName& ColumnName, bool bUseCustomYRange);
+		void SetPlotColumnCustomYRangeMin(const FName& ColumnName, float Value, TOptional<ETextCommit::Type> CommitType = NullOpt);
+		void SetPlotColumnCustomYRangeMax(const FName& ColumnName, float Value, TOptional<ETextCommit::Type> CommitType = NullOpt);
+
 		TRange<double> GetViewRange();
 
 		bool GetIsPlotEnabledForEntry(const IDashboardDataTreeViewEntry& InEntry) const;
