[Audio Insights] Added AudioBus messages/provider for standalone, it also has asset support (editor only), pending to add the trace messages from the audio system

diff --git a/Engine/Plugins/AudioInsights/Source/AudioInsights/Private/Providers/AudioBusAssetProvider.cpp b/Engine/Plugins/AudioInsights/Source/AudioInsights/Private/Providers/AudioBusAssetProvider.cpp
new file mode 100644
--- /dev/null
+++ b/Engine/Plugins/AudioInsights/Source/AudioInsights/Private/Providers/AudioBusAssetProvider.cpp
@@ -0,0 +1,159 @@
+// Copyright Epic Games, Inc. All Rights Reserved.
+
+#include "Providers/AudioBusAssetProvider.h"
+
+#if WITH_EDITOR
+#include "IAudioInsightsModule.h"
+#include "AssetRegistry/AssetData.h"
+#include "AssetRegistry/AssetRegistryModule.h"
+#include "AudioDeviceManager.h"
+#include "AudioMixerDevice.h"
+#include "Sound/AudioBus.h"
+
+namespace UE::Audio::Insights
+{
+	FAudioBusAssetProvider::FAudioBusAssetProvider()
+	{
+		FAssetRegistryModule& AssetRegistryModule = FModuleManager::LoadModuleChecked<FAssetRegistryModule>("AssetRegistry");
+
+		AssetRegistryModule.Get().OnAssetAdded().AddRaw(this, &FAudioBusAssetProvider::OnAssetAdded);
+		AssetRegistryModule.Get().OnAssetRemoved().AddRaw(this, &FAudioBusAssetProvider::OnAssetRemoved);
+		AssetRegistryModule.Get().OnFilesLoaded().AddRaw(this, &FAudioBusAssetProvider::OnFilesLoaded);
+
+		FAudioDeviceManagerDelegates::OnAudioDeviceCreated.AddRaw(this, &FAudioBusAssetProvider::OnActiveAudioDeviceChanged);
+		FAudioDeviceManagerDelegates::OnAudioDeviceDestroyed.AddRaw(this, &FAudioBusAssetProvider::OnActiveAudioDeviceChanged);
+
+		FTraceAuxiliary::OnTraceStarted.AddRaw(this, &FAudioBusAssetProvider::OnTraceStarted);
+	}
+
+	FAudioBusAssetProvider::~FAudioBusAssetProvider()
+	{
+		if (FAssetRegistryModule* AssetRegistryModule = FModuleManager::GetModulePtr<FAssetRegistryModule>("AssetRegistry"))
+		{
+			AssetRegistryModule->Get().OnAssetAdded().RemoveAll(this);
+			AssetRegistryModule->Get().OnAssetRemoved().RemoveAll(this);
+			AssetRegistryModule->Get().OnFilesLoaded().RemoveAll(this);
+		}
+
+		FAudioDeviceManagerDelegates::OnAudioDeviceCreated.RemoveAll(this);
+		FAudioDeviceManagerDelegates::OnAudioDeviceDestroyed.RemoveAll(this);
+
+		FTraceAuxiliary::OnTraceStarted.RemoveAll(this);
+	}
+
+	void FAudioBusAssetProvider::OnAssetAdded(const FAssetData& InAssetData)
+	{
+		if (bAreFilesLoaded && InAssetData.AssetClassPath == FTopLevelAssetPath(UAudioBus::StaticClass()))
+		{
+			constexpr bool bShouldNotify = true;
+			AddAudioBusAsset(InAssetData, bShouldNotify);
+		}
+	}
+
+	void FAudioBusAssetProvider::OnAssetRemoved(const FAssetData& InAssetData)
+	{
+		if (InAssetData.AssetClassPath == FTopLevelAssetPath(UAudioBus::StaticClass()))
+		{
+			RemoveAudioBusAsset(InAssetData);
+		}
+	}
+
+	void FAudioBusAssetProvider::OnFilesLoaded()
+	{
+		bAreFilesLoaded = true;
+		UpdateAudioBusAssetNames();
+	}
+
+	void FAudioBusAssetProvider::OnActiveAudioDeviceChanged(::Audio::FDeviceId InDeviceId)
+	{
+		UpdateAudioBusAssetNames();
+	}
+
+	void FAudioBusAssetProvider::OnTraceStarted(FTraceAuxiliary::EConnectionType TraceType, const FString& TraceDestination)
+	{
+		UpdateAudioBusAssetNames();
+	}
+
+	void FAudioBusAssetProvider::RequestEntriesUpdate()
+	{
+		UpdateAudioBusAssetNames();
+	}
+	
+	void FAudioBusAssetProvider::AddAudioBusAsset(const FAssetData& InAssetData, const bool bInShouldNotify)
+	{
+		const bool bIsAudioBusEntryAlreadAdded = AudioBusAssetEntries.ContainsByPredicate(
+			[AssetName = InAssetData.GetObjectPathString()](const TSharedPtr<FAudioBusDashboardEntry>& AudioBusDashboardEntry)
+			{
+				return AudioBusDashboardEntry->Name == AssetName;
+			});
+
+		if (!bIsAudioBusEntryAlreadAdded)
+		{
+			const IAudioInsightsModule& AudioInsightsModule = IAudioInsightsModule::GetEditorChecked();
+			const ::Audio::FDeviceId AudioDeviceId = AudioInsightsModule.GetDeviceId();
+
+			const TObjectPtr<UAudioBus> AudioBus = Cast<UAudioBus>(InAssetData.GetAsset());
+
+			TSharedPtr<FAudioBusDashboardEntry>& AudioBusDashboardEntry = AudioBusAssetEntries.Add_GetRef(MakeShared<FAudioBusDashboardEntry>());
+			AudioBusDashboardEntry->Timestamp = FPlatformTime::Seconds() - GStartTime;
+			AudioBusDashboardEntry->AudioBusId = AudioBus ? AudioBus->GetUniqueID() : INDEX_NONE; //GetTypeHash(InAssetData.GetObjectPathString());
+			AudioBusDashboardEntry->Name = InAssetData.GetObjectPathString();
+			AudioBusDashboardEntry->AudioMeterInfo->NumChannels = AudioBus->GetNumChannels();
+
+			if (bInShouldNotify)
+			{
+				OnAudioBusAssetAdded.ExecuteIfBound(AudioBusDashboardEntry);
+			}
+		}
+	}
+
+	void FAudioBusAssetProvider::RemoveAudioBusAsset(const FAssetData& InAssetData)
+	{
+		const TSharedPtr<FAudioBusDashboardEntry>* FoundAudioBusDashboardEntry = AudioBusAssetEntries.FindByPredicate(
+			[AssetName = InAssetData.GetObjectPathString()](const TSharedPtr<FAudioBusDashboardEntry>& InAudioBusDashboardEntry)
+			{
+				return InAudioBusDashboardEntry->Name == AssetName;
+			});
+
+		if (FoundAudioBusDashboardEntry)
+		{
+			const TSharedPtr<FAudioBusDashboardEntry> AudioBusDashboardEntry = *FoundAudioBusDashboardEntry;
+
+			if (AudioBusDashboardEntry.IsValid())
+			{
+				const uint32 AudioBusId = AudioBusDashboardEntry->AudioBusId;
+
+				AudioBusAssetEntries.Remove(AudioBusDashboardEntry);
+
+				OnAudioBusAssetRemoved.ExecuteIfBound(AudioBusId);
+			}
+		}
+	}
+
+	void FAudioBusAssetProvider::UpdateAudioBusAssetNames()
+	{
+		// Get all UAudioBus assets
+		TArray<FAssetData> AssetDataArray;
+
+		FAssetRegistryModule& AssetRegistryModule = FModuleManager::LoadModuleChecked<FAssetRegistryModule>("AssetRegistry");
+		AssetRegistryModule.Get().GetAssetsByClass(FTopLevelAssetPath(UAudioBus::StaticClass()), AssetDataArray);
+
+		// Build AudioBusAssetEntries
+		AudioBusAssetEntries.Reset();
+
+		constexpr bool bShouldNotify = false;
+		for (const FAssetData& AssetData : AssetDataArray)
+		{
+			AddAudioBusAsset(AssetData, bShouldNotify);
+		}
+
+		AudioBusAssetEntries.Sort([](const TSharedPtr<FAudioBusDashboardEntry>& A, const TSharedPtr<FAudioBusDashboardEntry>& B)
+		{
+			return A->GetDisplayName().CompareToCaseIgnored(B->GetDisplayName()) < 0;
+		});
+
+		OnAudioBusAssetListUpdated.ExecuteIfBound(AudioBusAssetEntries);
+	}
+} // namespace UE::Audio::Insights
+
+#endif // WITH_EDITOR
diff --git a/Engine/Plugins/AudioInsights/Source/AudioInsights/Private/Providers/AudioBusAssetProvider.h b/Engine/Plugins/AudioInsights/Source/AudioInsights/Private/Providers/AudioBusAssetProvider.h
new file mode 100644
--- /dev/null
+++ b/Engine/Plugins/AudioInsights/Source/AudioInsights/Private/Providers/AudioBusAssetProvider.h
@@ -0,0 +1,50 @@
+// Copyright Epic Games, Inc. All Rights Reserved.
+#pragma once
+
+#if WITH_EDITOR
+#include "Delegates/Delegate.h"
+#include "Messages/AudioBusTraceMessages.h"
+#include "ProfilingDebugging/TraceAuxiliary.h"
+
+struct FAssetData;
+
+namespace UE::Audio::Insights
+{
+	class FAudioBusAssetProvider : public TSharedFromThis<FAudioBusAssetProvider>
+	{
+	public:
+		FAudioBusAssetProvider();
+		virtual ~FAudioBusAssetProvider();
+
+		void RequestEntriesUpdate();
+
+		DECLARE_DELEGATE_OneParam(FOnAudioBusAssetAdded, const TSharedPtr<FAudioBusDashboardEntry>& /*AudioBusDashboardEntry*/);
+		FOnAudioBusAssetAdded OnAudioBusAssetAdded;
+
+		DECLARE_DELEGATE_OneParam(FOnAudioBusAssetRemoved, const uint32 /*AudioBusId*/);
+		FOnAudioBusAssetRemoved OnAudioBusAssetRemoved;
+
+		DECLARE_DELEGATE_OneParam(FOnAudioBusAssetListUpdated, const TArray<TSharedPtr<FAudioBusDashboardEntry>>& /*AudioBusDashboardEntries*/);
+		FOnAudioBusAssetListUpdated OnAudioBusAssetListUpdated;
+
+	private:
+		void OnAssetAdded(const FAssetData& InAssetData);
+		void OnAssetRemoved(const FAssetData& InAssetData);
+		void OnFilesLoaded();
+
+		void OnActiveAudioDeviceChanged(::Audio::FDeviceId InDeviceId);
+		void OnTraceStarted(FTraceAuxiliary::EConnectionType TraceType, const FString& TraceDestination);
+
+		void AddAudioBusAsset(const FAssetData& InAssetData, const bool bInShouldNotify);
+		void RemoveAudioBusAsset(const FAssetData& InAssetData);
+
+		void UpdateAudioBusAssetNames();
+
+		bool bAreFilesLoaded = false;
+		bool bAssetEntriesNeedRefreshing = false;
+
+		TArray<TSharedPtr<FAudioBusDashboardEntry>> AudioBusAssetEntries;
+	};
+} // namespace UE::Audio::Insights
+
+#endif // WITH_EDITOR
diff --git a/Engine/Plugins/AudioInsights/Source/AudioInsights/Private/Providers/AudioBusTraceProvider.cpp b/Engine/Plugins/AudioInsights/Source/AudioInsights/Private/Providers/AudioBusTraceProvider.cpp
new file mode 100644
--- /dev/null
+++ b/Engine/Plugins/AudioInsights/Source/AudioInsights/Private/Providers/AudioBusTraceProvider.cpp
@@ -0,0 +1,445 @@
+// Copyright Epic Games, Inc. All Rights Reserved.
+
+#include "Providers/AudioBusTraceProvider.h"
+
+#if WITH_EDITOR
+#include "AudioDevice.h"
+#include "AudioDeviceManager.h"
+#else
+#include "Async/ParallelFor.h"
+#include "Common/PagedArray.h"
+#endif // WITH_EDITOR
+
+namespace UE::Audio::Insights
+{
+	namespace FAudioBusTraceProviderPrivate
+	{
+#if !WITH_EDITOR
+		template<typename T>
+		const T* FindClosestMessageToTimestamp(const TraceServices::TPagedArray<T>& InCachedMessages, const double InTimeMarker, const uint32 InAudioBusId)
+		{
+			const int32 ClosestMessageToTimeStampIndex = TraceServices::PagedArrayAlgo::BinarySearchClosestBy(InCachedMessages, InTimeMarker, [](const T& Msg) { return Msg.Timestamp; });
+
+			// Iterate backwards from TimeMarker until we find the matching AudioBusId
+			for (auto It = InCachedMessages.GetIteratorFromItem(ClosestMessageToTimeStampIndex); It; --It)
+			{
+				if (It->AudioBusId == InAudioBusId)
+				{
+					return &(*It);
+				}
+			}
+
+			return nullptr;
+		}
+#endif // !WITH_EDITOR
+	}
+
+	FAudioBusTraceProvider::FAudioBusTraceProvider()
+		: TDeviceDataMapTraceProvider<uint32, TSharedPtr<FAudioBusDashboardEntry>>(GetName_Static())
+	{
+#if WITH_EDITOR
+		AudioBusAssetProvider.OnAudioBusAssetAdded.BindRaw(this, &FAudioBusTraceProvider::HandleOnAudioBusAssetAdded);
+		AudioBusAssetProvider.OnAudioBusAssetRemoved.BindRaw(this, &FAudioBusTraceProvider::HandleOnAudioBusAssetRemoved);
+		AudioBusAssetProvider.OnAudioBusAssetListUpdated.BindRaw(this, &FAudioBusTraceProvider::HandleOnAudioBusAssetListUpdated);
+#endif // WITH_EDITOR
+	}
+
+	FAudioBusTraceProvider::~FAudioBusTraceProvider()
+	{
+#if WITH_EDITOR
+		AudioBusAssetProvider.OnAudioBusAssetAdded.Unbind();
+		AudioBusAssetProvider.OnAudioBusAssetRemoved.Unbind();
+		AudioBusAssetProvider.OnAudioBusAssetListUpdated.Unbind();
+#endif // WITH_EDITOR
+	}
+	
+	FName FAudioBusTraceProvider::GetName_Static()
+	{
+		static const FLazyName AudioBusTraceProviderName = "AudioBusProvider";
+		return AudioBusTraceProviderName;
+	}
+
+	UE::Trace::IAnalyzer* FAudioBusTraceProvider::ConstructAnalyzer(TraceServices::IAnalysisSession& InSession)
+	{
+		class FAudioBusTraceAnalyzer : public FTraceAnalyzerBase
+		{
+		public:
+			FAudioBusTraceAnalyzer(TSharedRef<FAudioBusTraceProvider> InProvider, TraceServices::IAnalysisSession& InSession)
+				: FTraceAnalyzerBase(InProvider)
+				, Session(InSession)
+			{
+			}
+
+			virtual void OnAnalysisBegin(const FOnAnalysisContext& Context) override
+			{
+				FTraceAnalyzerBase::OnAnalysisBegin(Context);
+
+				UE::Trace::IAnalyzer::FInterfaceBuilder& Builder = Context.InterfaceBuilder;
+
+				Builder.RouteEvent(RouteId_Start,          "Audio", "AudioBusStartOnTraceStart");
+				Builder.RouteEvent(RouteId_HasActivity,    "Audio", "AudioBusHasActivity");
+				Builder.RouteEvent(RouteId_Start,          "Audio", "AudioBusStart");
+				Builder.RouteEvent(RouteId_EnvelopeValues, "Audio", "AudioBusEnvelopeValues");
+				Builder.RouteEvent(RouteId_Stop,           "Audio", "AudioBusStop");
+			}
+
+			virtual bool OnEvent(uint16 RouteId, EStyle Style, const FOnEventContext& Context) override
+			{
+				LLM_SCOPE_BYNAME(TEXT("Insights/FAudioBusTraceAnalyzer"));
+
+				FAudioBusMessages& Messages = GetProvider<FAudioBusTraceProvider>().TraceMessages;
+
+				switch (RouteId)
+				{
+					case RouteId_Start:
+					{
+						Messages.StartMessages.Enqueue(FAudioBusStartMessage{ Context });
+						break;
+					}
+
+					case RouteId_HasActivity:
+					{
+						Messages.HasActivityMessages.Enqueue(FAudioBusHasActivityMessage{ Context });
+						break;
+					}
+
+					case RouteId_EnvelopeValues:
+					{
+						Messages.EnvelopeValuesMessages.Enqueue(FAudioBusEnvelopeValuesMessage{ Context });
+						break;
+					}
+
+					case RouteId_Stop:
+					{
+						Messages.StopMessages.Enqueue(FAudioBusStopMessage{ Context });
+						break;
+					}
+
+					default:
+					{
+						return OnEventFailure(RouteId, Style, Context);
+					}
+				}
+
+				const double Timestamp = Context.EventTime.AsSeconds(Context.EventData.GetValue<uint64>("Timestamp"));
+
+				{
+					TraceServices::FAnalysisSessionEditScope SessionEditScope(Session);
+					Session.UpdateDurationSeconds(Timestamp);
+				}
+
+				return OnEventSuccess(RouteId, Style, Context);
+			}
+
+		private:
+			enum : uint16
+			{
+				RouteId_Start,
+				RouteId_HasActivity,
+				RouteId_EnvelopeValues,
+				RouteId_Stop
+			};
+
+			TraceServices::IAnalysisSession& Session;
+		};
+
+		return new FAudioBusTraceAnalyzer(AsShared(), InSession);
+	}
+
+	bool FAudioBusTraceProvider::ProcessMessages()
+	{
+		// Helper lambdas
+#if !WITH_EDITOR
+		auto CreateEntry = [this](const FAudioBusMessageBase& Msg)
+		{
+			TSharedPtr<FAudioBusDashboardEntry>* ToReturn = nullptr;
+			
+			UpdateDeviceEntry(Msg.DeviceId, Msg.AudioBusId, [&ToReturn, &Msg](TSharedPtr<FAudioBusDashboardEntry>& Entry)
+			{
+				if (!Entry.IsValid())
+				{
+					Entry = MakeShared<FAudioBusDashboardEntry>();
+					Entry->AudioBusId = Msg.AudioBusId;
+				}
+
+				Entry->Timestamp = Msg.Timestamp;
+
+				ToReturn = &Entry;
+			});
+
+			return ToReturn;
+		};
+#endif // !WITH_EDITOR
+
+		auto GetEntry = [this](const FAudioBusMessageBase& Msg)
+		{
+			return FindDeviceEntry(Msg.DeviceId, Msg.AudioBusId);
+		};
+
+		// Process messages
+#if !WITH_EDITOR
+		ProcessMessageQueue<FAudioBusStartMessage>(TraceMessages.StartMessages, CreateEntry,
+		[this](const FAudioBusStartMessage& Msg, TSharedPtr<FAudioBusDashboardEntry>* OutEntry)
+		{
+			if (SessionCachedMessages.IsValid())
+			{
+				SessionCachedMessages->StartCachedMessages.EmplaceBack(Msg);
+			}
+
+			if (OutEntry)
+			{
+				FAudioBusDashboardEntry& EntryRef = *OutEntry->Get();
+				EntryRef.Name = Msg.Name;
+				EntryRef.AudioMeterInfo->NumChannels = Msg.NumChannels;
+			}
+		});
+#endif // !WITH_EDITOR
+
+		ProcessMessageQueue<FAudioBusHasActivityMessage>(TraceMessages.HasActivityMessages, GetEntry,
+		[this](const FAudioBusHasActivityMessage& Msg, TSharedPtr<FAudioBusDashboardEntry>* OutEntry)
+		{
+#if !WITH_EDITOR
+			if (SessionCachedMessages.IsValid())
+			{
+				SessionCachedMessages->HasActivityCachedMessages.EmplaceBack(Msg);
+			}
+#endif // !WITH_EDITOR
+
+			if (OutEntry)
+			{
+				FAudioBusDashboardEntry& EntryRef = *OutEntry->Get();
+				EntryRef.bHasActivity = Msg.bHasActivity;
+				EntryRef.Timestamp = Msg.Timestamp;
+			}
+		});
+
+		ProcessMessageQueue<FAudioBusEnvelopeValuesMessage>(TraceMessages.EnvelopeValuesMessages, GetEntry,
+		[this](const FAudioBusEnvelopeValuesMessage& Msg, TSharedPtr<FAudioBusDashboardEntry>* OutEntry)
+		{
+			if (OutEntry)
+			{
+				FAudioBusDashboardEntry& EntryRef = *OutEntry->Get();
+
+				EntryRef.AudioMeterInfo->EnvelopeValues = Msg.EnvelopeValues;
+				EntryRef.Timestamp = Msg.Timestamp;
+			}
+		});
+
+#if !WITH_EDITOR
+		ProcessMessageQueue<FAudioBusStopMessage>(TraceMessages.StopMessages, GetEntry,
+		[this](const FAudioBusStopMessage& Msg, TSharedPtr<FAudioBusDashboardEntry>* OutEntry)
+		{
+			if (SessionCachedMessages.IsValid())
+			{
+				SessionCachedMessages->StopCachedMessages.EmplaceBack(Msg);
+			}
+
+			if (OutEntry && (*OutEntry)->Timestamp < Msg.Timestamp)
+			{
+				RemoveDeviceEntry(Msg.DeviceId, Msg.AudioBusId);
+			}
+		});
+#endif // !WITH_EDITOR
+
+		return true;
+	}
+
+#if WITH_EDITOR
+	void FAudioBusTraceProvider::HandleOnAudioBusAssetAdded(const TSharedPtr<FAudioBusDashboardEntry>& InAudioBusDashboardEntry)
+	{
+		if (const FAudioDeviceManager* DeviceManager = FAudioDeviceManager::Get())
+		{
+			const TArray<FAudioDevice*> AudioDevices = DeviceManager->GetAudioDevices();
+
+			for (const FAudioDevice* AudioDevice : AudioDevices)
+			{
+				if (AudioDevice)
+				{
+					const TSharedPtr<FAudioBusDashboardEntry>* FoundAudioBusDashboardEntry = FindDeviceEntry(AudioDevice->DeviceID, InAudioBusDashboardEntry->AudioBusId);
+
+					if (FoundAudioBusDashboardEntry == nullptr)
+					{
+						UpdateDeviceEntry(AudioDevice->DeviceID, InAudioBusDashboardEntry->AudioBusId, [&InAudioBusDashboardEntry](TSharedPtr<FAudioBusDashboardEntry>& Entry)
+						{
+							if (!Entry.IsValid())
+							{
+								Entry = InAudioBusDashboardEntry;
+							}
+						});
+					}
+
+					OnAudioBusAssetAdded.ExecuteIfBound();
+				}
+			}
+		}
+	}
+
+	void FAudioBusTraceProvider::HandleOnAudioBusAssetRemoved(const uint32 InAudioBusId)
+	{
+		if (const FAudioDeviceManager* DeviceManager = FAudioDeviceManager::Get())
+		{
+			const TArray<FAudioDevice*> AudioDevices = DeviceManager->GetAudioDevices();
+
+			for (const FAudioDevice* AudioDevice : AudioDevices)
+			{
+				if (AudioDevice)
+				{
+					RemoveDeviceEntry(AudioDevice->DeviceID, InAudioBusId);
+
+					OnAudioBusAssetRemoved.ExecuteIfBound();
+				}
+			}
+		}
+	}
+
+	void FAudioBusTraceProvider::HandleOnAudioBusAssetListUpdated(const TArray<TSharedPtr<FAudioBusDashboardEntry>>& InAudioBusAssetEntries)
+	{
+		const FAudioDeviceManager* DeviceManager = FAudioDeviceManager::Get();
+
+		if (DeviceManager == nullptr)
+		{
+			return;
+		}
+
+		DeviceDataMap.Reset();
+
+		for (const FAudioDevice* AudioDevice : DeviceManager->GetAudioDevices())
+		{
+			if (AudioDevice == nullptr)
+			{
+				continue;
+			}
+
+			for (const auto& AudioBusDashboardEntry : InAudioBusAssetEntries)
+			{
+				const TSharedPtr<FAudioBusDashboardEntry>* FoundAudioBusDashboardEntry = FindDeviceEntry(AudioDevice->DeviceID, AudioBusDashboardEntry->AudioBusId);
+
+				if (FoundAudioBusDashboardEntry == nullptr)
+				{
+					UpdateDeviceEntry(AudioDevice->DeviceID, AudioBusDashboardEntry->AudioBusId, [&AudioBusDashboardEntry](TSharedPtr<FAudioBusDashboardEntry>& Entry)
+					{
+						if (!Entry.IsValid())
+						{
+							Entry = AudioBusDashboardEntry;
+						}
+					});
+				}
+			}
+		}
+
+		OnAudioBusAssetListUpdated.ExecuteIfBound();
+	}
+#else
+	void FAudioBusTraceProvider::InitSessionCachedMessages(TraceServices::IAnalysisSession& InSession)
+	{
+		SessionCachedMessages = MakeUnique<FAudioBusSessionCachedMessages>(InSession);
+	}
+
+	void FAudioBusTraceProvider::OnTimingViewTimeMarkerChanged(double TimeMarker)
+	{
+		using namespace FAudioBusTraceProviderPrivate;
+
+		if (!SessionCachedMessages.IsValid())
+		{
+			return;
+		}
+
+		DeviceDataMap.Empty();
+
+		// Collect all the audio bus start messages registered until this point in time 
+		for (const FAudioBusStartMessage& AudioBusStartCachedMessage : SessionCachedMessages->StartCachedMessages)
+		{
+			if (AudioBusStartCachedMessage.Timestamp > TimeMarker)
+			{
+				break;
+			}
+
+			UpdateDeviceEntry(AudioBusStartCachedMessage.DeviceId, AudioBusStartCachedMessage.AudioBusId, [&AudioBusStartCachedMessage](TSharedPtr<FAudioBusDashboardEntry>& Entry)
+			{
+				if (!Entry.IsValid())
+				{
+					Entry = MakeShared<FAudioBusDashboardEntry>();
+					Entry->AudioBusId = AudioBusStartCachedMessage.AudioBusId;
+				}
+
+				Entry->Timestamp = AudioBusStartCachedMessage.Timestamp;
+
+				Entry->Name = *AudioBusStartCachedMessage.Name;
+				Entry->AudioMeterInfo->NumChannels = AudioBusStartCachedMessage.NumChannels;
+			});
+		}
+
+		// Selectively remove start messages collected in the step above by knowing which audio buses were stopped.
+		// With this we will know what are the available audio buses at this point in time.
+		for (const FAudioBusStopMessage& AudioBusStopCachedMessage : SessionCachedMessages->StopCachedMessages)
+		{
+			if (AudioBusStopCachedMessage.Timestamp > TimeMarker)
+			{
+				break;
+			}
+
+			auto* OutEntry = FindDeviceEntry(AudioBusStopCachedMessage.DeviceId, AudioBusStopCachedMessage.AudioBusId);
+
+			if (OutEntry && (*OutEntry)->Timestamp < AudioBusStopCachedMessage.Timestamp)
+			{
+				RemoveDeviceEntry(AudioBusStopCachedMessage.DeviceId, AudioBusStopCachedMessage.AudioBusId);
+			}
+		}
+
+		const FDeviceData* DeviceData = FindFilteredDeviceData();
+		if (DeviceData)
+		{
+			// Collect has activity messages from audio buses (based on active audio buses AudioBusId)
+			struct CachedEntryInfo
+			{
+				TOptional<FAudioBusHasActivityMessage> HasActivityMessage;
+			};
+
+			TArray<uint32> AudioBusIdArray;
+			(*DeviceData).GenerateKeyArray(AudioBusIdArray);
+
+			TArray<CachedEntryInfo> CachedEntryInfos;
+			CachedEntryInfos.SetNum(AudioBusIdArray.Num());
+
+			ParallelFor(AudioBusIdArray.Num(),
+			[&AudioBusIdArray, &CachedEntryInfos, TimeMarker, this](const int32 Index)
+			{
+				const uint32 AudioBusId = AudioBusIdArray[Index];
+
+				// HasActivity
+				const FAudioBusHasActivityMessage* FoundHasActivityCachedMessage = FindClosestMessageToTimestamp(SessionCachedMessages->HasActivityCachedMessages, TimeMarker, AudioBusId);
+				if (FoundHasActivityCachedMessage)
+				{
+					CachedEntryInfos[Index].HasActivityMessage = *FoundHasActivityCachedMessage;
+				}
+			});
+
+			// Update the device entries with the collected info
+			for (const CachedEntryInfo& CachedEntryInfo : CachedEntryInfos)
+			{
+				if (CachedEntryInfo.HasActivityMessage.IsSet())
+				{
+					const FAudioBusHasActivityMessage& HasActivityMessage = CachedEntryInfo.HasActivityMessage.GetValue();
+
+					UpdateDeviceEntry(HasActivityMessage.DeviceId, HasActivityMessage.AudioBusId, [&HasActivityMessage](TSharedPtr<FAudioBusDashboardEntry>& Entry)
+					{
+						if (!Entry.IsValid())
+						{
+							Entry = MakeShared<FAudioBusDashboardEntry>();
+							Entry->AudioBusId = HasActivityMessage.AudioBusId;
+						}
+
+						Entry->Timestamp = HasActivityMessage.Timestamp;
+
+						Entry->bHasActivity = HasActivityMessage.bHasActivity;
+					});
+				}
+			}
+		}
+
+		// Call parent method to update LastMessageId
+		FTraceProviderBase::OnTimingViewTimeMarkerChanged(TimeMarker);
+	}
+#endif // WITH_EDITOR
+
+} // namespace UE::Audio::Insights
diff --git a/Engine/Plugins/AudioInsights/Source/AudioInsights/Private/Providers/AudioBusTraceProvider.h b/Engine/Plugins/AudioInsights/Source/AudioInsights/Private/Providers/AudioBusTraceProvider.h
new file mode 100644
--- /dev/null
+++ b/Engine/Plugins/AudioInsights/Source/AudioInsights/Private/Providers/AudioBusTraceProvider.h
@@ -0,0 +1,54 @@
+// Copyright Epic Games, Inc. All Rights Reserved.
+
+#pragma once
+
+#include "AudioInsightsTraceProviderBase.h"
+#include "Messages/AudioBusTraceMessages.h"
+#include "UObject/NameTypes.h"
+
+#if WITH_EDITOR
+#include "Delegates/Delegate.h"
+#include "Providers/AudioBusAssetProvider.h"
+#endif // WITH_EDITOR
+
+namespace UE::Audio::Insights
+{
+	class FAudioBusTraceProvider : public TDeviceDataMapTraceProvider<uint32, TSharedPtr<FAudioBusDashboardEntry>>, public TSharedFromThis<FAudioBusTraceProvider>
+	{
+	public:
+		FAudioBusTraceProvider();
+		virtual ~FAudioBusTraceProvider();
+
+		static FName GetName_Static();
+		
+		virtual UE::Trace::IAnalyzer* ConstructAnalyzer(TraceServices::IAnalysisSession& InSession) override;
+
+#if !WITH_EDITOR
+		virtual void InitSessionCachedMessages(TraceServices::IAnalysisSession& InSession) override;
+#else
+		DECLARE_DELEGATE(FOnAudioBusAssetsUpdated);
+		FOnAudioBusAssetsUpdated OnAudioBusAssetAdded;
+		FOnAudioBusAssetsUpdated OnAudioBusAssetRemoved;
+		FOnAudioBusAssetsUpdated OnAudioBusAssetListUpdated;
+#endif // !WITH_EDITOR
+
+	private:
+		virtual bool ProcessMessages() override;
+
+#if WITH_EDITOR
+		void HandleOnAudioBusAssetAdded(const TSharedPtr<FAudioBusDashboardEntry>& InAudioBusDashboardEntry);
+		void HandleOnAudioBusAssetRemoved(const uint32 InAudioBusId);
+		void HandleOnAudioBusAssetListUpdated(const TArray<TSharedPtr<FAudioBusDashboardEntry>>& InAudioBusAssetEntries);
+#else
+		virtual void OnTimingViewTimeMarkerChanged(double TimeMarker) override;
+
+		TUniquePtr<FAudioBusSessionCachedMessages> SessionCachedMessages;
+#endif // WITH_EDITOR
+
+		FAudioBusMessages TraceMessages;
+
+#if WITH_EDITOR
+		FAudioBusAssetProvider AudioBusAssetProvider;
+#endif // WITH_EDITOR
+	};
+} // namespace UE::Audio::Insights
diff --git a/Engine/Plugins/AudioInsights/Source/AudioInsights/Public/Messages/AudioBusTraceMessages.h b/Engine/Plugins/AudioInsights/Source/AudioInsights/Public/Messages/AudioBusTraceMessages.h
new file mode 100644
--- /dev/null
+++ b/Engine/Plugins/AudioInsights/Source/AudioInsights/Public/Messages/AudioBusTraceMessages.h
@@ -0,0 +1,138 @@
+// Copyright Epic Games, Inc. All Rights Reserved.
+#pragma once
+
+#include "Providers/AudioMeterProvider.h"
+#include "Messages/AnalyzerMessageQueue.h"
+#include "Trace/Analyzer.h"
+#include "Views/TableDashboardViewFactory.h"
+
+#if !WITH_EDITOR
+#include "Common/PagedArray.h"
+#include "TraceServices/Model/AnalysisSession.h"
+#endif // !WITH_EDITOR
+
+namespace UE::Audio::Insights
+{
+	struct FAudioBusMessageBase
+	{
+		FAudioBusMessageBase() = default;
+
+		FAudioBusMessageBase(const UE::Trace::IAnalyzer::FOnEventContext& InContext)
+		{
+			const UE::Trace::IAnalyzer::FEventData& EventData = InContext.EventData;
+
+			DeviceId  = static_cast<::Audio::FDeviceId>(EventData.GetValue<uint32>("DeviceId"));
+			AudioBusId  = EventData.GetValue<uint32>("AudioBusId");
+			Timestamp = InContext.EventTime.AsSeconds(EventData.GetValue<uint64>("Timestamp"));
+		}
+
+		::Audio::FDeviceId DeviceId = INDEX_NONE;
+		uint32 AudioBusId = INDEX_NONE;
+		double Timestamp = 0.0;
+	};
+
+	struct FAudioBusStartMessage : public FAudioBusMessageBase
+	{
+		FAudioBusStartMessage() = default;
+
+		FAudioBusStartMessage(const UE::Trace::IAnalyzer::FOnEventContext& InContext)
+			: FAudioBusMessageBase(InContext)
+		{
+			const UE::Trace::IAnalyzer::FEventData& EventData = InContext.EventData;
+
+			EventData.GetString("Name", Name);
+			NumChannels = EventData.GetValue<int32>("NumChannels");
+		}
+
+		FString Name;
+		int32 NumChannels = 0;
+	};
+
+	struct FAudioBusHasActivityMessage : public FAudioBusMessageBase
+	{
+		FAudioBusHasActivityMessage() = default;
+
+		FAudioBusHasActivityMessage(const UE::Trace::IAnalyzer::FOnEventContext& InContext)
+			: FAudioBusMessageBase(InContext)
+		{
+			const UE::Trace::IAnalyzer::FEventData& EventData = InContext.EventData;
+
+			bHasActivity = EventData.GetValue<bool>("HasActivity");
+		}
+
+		bool bHasActivity = false;
+	};
+
+	struct FAudioBusEnvelopeValuesMessage : public FAudioBusMessageBase
+	{
+		FAudioBusEnvelopeValuesMessage() = default;
+
+		FAudioBusEnvelopeValuesMessage(const UE::Trace::IAnalyzer::FOnEventContext& InContext)
+			: FAudioBusMessageBase(InContext)
+		{
+			const UE::Trace::IAnalyzer::FEventData& EventData = InContext.EventData;
+
+			const TArrayView<const float> EnvelopeValuesView = EventData.GetArrayView<float>("EnvelopeValues");
+
+			for (const auto& EnvelopeValue : EnvelopeValuesView)
+			{
+				EnvelopeValues.Emplace(EnvelopeValue);
+			}
+		}
+
+		TArray<float> EnvelopeValues;
+	};
+
+	using FAudioBusStopMessage = FAudioBusMessageBase;
+
+	class FAudioBusMessages
+	{
+		TAnalyzerMessageQueue<FAudioBusStartMessage> StartMessages{ 0.1 };
+		TAnalyzerMessageQueue<FAudioBusHasActivityMessage> HasActivityMessages{ 0.1 };
+		TAnalyzerMessageQueue<FAudioBusEnvelopeValuesMessage> EnvelopeValuesMessages{ 0.1 };
+		TAnalyzerMessageQueue<FAudioBusStopMessage> StopMessages{ 0.1 };
+
+		friend class FAudioBusTraceProvider;
+	};
+
+#if !WITH_EDITOR
+	struct FAudioBusSessionCachedMessages
+	{
+		FAudioBusSessionCachedMessages(TraceServices::IAnalysisSession& InSession)
+			: StartCachedMessages(InSession.GetLinearAllocator(), 4096)
+			, HasActivityCachedMessages(InSession.GetLinearAllocator(), 4096)
+			, StopCachedMessages(InSession.GetLinearAllocator(), 4096)
+		{
+
+		}
+
+		TraceServices::TPagedArray<FAudioBusStartMessage> StartCachedMessages;
+		TraceServices::TPagedArray<FAudioBusHasActivityMessage> HasActivityCachedMessages;
+		TraceServices::TPagedArray<FAudioBusStopMessage> StopCachedMessages;
+	};
+#endif // !WITH_EDITOR
+
+	struct FAudioBusDashboardEntry : public IObjectDashboardEntry
+	{
+		FAudioBusDashboardEntry()
+			: AudioMeterInfo(MakeShared<FAudioMeterInfo>())
+		{
+
+		}
+
+		virtual ~FAudioBusDashboardEntry() = default;
+
+		virtual FText GetDisplayName() const override { return FText::FromString(FSoftObjectPath(Name).GetAssetName()); }
+		virtual const UObject* GetObject() const override { return FSoftObjectPath(Name).ResolveObject(); }
+		virtual UObject* GetObject() override { return FSoftObjectPath(Name).ResolveObject(); }
+		virtual bool IsValid() const override { return AudioBusId != static_cast<uint32>(INDEX_NONE); }
+
+		uint32 AudioBusId = INDEX_NONE;
+		double Timestamp = 0.0;
+
+		TSharedRef<FAudioMeterInfo> AudioMeterInfo;
+
+		FString Name;
+		bool bHasActivity = false;
+	};
+} // namespace UE::Audio::Insights
