[Audio Insights] Added AudioBus dashboard for standalone

diff --git a/Engine/Plugins/AudioInsights/Source/AudioInsights/Private/AudioInsightsModule.cpp b/Engine/Plugins/AudioInsights/Source/AudioInsights/Private/AudioInsightsModule.cpp
--- a/Engine/Plugins/AudioInsights/Source/AudioInsights/Private/AudioInsightsModule.cpp
+++ b/Engine/Plugins/AudioInsights/Source/AudioInsights/Private/AudioInsightsModule.cpp
@@ -15,6 +15,7 @@
 
 #if !WITH_EDITOR
 #include "AudioInsightsComponent.h"
+#include "Views/AudioBusDashboardViewFactory.h"
 #include "Views/AudioEventLogDashboardViewFactory.h"
 #include "Views/AudioMetersPanelDashboardViewFactory.h"
 #include "Views/SoundDashboardViewFactory.h"
@@ -60,6 +61,7 @@ namespace UE::Audio::Insights
 			DashboardFactory->RegisterViewFactory(MakeShared<FVirtualLoopDashboardViewFactory>());
 			DashboardFactory->RegisterViewFactory(PlotsDashboard);
 			DashboardFactory->RegisterViewFactory(MakeShared<FSubmixDashboardViewFactory>());
+			DashboardFactory->RegisterViewFactory(MakeShared<FAudioBusDashboardViewFactory>());
 			DashboardFactory->RegisterViewFactory(MakeShared<FAudioMetersPanelDashboardViewFactory>());
 
 			PlotsDashboard->InitPlots(SoundsDashboard);
diff --git a/Engine/Plugins/AudioInsights/Source/AudioInsights/Private/Views/AudioBusDashboardViewFactory.cpp b/Engine/Plugins/AudioInsights/Source/AudioInsights/Private/Views/AudioBusDashboardViewFactory.cpp
new file mode 100644
--- /dev/null
+++ b/Engine/Plugins/AudioInsights/Source/AudioInsights/Private/Views/AudioBusDashboardViewFactory.cpp
@@ -0,0 +1,356 @@
+// Copyright Epic Games, Inc. All Rights Reserved.
+#include "Views/AudioBusDashboardViewFactory.h"
+
+#include "AudioInsightsStyle.h"
+#include "IAudioInsightsModule.h"
+#include "Internationalization/Text.h"
+#include "Providers/AudioMeterProvider.h"
+#include "Providers/AudioBusTraceProvider.h"
+#include "Widgets/Input/SCheckBox.h"
+
+#if WITH_EDITOR
+#include "Editor.h"
+#include "Subsystems/AssetEditorSubsystem.h"
+#endif // WITH_EDITOR
+
+#define LOCTEXT_NAMESPACE "AudioInsights"
+
+namespace UE::Audio::Insights
+{
+	namespace FAudioBusDashboardViewFactoryPrivate
+	{
+		const FAudioBusDashboardEntry& CastEntry(const IDashboardDataViewEntry& InData)
+		{
+			return static_cast<const FAudioBusDashboardEntry&>(InData);
+		};
+
+		void SendAudioBusEnvelopeFollowerCVar(const bool bInIsAudioBusChecked, const FString& InAudioBusName)
+		{
+#if WITH_EDITOR
+			const IAudioInsightsModule& AudioInsightsModule = IAudioInsightsModule::GetEditorChecked();
+#else
+			const IAudioInsightsModule& AudioInsightsModule = IAudioInsightsModule::GetChecked();
+#endif // WITH_EDITOR
+
+			const ::Audio::FDeviceId AudioDeviceId = AudioInsightsModule.GetDeviceId();
+
+			const FString AudioBusCommandStr = FString::Printf(TEXT("au.%sAudioBusEnvelopeFollower_AD%d %s"), bInIsAudioBusChecked ? TEXT("Start") : TEXT("Stop"), AudioDeviceId, *InAudioBusName);
+			
+			IAudioInsightsTraceModule& AudioInsightsTraceModule = IAudioInsightsModule::GetChecked().GetTraceModule();
+			AudioInsightsTraceModule.ExecuteConsoleCommand(AudioBusCommandStr);
+		}
+	}
+
+	FAudioBusDashboardViewFactory::FAudioBusDashboardViewFactory()
+	{
+		IAudioInsightsTraceModule& AudioInsightsTraceModule = IAudioInsightsModule::GetChecked().GetTraceModule();
+
+		AudioBusProvider = MakeShared<FAudioBusTraceProvider>();
+
+		AudioInsightsTraceModule.AddTraceProvider(AudioBusProvider);
+
+		Providers = TArray<TSharedPtr<FTraceProviderBase>>
+		{
+			AudioBusProvider
+		};
+
+		SortByColumn = "Name";
+		SortMode     = EColumnSortMode::Ascending;
+
+#if WITH_EDITOR
+		AudioBusProvider->OnAudioBusAssetAdded.BindRaw(this, &FAudioBusDashboardViewFactory::RequestListRefresh);
+		AudioBusProvider->OnAudioBusAssetRemoved.BindRaw(this, &FAudioBusDashboardViewFactory::RequestListRefresh);
+		AudioBusProvider->OnAudioBusAssetListUpdated.BindRaw(this, &FAudioBusDashboardViewFactory::RequestListRefresh);
+#endif // WITH_EDITOR
+	}
+
+	FAudioBusDashboardViewFactory::~FAudioBusDashboardViewFactory()
+	{
+#if WITH_EDITOR
+		AudioBusProvider->OnAudioBusAssetAdded.Unbind();
+		AudioBusProvider->OnAudioBusAssetRemoved.Unbind();
+		AudioBusProvider->OnAudioBusAssetListUpdated.Unbind();
+#endif // WITH_EDITOR
+	}
+
+	FName FAudioBusDashboardViewFactory::GetName() const
+	{
+		return "AudioBus";
+	}
+
+	FText FAudioBusDashboardViewFactory::GetDisplayName() const
+	{
+		return LOCTEXT("AudioDashboard_AudioBus_DisplayName", "AudioBuses");
+	}
+	
+	FSlateIcon FAudioBusDashboardViewFactory::GetIcon() const
+	{
+		return FSlateStyle::Get().CreateIcon("AudioInsights.Icon");
+	}
+
+	EDefaultDashboardTabStack FAudioBusDashboardViewFactory::GetDefaultTabStack() const
+	{
+		return EDefaultDashboardTabStack::Analysis;
+	}
+
+	TSharedRef<SWidget> FAudioBusDashboardViewFactory::GenerateWidgetForColumn(TSharedRef<IDashboardDataViewEntry> InRowData, const FName& InColumnName)
+	{
+		using namespace FAudioBusDashboardViewFactoryPrivate;
+
+		if (InColumnName == "Active")
+		{
+			static const FLinearColor DarkGreen(0.027f, 0.541f, 0.22f);
+			static const float Radius = 4.0f;
+			static const FVector2f Size(7.0f, 7.0f);
+			static const FSlateRoundedBoxBrush GreenRoundedBrush(DarkGreen, Radius, Size);
+
+			return SNew(SBox)
+				.Clipping(EWidgetClipping::ClipToBounds)
+				.Padding(6.0f)
+				.Visibility_Lambda([InRowData]()
+				{ 
+					const FAudioBusDashboardEntry& AudioBusDashboardEntry = CastEntry(InRowData.Get());
+					return AudioBusDashboardEntry.bHasActivity ? EVisibility::Visible : EVisibility::Hidden;
+				})
+				[
+					SNew(SImage)
+					.Image(&GreenRoundedBrush)
+				];
+		}
+		else if (InColumnName == "Name")
+		{
+			const FColumnData& ColumnData = GetColumns()[InColumnName];
+			const FText ValueText = ColumnData.GetDisplayValue(InRowData.Get());
+
+			if (ValueText.IsEmpty())
+			{
+				return SNullWidget::NullWidget;
+			}
+
+			const FAudioBusDashboardEntry& AudioBusDashboardEntry = CastEntry(*InRowData);
+			const uint32 AudioBusId = AudioBusDashboardEntry.AudioBusId;
+
+			const bool bIsAudioBusChecked = AudioBusCheckboxCheckedStates.FindOrAdd(AudioBusId);
+
+			if (bIsAudioBusChecked)
+			{
+				SendAudioBusEnvelopeFollowerCVar(bIsAudioBusChecked, AudioBusDashboardEntry.GetDisplayName().ToString());
+				FAudioMeterProvider::OnUpdateAudioMeterInfo.ExecuteIfBound(AudioBusId, AudioBusDashboardEntry.AudioMeterInfo, AudioBusDashboardEntry.GetDisplayName());
+			}
+
+			return SNew(SHorizontalBox)
+				.Clipping(EWidgetClipping::ClipToBounds)
+				+ SHorizontalBox::Slot()
+				.AutoWidth()
+				[
+					SNew(SCheckBox)
+					.IsChecked_Lambda([this, AudioBusId]()
+					{
+						const bool* FoundAudioBusCheckedState = AudioBusCheckboxCheckedStates.Find(AudioBusId);
+						const bool bIsAudioBusChecked = FoundAudioBusCheckedState && *FoundAudioBusCheckedState == true;
+
+						return bIsAudioBusChecked ? ECheckBoxState::Checked : ECheckBoxState::Unchecked;
+					})
+					.OnCheckStateChanged_Lambda([this, InRowData](ECheckBoxState NewState)
+					{
+						const bool bIsChecked = NewState == ECheckBoxState::Checked;
+
+						const FAudioBusDashboardEntry& AudioBusDashboardEntry = CastEntry(*InRowData);
+
+						if (bool* FoundAudioBusCheckedState = AudioBusCheckboxCheckedStates.Find(AudioBusDashboardEntry.AudioBusId))
+						{
+							*FoundAudioBusCheckedState = bIsChecked;
+						}
+
+						SendAudioBusEnvelopeFollowerCVar(bIsChecked, AudioBusDashboardEntry.GetDisplayName().ToString());
+
+						if (bIsChecked)
+						{
+							FAudioMeterProvider::OnAddAudioMeter.ExecuteIfBound(AudioBusDashboardEntry.AudioBusId, AudioBusDashboardEntry.AudioMeterInfo, AudioBusDashboardEntry.GetDisplayName(), FSlateColor(FColor(80, 200, 255)));
+						}
+						else
+						{
+							FAudioMeterProvider::OnRemoveAudioMeter.ExecuteIfBound(AudioBusDashboardEntry.AudioBusId);
+						}
+					})
+				]
+				+ SHorizontalBox::Slot()
+				.AutoWidth()
+				[
+					SNew(SBox)
+					.MinDesiredWidth(5.0f)
+				]
+				+ SHorizontalBox::Slot()
+				.AutoWidth()
+				[
+					SNew(STextBlock)
+					.Text(ValueText)
+					.OnDoubleClicked_Lambda([this, InRowData](const FGeometry& MyGeometry, const FPointerEvent& PointerEvent)
+					{
+#if WITH_EDITOR
+						if (GEditor)
+						{
+							const TSharedPtr<IObjectDashboardEntry> ObjectData = StaticCastSharedPtr<IObjectDashboardEntry>(InRowData.ToSharedPtr());
+							if (ObjectData.IsValid())
+							{
+								const TObjectPtr<UObject> Object = ObjectData->GetObject();
+								if (Object && Object->IsAsset())
+								{
+									GEditor->GetEditorSubsystem<UAssetEditorSubsystem>()->OpenEditorForAsset(Object);
+									return FReply::Handled();
+								}
+							}
+						}
+#endif // WITH_EDITOR
+
+						return FReply::Unhandled();
+					})
+				];
+		}
+
+		return SNullWidget::NullWidget;
+	}
+
+	void FAudioBusDashboardViewFactory::ProcessEntries(FTraceTableDashboardViewFactory::EProcessReason Reason)
+	{
+		const FString FilterString = GetSearchFilterText().ToString();
+		
+		FTraceTableDashboardViewFactory::FilterEntries<FAudioBusTraceProvider>([&FilterString](const IDashboardDataViewEntry& Entry)
+		{
+			const FAudioBusDashboardEntry& AudioBusEntry = static_cast<const FAudioBusDashboardEntry&>(Entry);
+			
+			return AudioBusEntry.GetDisplayName().ToString().Contains(FilterString);
+		});
+	}
+
+	TSharedRef<SWidget> FAudioBusDashboardViewFactory::MakeWidget(TSharedRef<SDockTab> OwnerTab, const FSpawnTabArgs& SpawnTabArgs)
+	{
+		using namespace FAudioBusDashboardViewFactoryPrivate;
+
+		if (!DashboardWidget.IsValid())
+		{
+			DashboardWidget = FTraceTableDashboardViewFactory::MakeWidget(OwnerTab, SpawnTabArgs);
+
+			if (FilteredEntriesListView.IsValid())
+			{
+				FilteredEntriesListView->SetSelectionMode(ESelectionMode::Single);
+			}
+		}
+
+		// Update audio meters (if any has been checked previously)
+		for (const TSharedPtr<IDashboardDataViewEntry>& Entry : DataViewEntries)
+		{
+			const FAudioBusDashboardEntry& AudioBusDashboardEntry = CastEntry(*Entry);
+
+			if (const bool* bFoundAudioBusCheckboxCheckedState = AudioBusCheckboxCheckedStates.Find(AudioBusDashboardEntry.AudioBusId))
+			{
+				if (*bFoundAudioBusCheckboxCheckedState)
+				{
+					FAudioMeterProvider::OnUpdateAudioMeterInfo.ExecuteIfBound(AudioBusDashboardEntry.AudioBusId, AudioBusDashboardEntry.AudioMeterInfo, AudioBusDashboardEntry.GetDisplayName());
+				}
+			}
+		}
+
+		return DashboardWidget->AsShared();
+	}
+
+	const TMap<FName, FTraceTableDashboardViewFactory::FColumnData>& FAudioBusDashboardViewFactory::GetColumns() const
+	{
+		using namespace FAudioBusDashboardViewFactoryPrivate;
+
+		auto CreateColumnData = []()
+		{
+			return TMap<FName, FTraceTableDashboardViewFactory::FColumnData>
+			{
+				{
+					"Active",
+					{
+						LOCTEXT("AudioBusDashboard_ActiveDisplayName", "Active"),
+						[](const IDashboardDataViewEntry& InData) { return FText::GetEmpty(); },
+						nullptr		/* GetIconName */,
+						false,		/* bDefaultHidden */
+						0.08f,		/* FillWidth */
+						EHorizontalAlignment::HAlign_Center
+					}
+				},
+				{
+					"Name",
+					{
+						LOCTEXT("AudioBusDashboard_NameColumnDisplayName", "Name"),
+						[](const IDashboardDataViewEntry& InData) { return CastEntry(InData).GetDisplayName(); },
+						nullptr		/* GetIconName */,
+						false,		/* bDefaultHidden */
+						0.92f,		/* FillWidth */
+						EHorizontalAlignment::HAlign_Left
+					}
+				}
+			};
+		};
+		
+		static const TMap<FName, FTraceTableDashboardViewFactory::FColumnData> ColumnData = CreateColumnData();
+		
+		return ColumnData;
+	}
+
+	void FAudioBusDashboardViewFactory::SortTable()
+	{
+		using namespace FAudioBusDashboardViewFactoryPrivate;
+
+		if (SortByColumn == "Active")
+		{
+			if (SortMode == EColumnSortMode::Ascending)
+			{
+				DataViewEntries.Sort([](const TSharedPtr<IDashboardDataViewEntry>& A, const TSharedPtr<IDashboardDataViewEntry>& B)
+				{
+					const FAudioBusDashboardEntry& AData = CastEntry(*A.Get());
+					const FAudioBusDashboardEntry& BData = CastEntry(*B.Get());
+
+					return !AData.bHasActivity && BData.bHasActivity;
+				});
+			}
+			else if (SortMode == EColumnSortMode::Descending)
+			{
+				DataViewEntries.Sort([](const TSharedPtr<IDashboardDataViewEntry>& A, const TSharedPtr<IDashboardDataViewEntry>& B)
+				{
+					const FAudioBusDashboardEntry& AData = CastEntry(*A.Get());
+					const FAudioBusDashboardEntry& BData = CastEntry(*B.Get());
+
+					return !BData.bHasActivity && AData.bHasActivity;
+				});
+			}
+		}
+		else if (SortByColumn == "Name")
+		{
+			if (SortMode == EColumnSortMode::Ascending)
+			{
+				DataViewEntries.Sort([](const TSharedPtr<IDashboardDataViewEntry>& A, const TSharedPtr<IDashboardDataViewEntry>& B)
+				{
+					const FAudioBusDashboardEntry& AData = CastEntry(*A.Get());
+					const FAudioBusDashboardEntry& BData = CastEntry(*B.Get());
+
+					return AData.GetDisplayName().CompareToCaseIgnored(BData.GetDisplayName()) < 0;
+				});
+			}
+			else if (SortMode == EColumnSortMode::Descending)
+			{
+				DataViewEntries.Sort([](const TSharedPtr<IDashboardDataViewEntry>& A, const TSharedPtr<IDashboardDataViewEntry>& B)
+				{
+					const FAudioBusDashboardEntry& AData = CastEntry(*A.Get());
+					const FAudioBusDashboardEntry& BData = CastEntry(*B.Get());
+
+					return BData.GetDisplayName().CompareToCaseIgnored(AData.GetDisplayName()) < 0;
+				});
+			}
+		}
+	}
+
+	void FAudioBusDashboardViewFactory::RequestListRefresh()
+	{
+		if (FilteredEntriesListView.IsValid())
+		{
+			FilteredEntriesListView->RequestListRefresh();
+		}
+	}
+} // namespace UE::Audio::Insights
+
+#undef LOCTEXT_NAMESPACE
diff --git a/Engine/Plugins/AudioInsights/Source/AudioInsights/Public/Views/AudioBusDashboardViewFactory.h b/Engine/Plugins/AudioInsights/Source/AudioInsights/Public/Views/AudioBusDashboardViewFactory.h
new file mode 100644
--- /dev/null
+++ b/Engine/Plugins/AudioInsights/Source/AudioInsights/Public/Views/AudioBusDashboardViewFactory.h
@@ -0,0 +1,40 @@
+// Copyright Epic Games, Inc. All Rights Reserved.
+#pragma once
+
+#include "Templates/SharedPointer.h"
+#include "Views/TableDashboardViewFactory.h"
+
+#define UE_API AUDIOINSIGHTS_API
+
+namespace UE::Audio::Insights
+{
+	class FAudioBusTraceProvider;
+
+	class FAudioBusDashboardViewFactory : public FTraceObjectTableDashboardViewFactory
+	{
+	public:
+		UE_API FAudioBusDashboardViewFactory();
+		UE_API virtual ~FAudioBusDashboardViewFactory();
+
+		UE_API virtual FName GetName() const override;
+		UE_API virtual FText GetDisplayName() const override;
+		UE_API virtual FSlateIcon GetIcon() const override;
+		UE_API virtual EDefaultDashboardTabStack GetDefaultTabStack() const override;
+		UE_API virtual TSharedRef<SWidget> MakeWidget(TSharedRef<SDockTab> OwnerTab, const FSpawnTabArgs& SpawnTabArgs) override;
+
+	protected:
+		UE_API virtual TSharedRef<SWidget> GenerateWidgetForColumn(TSharedRef<IDashboardDataViewEntry> InRowData, const FName& InColumnName) override;
+		UE_API virtual void ProcessEntries(FTraceTableDashboardViewFactory::EProcessReason Reason) override;
+		UE_API virtual const TMap<FName, FTraceTableDashboardViewFactory::FColumnData>& GetColumns() const override;
+
+		UE_API virtual void SortTable() override;
+
+		UE_API void RequestListRefresh();
+
+		TSharedPtr<FAudioBusTraceProvider> AudioBusProvider;
+
+		TMap<uint32, bool> AudioBusCheckboxCheckedStates;
+	};
+} // namespace UE::Audio::Insights
+
+#undef UE_API
