[Audio Insights] Browse and edit sound asset from selected event in event log.

diff --git a/Engine/Plugins/AudioInsights/Source/AudioInsights/Private/AudioEventLogEditorCommands.cpp b/Engine/Plugins/AudioInsights/Source/AudioInsights/Private/AudioEventLogEditorCommands.cpp
new file mode 100644
--- /dev/null
+++ b/Engine/Plugins/AudioInsights/Source/AudioInsights/Private/AudioEventLogEditorCommands.cpp
@@ -0,0 +1,25 @@
+// Copyright Epic Games, Inc. All Rights Reserved.
+#include "AudioEventLogEditorCommands.h"
+
+#include "AudioInsightsStyle.h"
+#include "Framework/Commands/InputChord.h"
+#include "Framework/Commands/UICommandInfo.h"
+
+#define LOCTEXT_NAMESPACE "AudioInsights"
+
+namespace UE::Audio::Insights
+{
+	FAudioEventLogEditorCommands::FAudioEventLogEditorCommands()
+		: TCommands<FAudioEventLogEditorCommands>("AudioEventLogEditorCommands", LOCTEXT("AudioEventLogEditorCommands_ContextDescText", "Audio Event Log Editor Commands"), NAME_None, UE::Audio::Insights::FSlateStyle::GetStyleName())
+	{
+	}
+
+	void FAudioEventLogEditorCommands::RegisterCommands()
+	{
+		UI_COMMAND(Browse, "Browse To Asset", "Browses to the selected sound asset in the content browser.", EUserInterfaceActionType::Button, FInputChord(EKeys::B, EModifierKey::Control));
+		UI_COMMAND(Edit, "Edit", "Opens the selected sound for edit.", EUserInterfaceActionType::Button, FInputChord(EKeys::E, EModifierKey::Control));
+	}
+
+} // namespace UE::Audio::Insights
+
+#undef LOCTEXT_NAMESPACE
diff --git a/Engine/Plugins/AudioInsights/Source/AudioInsights/Private/AudioEventLogEditorCommands.h b/Engine/Plugins/AudioInsights/Source/AudioInsights/Private/AudioEventLogEditorCommands.h
new file mode 100644
--- /dev/null
+++ b/Engine/Plugins/AudioInsights/Source/AudioInsights/Private/AudioEventLogEditorCommands.h
@@ -0,0 +1,25 @@
+// Copyright Epic Games, Inc. All Rights Reserved.
+#pragma once
+
+#include "Framework/Commands/Commands.h"
+#include "Templates/SharedPointer.h"
+
+class FUICommandInfo;
+
+namespace UE::Audio::Insights
+{
+	class FAudioEventLogEditorCommands : public TCommands<FAudioEventLogEditorCommands>
+	{
+	public:
+		FAudioEventLogEditorCommands();
+
+		virtual void RegisterCommands() override;
+
+		TSharedPtr<const FUICommandInfo> GetBrowseCommand() const { return Browse; }
+		TSharedPtr<const FUICommandInfo> GetEditCommand() const { return Edit; }
+
+	private:
+		TSharedPtr<FUICommandInfo> Browse;
+		TSharedPtr<FUICommandInfo> Edit;
+	};
+} // namespace UE::Audio::Insights
diff --git a/Engine/Plugins/AudioInsights/Source/AudioInsights/Private/Views/AudioEventLogDashboardViewFactory.cpp b/Engine/Plugins/AudioInsights/Source/AudioInsights/Private/Views/AudioEventLogDashboardViewFactory.cpp
--- a/Engine/Plugins/AudioInsights/Source/AudioInsights/Private/Views/AudioEventLogDashboardViewFactory.cpp
+++ b/Engine/Plugins/AudioInsights/Source/AudioInsights/Private/Views/AudioEventLogDashboardViewFactory.cpp
@@ -1,6 +1,7 @@
 // Copyright Epic Games, Inc. All Rights Reserved.
 #include "Views/AudioEventLogDashboardViewFactory.h"
 
+#include "AudioEventLogEditorCommands.h"
 #include "AudioInsightsModule.h"
 #include "AudioInsightsStyle.h"
 #include "AudioMixerTrace.h"
@@ -119,7 +120,8 @@ namespace UE::Audio::Insights
 		SortByColumn = TimestampColumnID;
 		SortMode = EColumnSortMode::Ascending;
 
-		CommandList = MakeShared<FUICommandList>();
+		FAudioEventLogEditorCommands::Register();
+		BindCommands();
 
 #if WITH_EDITOR
 		FAudioEventLogSettings::OnReadSettings.AddRaw(this, &FAudioEventLogDashboardViewFactory::OnReadEditorSettings);
@@ -133,6 +135,20 @@ namespace UE::Audio::Insights
 		FAudioEventLogSettings::OnReadSettings.RemoveAll(this);
 		FAudioEventLogSettings::OnWriteSettings.RemoveAll(this);
 #endif // WITH_EDITOR
+
+		FAudioEventLogEditorCommands::Unregister();
+	}
+
+	void FAudioEventLogDashboardViewFactory::BindCommands()
+	{
+		CommandList = MakeShared<FUICommandList>();
+
+		const FAudioEventLogEditorCommands& Commands = FAudioEventLogEditorCommands::Get();
+
+#if WITH_EDITOR
+		CommandList->MapAction(Commands.GetBrowseCommand(), FExecuteAction::CreateLambda([this]() { BrowseToAsset(); }));
+		CommandList->MapAction(Commands.GetEditCommand(), FExecuteAction::CreateLambda([this]() { OpenAsset(); }));
+#endif // WITH_EDITOR
 	}
 
 	FName FAudioEventLogDashboardViewFactory::GetName() const
@@ -473,6 +489,26 @@ namespace UE::Audio::Insights
 		return LOCTEXT("AudioEventLog_Settings_TooltipText", "Event log settings");
 	}
 
+#if WITH_EDITOR
+	TSharedPtr<SWidget> FAudioEventLogDashboardViewFactory::OnConstructContextMenu()
+	{
+		constexpr bool bShouldCloseWindowAfterMenuSelection = true;
+		FMenuBuilder MenuBuilder(bShouldCloseWindowAfterMenuSelection, CommandList);
+
+		const FAudioEventLogEditorCommands& Commands = FAudioEventLogEditorCommands::Get();
+
+		MenuBuilder.BeginSection("EventLogDashboardActions", LOCTEXT("EventLogDashboard_Actions_HeaderText", "Editor Options"));
+		{
+			MenuBuilder.AddMenuEntry(Commands.GetBrowseCommand(), NAME_None, TAttribute<FText>(), TAttribute<FText>(), FSlateStyle::Get().CreateIcon("AudioInsights.Icon.SoundDashboard.Browse"));
+			MenuBuilder.AddMenuEntry(Commands.GetEditCommand(), NAME_None, TAttribute<FText>(), TAttribute<FText>(), FSlateStyle::Get().CreateIcon("AudioInsights.Icon.SoundDashboard.Edit"));
+		}
+
+		MenuBuilder.EndSection();
+
+		return MenuBuilder.MakeWidget();
+	}
+#endif // WITH_EDITOR
+
 	void FAudioEventLogDashboardViewFactory::OnSelectionChanged(TSharedPtr<IDashboardDataViewEntry> SelectedItem, ESelectInfo::Type SelectInfo)
 	{
 		using namespace AudioEventLogPrivate;
diff --git a/Engine/Plugins/AudioInsights/Source/AudioInsights/Public/Views/AudioEventLogDashboardViewFactory.h b/Engine/Plugins/AudioInsights/Source/AudioInsights/Public/Views/AudioEventLogDashboardViewFactory.h
--- a/Engine/Plugins/AudioInsights/Source/AudioInsights/Public/Views/AudioEventLogDashboardViewFactory.h
+++ b/Engine/Plugins/AudioInsights/Source/AudioInsights/Public/Views/AudioEventLogDashboardViewFactory.h
@@ -74,6 +74,10 @@ namespace UE::Audio::Insights
 		UE_API virtual TSharedRef<SWidget> OnGetSettingsMenuContent() override;
 		UE_API virtual FText OnGetSettingsMenuToolTip() override;
 
+#if WITH_EDITOR
+		UE_API virtual TSharedPtr<SWidget> OnConstructContextMenu() override;
+#endif //WITH_EDITOR
+
 		UE_API virtual void OnSelectionChanged(TSharedPtr<IDashboardDataViewEntry> SelectedItem, ESelectInfo::Type SelectInfo) override;
 		virtual bool ClearSelectionOnClick() const override { return true; }
 		UE_API virtual const FTableRowStyle* GetRowStyle() const override;
@@ -106,6 +110,8 @@ namespace UE::Audio::Insights
 
 		TSharedPtr<FUICommandList> CommandList;
 		TSharedPtr<FVisibleColumnsSettingsMenu<FAudioEventLogVisibleColumns>> VisibleColumnsSettingsMenu;
+
+		void BindCommands();
 	};
 } // namespace UE::Audio::Insights
 
