[Audio Insights] AudioBus dashboard: added AudioBusType so we can filter asset based vs code generated ones + send AudioBusId instead of AudioBusName via CVar + adjusted some WITH_EDITOR ifdefs

diff --git a/Engine/Plugins/AudioInsights/Source/AudioInsights/Private/Providers/AudioBusAssetProvider.cpp b/Engine/Plugins/AudioInsights/Source/AudioInsights/Private/Providers/AudioBusAssetProvider.cpp
--- a/Engine/Plugins/AudioInsights/Source/AudioInsights/Private/Providers/AudioBusAssetProvider.cpp
+++ b/Engine/Plugins/AudioInsights/Source/AudioInsights/Private/Providers/AudioBusAssetProvider.cpp
@@ -99,6 +99,7 @@ namespace UE::Audio::Insights
 			AudioBusDashboardEntry->AudioBusId = AudioBus ? AudioBus->GetUniqueID() : INDEX_NONE; //GetTypeHash(InAssetData.GetObjectPathString());
 			AudioBusDashboardEntry->Name = InAssetData.GetObjectPathString();
 			AudioBusDashboardEntry->AudioMeterInfo->NumChannels = AudioBus->GetNumChannels();
+			AudioBusDashboardEntry->AudioBusType = EAudioBusType::AssetBased;
 
 			if (bInShouldNotify)
 			{
diff --git a/Engine/Plugins/AudioInsights/Source/AudioInsights/Private/Providers/AudioBusTraceProvider.cpp b/Engine/Plugins/AudioInsights/Source/AudioInsights/Private/Providers/AudioBusTraceProvider.cpp
--- a/Engine/Plugins/AudioInsights/Source/AudioInsights/Private/Providers/AudioBusTraceProvider.cpp
+++ b/Engine/Plugins/AudioInsights/Source/AudioInsights/Private/Providers/AudioBusTraceProvider.cpp
@@ -149,7 +149,6 @@ namespace UE::Audio::Insights
 	bool FAudioBusTraceProvider::ProcessMessages()
 	{
 		// Helper lambdas
-#if !WITH_EDITOR
 		auto CreateEntry = [this](const FAudioBusMessageBase& Msg)
 		{
 			TSharedPtr<FAudioBusDashboardEntry>* ToReturn = nullptr;
@@ -169,7 +168,6 @@ namespace UE::Audio::Insights
 
 			return ToReturn;
 		};
-#endif // !WITH_EDITOR
 
 		auto GetEntry = [this](const FAudioBusMessageBase& Msg)
 		{
@@ -177,23 +175,36 @@ namespace UE::Audio::Insights
 		};
 
 		// Process messages
-#if !WITH_EDITOR
 		ProcessMessageQueue<FAudioBusStartMessage>(TraceMessages.StartMessages, CreateEntry,
 		[this](const FAudioBusStartMessage& Msg, TSharedPtr<FAudioBusDashboardEntry>* OutEntry)
 		{
+#if !WITH_EDITOR
 			if (SessionCachedMessages.IsValid())
 			{
 				SessionCachedMessages->StartCachedMessages.EmplaceBack(Msg);
 			}
+#endif // !WITH_EDITOR
 
 			if (OutEntry)
 			{
 				FAudioBusDashboardEntry& EntryRef = *OutEntry->Get();
-				EntryRef.Name = Msg.Name;
-				EntryRef.AudioMeterInfo->NumChannels = Msg.NumChannels;
+
+				if (EntryRef.Name.IsEmpty())
+				{
+					EntryRef.Name = Msg.Name;
+				}
+
+				if (EntryRef.AudioMeterInfo->NumChannels == 0)
+				{
+					EntryRef.AudioMeterInfo->NumChannels = Msg.NumChannels;
+				}
+
+				if (EntryRef.AudioBusType == EAudioBusType::None)
+				{
+					EntryRef.AudioBusType = Msg.AudioBusType;
+				}
 			}
 		});
-#endif // !WITH_EDITOR
 
 		ProcessMessageQueue<FAudioBusHasActivityMessage>(TraceMessages.HasActivityMessages, GetEntry,
 		[this](const FAudioBusHasActivityMessage& Msg, TSharedPtr<FAudioBusDashboardEntry>* OutEntry)
@@ -225,21 +236,28 @@ namespace UE::Audio::Insights
 			}
 		});
 
-#if !WITH_EDITOR
 		ProcessMessageQueue<FAudioBusStopMessage>(TraceMessages.StopMessages, GetEntry,
 		[this](const FAudioBusStopMessage& Msg, TSharedPtr<FAudioBusDashboardEntry>* OutEntry)
 		{
+#if !WITH_EDITOR
 			if (SessionCachedMessages.IsValid())
 			{
 				SessionCachedMessages->StopCachedMessages.EmplaceBack(Msg);
 			}
+#endif // !WITH_EDITOR
 
 			if (OutEntry && (*OutEntry)->Timestamp < Msg.Timestamp)
 			{
-				RemoveDeviceEntry(Msg.DeviceId, Msg.AudioBusId);
+#if WITH_EDITOR
+				const FAudioBusDashboardEntry& EntryRef = *OutEntry->Get();
+				if (EntryRef.AudioBusType == EAudioBusType::CodeGenerated)
+#endif  // WITH_EDITOR
+				{
+					OnAudioBusRemoved.ExecuteIfBound(Msg.AudioBusId);
+					RemoveDeviceEntry(Msg.DeviceId, Msg.AudioBusId);
+				}
 			}
 		});
-#endif // !WITH_EDITOR
 
 		return true;
 	}
@@ -268,7 +286,7 @@ namespace UE::Audio::Insights
 						});
 					}
 
-					OnAudioBusAssetAdded.ExecuteIfBound();
+					OnAudioBusAdded.ExecuteIfBound(InAudioBusDashboardEntry->AudioBusId);
 				}
 			}
 		}
@@ -286,7 +304,7 @@ namespace UE::Audio::Insights
 				{
 					RemoveDeviceEntry(AudioDevice->DeviceID, InAudioBusId);
 
-					OnAudioBusAssetRemoved.ExecuteIfBound();
+					OnAudioBusRemoved.ExecuteIfBound(InAudioBusId);
 				}
 			}
 		}
@@ -327,16 +345,18 @@ namespace UE::Audio::Insights
 			}
 		}
 
-		OnAudioBusAssetListUpdated.ExecuteIfBound();
+		OnAudioBusListUpdated.ExecuteIfBound();
 	}
 #else
 	void FAudioBusTraceProvider::InitSessionCachedMessages(TraceServices::IAnalysisSession& InSession)
 	{
 		SessionCachedMessages = MakeUnique<FAudioBusSessionCachedMessages>(InSession);
 	}
+#endif // WITH_EDITOR
 
 	void FAudioBusTraceProvider::OnTimingViewTimeMarkerChanged(double TimeMarker)
 	{
+#if !WITH_EDITOR
 		using namespace FAudioBusTraceProviderPrivate;
 
 		if (!SessionCachedMessages.IsValid())
@@ -366,6 +386,7 @@ namespace UE::Audio::Insights
 
 				Entry->Name = *AudioBusStartCachedMessage.Name;
 				Entry->AudioMeterInfo->NumChannels = AudioBusStartCachedMessage.NumChannels;
+				Entry->AudioBusType = AudioBusStartCachedMessage.AudioBusType;
 			});
 		}
 
@@ -436,10 +457,10 @@ namespace UE::Audio::Insights
 				}
 			}
 		}
+#endif // !WITH_EDITOR
 
 		// Call parent method to update LastMessageId
 		FTraceProviderBase::OnTimingViewTimeMarkerChanged(TimeMarker);
 	}
-#endif // WITH_EDITOR
 
 } // namespace UE::Audio::Insights
diff --git a/Engine/Plugins/AudioInsights/Source/AudioInsights/Private/Providers/AudioBusTraceProvider.h b/Engine/Plugins/AudioInsights/Source/AudioInsights/Private/Providers/AudioBusTraceProvider.h
--- a/Engine/Plugins/AudioInsights/Source/AudioInsights/Private/Providers/AudioBusTraceProvider.h
+++ b/Engine/Plugins/AudioInsights/Source/AudioInsights/Private/Providers/AudioBusTraceProvider.h
@@ -25,23 +25,25 @@ namespace UE::Audio::Insights
 
 #if !WITH_EDITOR
 		virtual void InitSessionCachedMessages(TraceServices::IAnalysisSession& InSession) override;
-#else
-		DECLARE_DELEGATE(FOnAudioBusAssetsUpdated);
-		FOnAudioBusAssetsUpdated OnAudioBusAssetAdded;
-		FOnAudioBusAssetsUpdated OnAudioBusAssetRemoved;
-		FOnAudioBusAssetsUpdated OnAudioBusAssetListUpdated;
 #endif // !WITH_EDITOR
 
+		DECLARE_DELEGATE_OneParam(FOnAudioBusChanged, const uint32 /*AudioBusId*/);
+		FOnAudioBusChanged OnAudioBusAdded;
+		FOnAudioBusChanged OnAudioBusRemoved;
+
+		DECLARE_DELEGATE(FOnAudioBusListUpdated);
+		FOnAudioBusListUpdated OnAudioBusListUpdated;
+
 	private:
 		virtual bool ProcessMessages() override;
 
+		virtual void OnTimingViewTimeMarkerChanged(double TimeMarker) override;
+
 #if WITH_EDITOR
 		void HandleOnAudioBusAssetAdded(const TSharedPtr<FAudioBusDashboardEntry>& InAudioBusDashboardEntry);
 		void HandleOnAudioBusAssetRemoved(const uint32 InAudioBusId);
 		void HandleOnAudioBusAssetListUpdated(const TArray<TSharedPtr<FAudioBusDashboardEntry>>& InAudioBusAssetEntries);
 #else
-		virtual void OnTimingViewTimeMarkerChanged(double TimeMarker) override;
-
 		TUniquePtr<FAudioBusSessionCachedMessages> SessionCachedMessages;
 #endif // WITH_EDITOR
 
diff --git a/Engine/Plugins/AudioInsights/Source/AudioInsights/Private/Views/AudioBusDashboardViewFactory.cpp b/Engine/Plugins/AudioInsights/Source/AudioInsights/Private/Views/AudioBusDashboardViewFactory.cpp
--- a/Engine/Plugins/AudioInsights/Source/AudioInsights/Private/Views/AudioBusDashboardViewFactory.cpp
+++ b/Engine/Plugins/AudioInsights/Source/AudioInsights/Private/Views/AudioBusDashboardViewFactory.cpp
@@ -7,6 +7,7 @@
 #include "Providers/AudioMeterProvider.h"
 #include "Providers/AudioBusTraceProvider.h"
 #include "Widgets/Input/SCheckBox.h"
+#include "Widgets/Input/SComboBox.h"
 
 #if WITH_EDITOR
 #include "Editor.h"
@@ -24,7 +25,7 @@ namespace UE::Audio::Insights
 			return static_cast<const FAudioBusDashboardEntry&>(InData);
 		};
 
-		void SendAudioBusEnvelopeFollowerCVar(const bool bInIsAudioBusChecked, const FString& InAudioBusName)
+		void SendAudioBusEnvelopeFollowerCVar(const bool bInIsAudioBusChecked, const uint32 InAudioBusId)
 		{
 #if WITH_EDITOR
 			const IAudioInsightsModule& AudioInsightsModule = IAudioInsightsModule::GetEditorChecked();
@@ -34,7 +35,7 @@ namespace UE::Audio::Insights
 
 			const ::Audio::FDeviceId AudioDeviceId = AudioInsightsModule.GetDeviceId();
 
-			const FString AudioBusCommandStr = FString::Printf(TEXT("au.%sAudioBusEnvelopeFollower_AD%d %s"), bInIsAudioBusChecked ? TEXT("Start") : TEXT("Stop"), AudioDeviceId, *InAudioBusName);
+			const FString AudioBusCommandStr = FString::Printf(TEXT("au.%sAudioBusEnvelopeFollower_AD%d %u"), bInIsAudioBusChecked ? TEXT("Start") : TEXT("Stop"), AudioDeviceId, InAudioBusId);
 			
 			IAudioInsightsTraceModule& AudioInsightsTraceModule = IAudioInsightsModule::GetChecked().GetTraceModule();
 			AudioInsightsTraceModule.ExecuteConsoleCommand(AudioBusCommandStr);
@@ -57,20 +58,16 @@ namespace UE::Audio::Insights
 		SortByColumn = "Name";
 		SortMode     = EColumnSortMode::Ascending;
 
-#if WITH_EDITOR
-		AudioBusProvider->OnAudioBusAssetAdded.BindRaw(this, &FAudioBusDashboardViewFactory::RequestListRefresh);
-		AudioBusProvider->OnAudioBusAssetRemoved.BindRaw(this, &FAudioBusDashboardViewFactory::RequestListRefresh);
-		AudioBusProvider->OnAudioBusAssetListUpdated.BindRaw(this, &FAudioBusDashboardViewFactory::RequestListRefresh);
-#endif // WITH_EDITOR
+		AudioBusProvider->OnAudioBusAdded.BindRaw(this, &FAudioBusDashboardViewFactory::HandleOnAudioBusAdded);
+		AudioBusProvider->OnAudioBusRemoved.BindRaw(this, &FAudioBusDashboardViewFactory::HandleOnAudioBusRemoved);
+		AudioBusProvider->OnAudioBusListUpdated.BindRaw(this, &FAudioBusDashboardViewFactory::RequestListRefresh);
 	}
 
 	FAudioBusDashboardViewFactory::~FAudioBusDashboardViewFactory()
 	{
-#if WITH_EDITOR
-		AudioBusProvider->OnAudioBusAssetAdded.Unbind();
-		AudioBusProvider->OnAudioBusAssetRemoved.Unbind();
-		AudioBusProvider->OnAudioBusAssetListUpdated.Unbind();
-#endif // WITH_EDITOR
+		AudioBusProvider->OnAudioBusAdded.Unbind();
+		AudioBusProvider->OnAudioBusRemoved.Unbind();
+		AudioBusProvider->OnAudioBusListUpdated.Unbind();
 	}
 
 	FName FAudioBusDashboardViewFactory::GetName() const
@@ -134,7 +131,7 @@ namespace UE::Audio::Insights
 
 			if (bIsAudioBusChecked)
 			{
-				SendAudioBusEnvelopeFollowerCVar(bIsAudioBusChecked, AudioBusDashboardEntry.GetDisplayName().ToString());
+				SendAudioBusEnvelopeFollowerCVar(bIsAudioBusChecked, AudioBusDashboardEntry.AudioBusId);
 				FAudioMeterProvider::OnUpdateAudioMeterInfo.ExecuteIfBound(AudioBusId, AudioBusDashboardEntry.AudioMeterInfo, AudioBusDashboardEntry.GetDisplayName());
 			}
 
@@ -162,7 +159,7 @@ namespace UE::Audio::Insights
 							*FoundAudioBusCheckedState = bIsChecked;
 						}
 
-						SendAudioBusEnvelopeFollowerCVar(bIsChecked, AudioBusDashboardEntry.GetDisplayName().ToString());
+						SendAudioBusEnvelopeFollowerCVar(bIsChecked, AudioBusDashboardEntry.AudioBusId);
 
 						if (bIsChecked)
 						{
@@ -185,6 +182,7 @@ namespace UE::Audio::Insights
 				[
 					SNew(STextBlock)
 					.Text(ValueText)
+					.ColorAndOpacity(AudioBusDashboardEntry.AudioBusType == EAudioBusType::CodeGenerated ? FSlateColor(FColor(80, 200, 255)) : FSlateColor::UseForeground())
 					.OnDoubleClicked_Lambda([this, InRowData](const FGeometry& MyGeometry, const FPointerEvent& PointerEvent)
 					{
 #if WITH_EDITOR
@@ -213,14 +211,72 @@ namespace UE::Audio::Insights
 
 	void FAudioBusDashboardViewFactory::ProcessEntries(FTraceTableDashboardViewFactory::EProcessReason Reason)
 	{
-		const FString FilterString = GetSearchFilterText().ToString();
-		
-		FTraceTableDashboardViewFactory::FilterEntries<FAudioBusTraceProvider>([&FilterString](const IDashboardDataViewEntry& Entry)
+		FilterByAudioBusName();
+		FilterByAudioBusType();
+	}
+
+	TSharedRef<SWidget> FAudioBusDashboardViewFactory::MakeAudioBusTypeFilterWidget()
+	{
+		if (AudioBusTypes.IsEmpty())
 		{
-			const FAudioBusDashboardEntry& AudioBusEntry = static_cast<const FAudioBusDashboardEntry&>(Entry);
-			
-			return AudioBusEntry.GetDisplayName().ToString().Contains(FilterString);
-		});
+			AudioBusTypes.Emplace(MakeShared<FComboboxSelectionItem>(EAudioBusTypeComboboxSelection::AssetBased,    LOCTEXT("AudioBusDashboard_AudioBusTypeAssetBased",    "Asset")));
+			AudioBusTypes.Emplace(MakeShared<FComboboxSelectionItem>(EAudioBusTypeComboboxSelection::CodeGenerated, LOCTEXT("AudioBusDashboard_AudioBusTypeCodeGenerated", "Code Generated")));
+			AudioBusTypes.Emplace(MakeShared<FComboboxSelectionItem>(EAudioBusTypeComboboxSelection::All,           LOCTEXT("AudioBusDashboard_AudioBusTypeAll",           "All")));
+
+			SelectedAudioBusType = AudioBusTypes[0];
+		}
+
+		return SNew(SHorizontalBox)
+			+ SHorizontalBox::Slot()
+			.AutoWidth()
+			.Padding(2.0f, 10.0f, 0.0f, 0.0f)
+			[
+				SNew(STextBlock)
+				.Margin(FMargin(0.0, 2.0, 0.0, 0.0))
+				.Text(LOCTEXT("AudioBusesDashboard_TypeFilterText", "Type Filter:"))
+			]
+			+ SHorizontalBox::Slot()
+			.MaxWidth(2.0f)
+			.Padding(0.0f, 10.0f, 0.0f, 0.0f)
+			[
+				SNew(SBox)
+			]
+			+ SHorizontalBox::Slot()
+			.AutoWidth()
+			.HAlign(EHorizontalAlignment::HAlign_Center)
+			.Padding(0.0f, 10.0f, 0.0f, 0.0f)
+			[
+				SNew(SComboBox<TSharedPtr<FComboboxSelectionItem>>)
+				.OptionsSource(&AudioBusTypes)
+				.OnGenerateWidget_Lambda([this](const TSharedPtr<FComboboxSelectionItem>& AudioBusTypePtr)
+				{
+					const FText AudioBusTypeDisplayName = AudioBusTypePtr.IsValid() ? AudioBusTypePtr->Value /*DisplayName*/ : FText::GetEmpty();
+
+					return SNew(STextBlock)
+						.Text(AudioBusTypeDisplayName);
+				})
+				.OnSelectionChanged_Lambda([this](TSharedPtr<FComboboxSelectionItem> InSelectedAudioBusTypePtr, ESelectInfo::Type)
+				{
+					if (InSelectedAudioBusTypePtr.IsValid())
+					{
+						SelectedAudioBusType = InSelectedAudioBusTypePtr;
+						UpdateFilterReason = EProcessReason::FilterUpdated;
+					}
+				})
+				[
+					SNew(STextBlock)
+					.Text_Lambda([this]()
+					{
+						const int32 FoundIndex = AudioBusTypes.Find(SelectedAudioBusType);
+						if (AudioBusTypes.IsValidIndex(FoundIndex) && AudioBusTypes[FoundIndex].IsValid())
+						{
+							return AudioBusTypes[FoundIndex]->Value;
+						}
+
+						return FText::GetEmpty();
+					})
+				]
+			];
 	}
 
 	TSharedRef<SWidget> FAudioBusDashboardViewFactory::MakeWidget(TSharedRef<SDockTab> OwnerTab, const FSpawnTabArgs& SpawnTabArgs)
@@ -229,7 +285,19 @@ namespace UE::Audio::Insights
 
 		if (!DashboardWidget.IsValid())
 		{
-			DashboardWidget = FTraceTableDashboardViewFactory::MakeWidget(OwnerTab, SpawnTabArgs);
+			SAssignNew(DashboardWidget, SVerticalBox)
+			+ SVerticalBox::Slot()
+			.AutoHeight()
+			.HAlign(HAlign_Fill)
+			.Padding(0.0f, 0.0f, 0.0f, 6.0f)
+			[
+				MakeAudioBusTypeFilterWidget()
+			]
+			+ SVerticalBox::Slot()
+			.HAlign(HAlign_Fill)
+			[
+				FTraceTableDashboardViewFactory::MakeWidget(OwnerTab, SpawnTabArgs)
+			];
 
 			if (FilteredEntriesListView.IsValid())
 			{
@@ -267,7 +335,7 @@ namespace UE::Audio::Insights
 					{
 						LOCTEXT("AudioBusDashboard_ActiveDisplayName", "Active"),
 						[](const IDashboardDataViewEntry& InData) { return FText::GetEmpty(); },
-						nullptr		/* GetIconName */,
+						nullptr,	/* GetIconName */
 						false,		/* bDefaultHidden */
 						0.08f,		/* FillWidth */
 						EHorizontalAlignment::HAlign_Center
@@ -278,7 +346,7 @@ namespace UE::Audio::Insights
 					{
 						LOCTEXT("AudioBusDashboard_NameColumnDisplayName", "Name"),
 						[](const IDashboardDataViewEntry& InData) { return CastEntry(InData).GetDisplayName(); },
-						nullptr		/* GetIconName */,
+						nullptr,	/* GetIconName */
 						false,		/* bDefaultHidden */
 						0.92f,		/* FillWidth */
 						EHorizontalAlignment::HAlign_Left
@@ -344,6 +412,61 @@ namespace UE::Audio::Insights
 		}
 	}
 
+	void FAudioBusDashboardViewFactory::FilterByAudioBusName()
+	{
+		const FString FilterString = GetSearchFilterText().ToString();
+		
+		FTraceTableDashboardViewFactory::FilterEntries<FAudioBusTraceProvider>([&FilterString](const IDashboardDataViewEntry& Entry)
+		{
+			const FAudioBusDashboardEntry& AudioBusEntry = static_cast<const FAudioBusDashboardEntry&>(Entry);
+			
+			return AudioBusEntry.GetDisplayName().ToString().Contains(FilterString);
+		});
+	}
+
+	void FAudioBusDashboardViewFactory::FilterByAudioBusType()
+	{
+		using namespace UE::Audio::Insights;
+		using namespace FAudioBusDashboardViewFactoryPrivate;
+
+		TArray<TSharedPtr<IDashboardDataViewEntry>> EntriesToFilterOut;
+
+		const EAudioBusTypeComboboxSelection SelectedAudioBusTypeEnum = SelectedAudioBusType.IsValid()? SelectedAudioBusType->Key : EAudioBusTypeComboboxSelection::All;
+
+		for (const TSharedPtr<IDashboardDataViewEntry>& Entry : DataViewEntries)
+		{
+			if (Entry.IsValid())
+			{
+				const FAudioBusDashboardEntry& AudioBusAssetDashboardEntry = CastEntry(*Entry);
+
+				if (SelectedAudioBusTypeEnum != EAudioBusTypeComboboxSelection::All)
+				{
+					if ((SelectedAudioBusTypeEnum == EAudioBusTypeComboboxSelection::AssetBased    && AudioBusAssetDashboardEntry.AudioBusType != EAudioBusType::AssetBased) ||
+						(SelectedAudioBusTypeEnum == EAudioBusTypeComboboxSelection::CodeGenerated && AudioBusAssetDashboardEntry.AudioBusType != EAudioBusType::CodeGenerated))
+					{
+						EntriesToFilterOut.Emplace(Entry);
+					}
+				}
+			}
+		}
+
+		for (const TSharedPtr<IDashboardDataViewEntry>& Entry : EntriesToFilterOut)
+		{
+			DataViewEntries.Remove(Entry);
+		}
+	}
+
+	void FAudioBusDashboardViewFactory::HandleOnAudioBusAdded(const uint32 InAudioBusId)
+	{
+		RequestListRefresh();
+	}
+
+	void FAudioBusDashboardViewFactory::HandleOnAudioBusRemoved(const uint32 InAudioBusId)
+	{
+		FAudioMeterProvider::OnRemoveAudioMeter.ExecuteIfBound(InAudioBusId);
+		RequestListRefresh();
+	}
+
 	void FAudioBusDashboardViewFactory::RequestListRefresh()
 	{
 		if (FilteredEntriesListView.IsValid())
diff --git a/Engine/Plugins/AudioInsights/Source/AudioInsights/Public/Messages/AudioBusTraceMessages.h b/Engine/Plugins/AudioInsights/Source/AudioInsights/Public/Messages/AudioBusTraceMessages.h
--- a/Engine/Plugins/AudioInsights/Source/AudioInsights/Public/Messages/AudioBusTraceMessages.h
+++ b/Engine/Plugins/AudioInsights/Source/AudioInsights/Public/Messages/AudioBusTraceMessages.h
@@ -13,6 +13,13 @@
 
 namespace UE::Audio::Insights
 {
+	enum class EAudioBusType : uint8
+	{
+		AssetBased,
+		CodeGenerated,
+		None
+	};
+
 	struct FAudioBusMessageBase
 	{
 		FAudioBusMessageBase() = default;
@@ -42,10 +49,13 @@ namespace UE::Audio::Insights
 
 			EventData.GetString("Name", Name);
 			NumChannels = EventData.GetValue<int32>("NumChannels");
+
+			AudioBusType = FSoftObjectPath(Name).IsAsset() ? EAudioBusType::AssetBased : EAudioBusType::CodeGenerated;
 		}
 
 		FString Name;
 		int32 NumChannels = 0;
+		EAudioBusType AudioBusType = EAudioBusType::None;
 	};
 
 	struct FAudioBusHasActivityMessage : public FAudioBusMessageBase
@@ -122,7 +132,7 @@ namespace UE::Audio::Insights
 
 		virtual ~FAudioBusDashboardEntry() = default;
 
-		virtual FText GetDisplayName() const override { return FText::FromString(FSoftObjectPath(Name).GetAssetName()); }
+		virtual FText GetDisplayName() const override { return FText::FromString(AudioBusType == EAudioBusType::AssetBased ? FSoftObjectPath(Name).GetAssetName() : Name); }
 		virtual const UObject* GetObject() const override { return FSoftObjectPath(Name).ResolveObject(); }
 		virtual UObject* GetObject() override { return FSoftObjectPath(Name).ResolveObject(); }
 		virtual bool IsValid() const override { return AudioBusId != static_cast<uint32>(INDEX_NONE); }
@@ -134,5 +144,7 @@ namespace UE::Audio::Insights
 
 		FString Name;
 		bool bHasActivity = false;
+
+		EAudioBusType AudioBusType = EAudioBusType::None;
 	};
 } // namespace UE::Audio::Insights
diff --git a/Engine/Plugins/AudioInsights/Source/AudioInsights/Public/Views/AudioBusDashboardViewFactory.h b/Engine/Plugins/AudioInsights/Source/AudioInsights/Public/Views/AudioBusDashboardViewFactory.h
--- a/Engine/Plugins/AudioInsights/Source/AudioInsights/Public/Views/AudioBusDashboardViewFactory.h
+++ b/Engine/Plugins/AudioInsights/Source/AudioInsights/Public/Views/AudioBusDashboardViewFactory.h
@@ -23,17 +23,35 @@ namespace UE::Audio::Insights
 		UE_API virtual TSharedRef<SWidget> MakeWidget(TSharedRef<SDockTab> OwnerTab, const FSpawnTabArgs& SpawnTabArgs) override;
 
 	protected:
+		UE_API enum class EAudioBusTypeComboboxSelection : uint8
+		{
+			AssetBased,
+			CodeGenerated,
+			All
+		};
+
+		UE_API TSharedRef<SWidget> MakeAudioBusTypeFilterWidget();
+
 		UE_API virtual TSharedRef<SWidget> GenerateWidgetForColumn(TSharedRef<IDashboardDataViewEntry> InRowData, const FName& InColumnName) override;
 		UE_API virtual void ProcessEntries(FTraceTableDashboardViewFactory::EProcessReason Reason) override;
 		UE_API virtual const TMap<FName, FTraceTableDashboardViewFactory::FColumnData>& GetColumns() const override;
 
 		UE_API virtual void SortTable() override;
 
+		UE_API void FilterByAudioBusName();
+		UE_API void FilterByAudioBusType();
+
+		UE_API void HandleOnAudioBusAdded(const uint32 InAudioBusId);
+		UE_API void HandleOnAudioBusRemoved(const uint32 InAudioBusId);
 		UE_API void RequestListRefresh();
 
 		TSharedPtr<FAudioBusTraceProvider> AudioBusProvider;
 
 		TMap<uint32, bool> AudioBusCheckboxCheckedStates;
+
+		using FComboboxSelectionItem = TPair<EAudioBusTypeComboboxSelection, FText>;
+		TArray<TSharedPtr<FComboboxSelectionItem>> AudioBusTypes;
+		TSharedPtr<FComboboxSelectionItem> SelectedAudioBusType;
 	};
 } // namespace UE::Audio::Insights
 
