[Audio Insights] Fix for build machine error

#rnx

diff --git a/Engine/Plugins/AudioInsights/Source/AudioInsights/Public/Views/AudioBusDashboardViewFactory.h b/Engine/Plugins/AudioInsights/Source/AudioInsights/Public/Views/AudioBusDashboardViewFactory.h
--- a/Engine/Plugins/AudioInsights/Source/AudioInsights/Public/Views/AudioBusDashboardViewFactory.h
+++ b/Engine/Plugins/AudioInsights/Source/AudioInsights/Public/Views/AudioBusDashboardViewFactory.h
@@ -23,7 +23,7 @@ namespace UE::Audio::Insights
 		UE_API virtual TSharedRef<SWidget> MakeWidget(TSharedRef<SDockTab> OwnerTab, const FSpawnTabArgs& SpawnTabArgs) override;
 
 	protected:
-		UE_API enum class EAudioBusTypeComboboxSelection : uint8
+		enum class UE_API EAudioBusTypeComboboxSelection : uint8
 		{
 			AssetBased,
 			CodeGenerated,
