[Audio Insights] updated tooltips and comments per Bughawk suggestions

diff --git a/Engine/Plugins/AudioInsights/Source/AudioInsights/Private/Providers/AudioEventLogTraceProvider.cpp b/Engine/Plugins/AudioInsights/Source/AudioInsights/Private/Providers/AudioEventLogTraceProvider.cpp
--- a/Engine/Plugins/AudioInsights/Source/AudioInsights/Private/Providers/AudioEventLogTraceProvider.cpp
+++ b/Engine/Plugins/AudioInsights/Source/AudioInsights/Private/Providers/AudioEventLogTraceProvider.cpp
@@ -136,8 +136,8 @@ namespace UE::Audio::Insights
 		Entry->EventName = Msg.EventName;
 		Entry->Name = Msg.AssetPath;
 
-		// If this message did not provide an actor name, the object may have been destroying when the message was sent
-		// Try to locate the actor name using the play order ID
+		// If this message did not provide an actor label, the object may have been destroying when the message was sent
+		// Try to locate the actor label using the play order ID
 		if (Msg.ActorLabel.IsEmpty())
 		{
 			if (const FActorInformation* FoundActorInfo = PlayOrderToActorInfoMap.Find(Entry->PlayOrder))
diff --git a/Engine/Plugins/AudioInsights/Source/AudioInsights/Private/Views/SoundDashboardViewFactory.cpp b/Engine/Plugins/AudioInsights/Source/AudioInsights/Private/Views/SoundDashboardViewFactory.cpp
--- a/Engine/Plugins/AudioInsights/Source/AudioInsights/Private/Views/SoundDashboardViewFactory.cpp
+++ b/Engine/Plugins/AudioInsights/Source/AudioInsights/Private/Views/SoundDashboardViewFactory.cpp
@@ -576,7 +576,7 @@ namespace UE::Audio::Insights
 
 		const FName ActorLabelColumnName = "ActorLabel";
 		const FText ActorLabelColumnDisplayName = LOCTEXT("SoundDashboard_ActorLabelColumnDisplayName", "Actor Label");
-		const FText ActorLabelColumnTooltip = LOCTEXT("SoundDashboard_ActorLabelColumnTooltip", "The editor label of the Actor this sound is playing on (if applicable).");
+		const FText ActorLabelColumnTooltip = LOCTEXT("SoundDashboard_ActorLabelColumnTooltip", "The label or name of the Actor this sound is playing on (if applicable).");
 
 		const FName CategoryColumnName = "Category";
 		const FText CategoryColumnDisplayName = LOCTEXT("SoundDashboard_CategoryColumnDisplayName", "Category");
