[Audio Insights] Custom event type filters for the Event Log can now be created in the Editor Preferences

diff --git a/Engine/Plugins/AudioInsights/Source/AudioInsights/Private/Views/AudioEventLogDashboardViewFactory.cpp b/Engine/Plugins/AudioInsights/Source/AudioInsights/Private/Views/AudioEventLogDashboardViewFactory.cpp
--- a/Engine/Plugins/AudioInsights/Source/AudioInsights/Private/Views/AudioEventLogDashboardViewFactory.cpp
+++ b/Engine/Plugins/AudioInsights/Source/AudioInsights/Private/Views/AudioEventLogDashboardViewFactory.cpp
@@ -28,81 +28,85 @@ namespace UE::Audio::Insights
 		const FString CategoryMessages = LOCTEXT("EventLogCategory_Messages", "Messages").ToString();
 		const FString CategoryCustom = LOCTEXT("EventLogCategory_Custom", "Custom").ToString();
 
-		const TMap<FString, TArray<FString>> GetInitEventTypeFilters()
+
+		const FName MessageIDColumnID = "MessageID";
+		const FName TimestampColumnID = "Timestamp";
+		const FName PlayOrderColumnID = "PlayOrder";
+		const FName EventColumnID = "Event";
+		const FName AssetNameColumnID = "Asset";
+		const FName ActorColumnID = "Actor";
+		const FName CategoryColumnID = "Category";
+	} // namespace AudioEventLogPrivate
+
+	const TMap<FString, TSet<FString>> FAudioEventLogDashboardViewFactory::GetInitEventTypeFilters()
+	{
+		using namespace AudioEventLogPrivate;
+		using namespace ::Audio::Trace::EventLog;
+
+		return
 		{
-			using namespace ::Audio::Trace::EventLog;
-			return 
 			{
-				{ 
-					CategorySoundActivity,
-					{
-						ID::SoundStart,
-						ID::SoundStop,
-						ID::PauseSoundRequested,
-						ID::ResumeSoundRequested
-					}
-				},
-				
+				CategorySoundActivity,
 				{
-					CategoryVirtualization,
-					{
-						ID::SoundVirtualized,
-						ID::SoundRealized
-					}
-				},
-				
+					ID::SoundStart,
+					ID::SoundStop,
+					ID::PauseSoundRequested,
+					ID::ResumeSoundRequested
+				}
+			},
+
+			{
+				CategoryVirtualization,
 				{
-					CategoryPlayRequests,
-					{
-						ID::PlayRequestSoundHandle,
-						ID::PlayRequestAudioComponent,
-						ID::PlayRequestOneShot,
-						ID::PlayRequestSoundAtLocation,
-						ID::PlayRequestSound2D,
-						ID::PlayRequestSlateSound
-					}
-				},
+					ID::SoundVirtualized,
+					ID::SoundRealized
+				}
+			},
 
+			{
+				CategoryPlayRequests,
 				{
-					CategoryStopRequests,
-					{
-						ID::StopAllRequested,
-						ID::StopRequestedSoundHandle,
-						ID::StopRequestAudioComponent,
-						ID::StopRequestActiveSound,
-						ID::StopRequestSoundsUsingResource,
-						ID::StopRequestConcurrency
-					}
-				},
+					ID::PlayRequestSoundHandle,
+					ID::PlayRequestAudioComponent,
+					ID::PlayRequestOneShot,
+					ID::PlayRequestSoundAtLocation,
+					ID::PlayRequestSound2D,
+					ID::PlayRequestSlateSound
+				}
+			},
 
+			{
+				CategoryStopRequests,
 				{
-					CategoryMessages,
-					{
-						ID::FlushAudioDeviceRequested
-					}
-				},
+					ID::StopAllRequested,
+					ID::StopRequestedSoundHandle,
+					ID::StopRequestAudioComponent,
+					ID::StopRequestActiveSound,
+					ID::StopRequestSoundsUsingResource,
+					ID::StopRequestConcurrency
+				}
+			},
 
+			{
+				CategoryMessages,
 				{
-					CategoryCustom,
-					{
-						// Auto-populates with custom events from users which haven't been categorized 
-					}
+					ID::FlushAudioDeviceRequested
 				}
-			};
-		}
+			},
 
-		const FName MessageIDColumnID = "MessageID";
-		const FName TimestampColumnID = "Timestamp";
-		const FName PlayOrderColumnID = "PlayOrder";
-		const FName EventColumnID = "Event";
-		const FName AssetNameColumnID = "Asset";
-		const FName ActorColumnID = "Actor";
-		const FName CategoryColumnID = "Category";
-	} // namespace AudioEventLogPrivate
+			{
+				CategoryCustom,
+				{
+					// Auto-populates with custom events from users which haven't been categorized 
+				}
+			}
+		};
+	}
 
 	/////////////////////////////////////////////////////////////////////////////////////////
 	// FAudioEventLogDashboardViewFactory
 	FAudioEventLogDashboardViewFactory::FAudioEventLogDashboardViewFactory()
+		: InitEventTypeFilters(GetInitEventTypeFilters())
 	{
 		using namespace AudioEventLogPrivate;
 
@@ -439,8 +443,6 @@ namespace UE::Audio::Insights
 
 		if (!FilterBar.IsValid())
 		{
-			const TMap<FString, TArray<FString>> InitEventTypeFilters = GetInitEventTypeFilters();
-
 			TArray<TSharedRef<FFilterBase<FEventLogFilterID>>> Filters;
 
 			for (const auto& [Category, EventTypes] : InitEventTypeFilters)
@@ -647,18 +649,12 @@ namespace UE::Audio::Insights
 	{
 		if (FilterBar.IsValid())
 		{
-			if (InSettings.EventFilters.IsEmpty())
-			{
-				FilterBar->RemoveAllFilters();
-			}
-			else
-			{
-				for (const auto& [EventName, FilterID] : EventFilterTypes)
-				{
-					const bool bFilterIsActive = InSettings.EventFilters.Contains(EventName);
-					FilterBar->SetFilterCheckState(FilterBar->GetFilter(EventName), bFilterIsActive ? ECheckBoxState::Checked : ECheckBoxState::Unchecked);
-				}
-			}
+			RemoveDeletedCustomEvents(InSettings.CustomCategoriesToEvents);
+			AddNewCustomEvents(InSettings.CustomCategoriesToEvents);
+
+			CachedCustomEventSettings = InSettings.CustomCategoriesToEvents;
+			
+			RefreshFilterBarFromSettings(InSettings.EventFilters);
 		}
 
 		if (VisibleColumnsSettingsMenu.IsValid())
@@ -698,6 +694,101 @@ namespace UE::Audio::Insights
 			VisibleColumnsSettingsMenu->WriteToSettings(OutSettings.VisibleColumns);
 		}
 	}
+
+	void FAudioEventLogDashboardViewFactory::RemoveDeletedCustomEvents(const TMap<FString, FAudioEventLogCustomEvents>& CustomEventsFromSettings)
+	{
+		if (!FilterBar.IsValid())
+		{
+			return;
+		}
+
+		// Run through our cached custom events, remove any that are not in the settings
+		for (const auto& [CachedCategory, CachedEventTypes] : CachedCustomEventSettings)
+		{
+			// Can be null if the category has been deleted
+			const FAudioEventLogCustomEvents* FoundEventTypesFromSettingsCategory = CustomEventsFromSettings.Find(CachedCategory);
+
+			int32 NumRemoved = 0;
+			for (const FString& CachedEvent : CachedEventTypes.EventNames)
+			{
+				if (FoundEventTypesFromSettingsCategory == nullptr || !FoundEventTypesFromSettingsCategory->EventNames.Contains(CachedEvent))
+				{
+					// The category or event has been destroyed, remove it here
+					EventFilterTypes.Remove(CachedEvent);
+					NumRemoved++;
+
+					if (const TSharedPtr<FFilterBase<FEventLogFilterID>> Filter = FilterBar->GetFilter(CachedEvent))
+					{
+						FilterBar->DeleteFromFilter(Filter.ToSharedRef());
+					}
+				}
+			}
+
+			// If every event in this category has been deleted, remove the category too
+			if (NumRemoved == CachedEventTypes.EventNames.Num())
+			{
+				if (const TSharedPtr<FFilterCategory>* FilterCategory = FilterCategories.Find(CachedCategory))
+				{
+					FilterBar->DeleteCategory(FilterCategory->ToSharedRef());
+					FilterCategories.Remove(CachedCategory);
+				}
+			}
+		}
+	}
+
+	void FAudioEventLogDashboardViewFactory::AddNewCustomEvents(const TMap<FString, FAudioEventLogCustomEvents>& CustomEventsFromSettings)
+	{
+		if (!FilterBar.IsValid())
+		{
+			return;
+		}
+
+		// Run through all of the events in the settings - detect any new ones and add them to the filter
+		for (const auto& [SettingsCategoryName, SettingsEventTypes] : CustomEventsFromSettings)
+		{
+			if (SettingsEventTypes.EventNames.IsEmpty())
+			{
+				// Do not re-add the category if the event name collection is empty
+				continue;
+			}
+
+			// Can be null if the category does not exist yet
+			const FAudioEventLogCustomEvents* FoundEventTypesFromCachedCategory = CachedCustomEventSettings.Find(SettingsCategoryName);
+
+			// Find the category if it exists, or create a new one if it doesn't
+			const TSharedPtr<FFilterCategory>& FilterCategory = FilterCategories.FindOrAdd(SettingsCategoryName, MakeShared<FFilterCategory>(FText::FromString(SettingsCategoryName), FText::GetEmpty()));
+
+			for (const FString& SettingsEvent : SettingsEventTypes.EventNames)
+			{
+				if (FoundEventTypesFromCachedCategory == nullptr || !FoundEventTypesFromCachedCategory->EventNames.Contains(SettingsEvent))
+				{
+					// This is a new category or event, add the event to the filter
+					FilterBar->AddFilter(CreateNewEventFilterType(FilterCategory, SettingsEvent));
+				}
+			}
+		}
+	}
+
+	void FAudioEventLogDashboardViewFactory::RefreshFilterBarFromSettings(const TSet<FString>& EventFilters)
+	{
+		if (!FilterBar.IsValid())
+		{
+			return;
+		}
+
+		if (EventFilters.IsEmpty())
+		{
+			FilterBar->RemoveAllFilters();
+		}
+		else
+		{
+			for (const auto& [EventName, FilterID] : EventFilterTypes)
+			{
+				const bool bFilterIsActive = EventFilters.Contains(EventName);
+				FilterBar->SetFilterCheckState(FilterBar->GetFilter(EventName), bFilterIsActive ? ECheckBoxState::Checked : ECheckBoxState::Unchecked);
+			}
+		}
+	}
 #endif // WITH_EDITOR
 
 	void FAudioEventLogDashboardViewFactory::OnTimingViewTimeMarkerChanged(double InTimeMarker)
diff --git a/Engine/Plugins/AudioInsights/Source/AudioInsights/Public/Settings/AudioEventLogSettings.h b/Engine/Plugins/AudioInsights/Source/AudioInsights/Public/Settings/AudioEventLogSettings.h
--- a/Engine/Plugins/AudioInsights/Source/AudioInsights/Public/Settings/AudioEventLogSettings.h
+++ b/Engine/Plugins/AudioInsights/Source/AudioInsights/Public/Settings/AudioEventLogSettings.h
@@ -10,6 +10,23 @@
 
 #define LOCTEXT_NAMESPACE "AudioInsights"
 
+/** FAudioEventLogCustomEventType
+ *
+ * Custom event types that can be filtered by in the Event Log
+ */
+USTRUCT()
+struct FAudioEventLogCustomEvents
+{
+	GENERATED_BODY()
+
+	/** 
+	* The names of any custom events associated with the above category.
+	* This should match the event name used to send the event to the Event Log.
+	*/
+	UPROPERTY(EditAnywhere, config, Category = EventLog, meta = (DisplayName = "Custom event names"))
+	TSet<FString> EventNames;
+};
+
 /** FAudioEventLogVisibleColumns
  *
  * Control which columns are visible in the Audio Event Log tab.
@@ -55,6 +72,13 @@ struct FAudioEventLogSettings
 	UPROPERTY(config, meta = (HideInDetailPanel))
 	TSet<FString> EventFilters;
 
+	/** 
+	* Categories for custom events being sent to the Audio Insights Event Log.
+	* These categories and events will appear in the event filter menu.
+	*/
+	UPROPERTY(EditAnywhere, config, Category = EventLog, meta = (DisplayName = "Custom Event Log categories"))
+	TMap<FString, FAudioEventLogCustomEvents> CustomCategoriesToEvents;
+
 	/** Visibility of columns in the Event Log tab */
 	UPROPERTY(EditAnywhere, config, Category = EventLog)
 	FAudioEventLogVisibleColumns VisibleColumns;
diff --git a/Engine/Plugins/AudioInsights/Source/AudioInsights/Public/Views/AudioEventLogDashboardViewFactory.h b/Engine/Plugins/AudioInsights/Source/AudioInsights/Public/Views/AudioEventLogDashboardViewFactory.h
--- a/Engine/Plugins/AudioInsights/Source/AudioInsights/Public/Views/AudioEventLogDashboardViewFactory.h
+++ b/Engine/Plugins/AudioInsights/Source/AudioInsights/Public/Views/AudioEventLogDashboardViewFactory.h
@@ -52,6 +52,8 @@ namespace UE::Audio::Insights
 		UE_API FAudioEventLogDashboardViewFactory();
 		UE_API virtual ~FAudioEventLogDashboardViewFactory();
 
+		UE_API static const TMap<FString, TSet<FString>> GetInitEventTypeFilters();
+
 		UE_API virtual FName GetName() const override;
 		UE_API virtual FText GetDisplayName() const override;
 		UE_API virtual EDefaultDashboardTabStack GetDefaultTabStack() const override;
@@ -98,6 +100,11 @@ namespace UE::Audio::Insights
 #if WITH_EDITOR
 		void OnReadEditorSettings(const FAudioEventLogSettings& InSettings);
 		void OnWriteEditorSettings(FAudioEventLogSettings& OutSettings);
+
+		void RemoveDeletedCustomEvents(const TMap<FString, FAudioEventLogCustomEvents>& CustomEventsFromSettings);
+		void AddNewCustomEvents(const TMap<FString, FAudioEventLogCustomEvents>& CustomEventsFromSettings);
+
+		void RefreshFilterBarFromSettings(const TSet<FString>& EventFilters);
 #endif // WITH_EDITOR
 
 		void OnTimingViewTimeMarkerChanged(double InTimeMarker);
@@ -112,6 +119,8 @@ namespace UE::Audio::Insights
 		TSharedPtr<FVisibleColumnsSettingsMenu<FAudioEventLogVisibleColumns>> VisibleColumnsSettingsMenu;
 
 		void BindCommands();
+		const TMap<FString, TSet<FString>> InitEventTypeFilters;
+		TMap<FString, FAudioEventLogCustomEvents> CachedCustomEventSettings;
 	};
 } // namespace UE::Audio::Insights
 
diff --git a/Engine/Plugins/AudioInsights/Source/AudioInsights/Public/Views/SAudioFilterBar.h b/Engine/Plugins/AudioInsights/Source/AudioInsights/Public/Views/SAudioFilterBar.h
--- a/Engine/Plugins/AudioInsights/Source/AudioInsights/Public/Views/SAudioFilterBar.h
+++ b/Engine/Plugins/AudioInsights/Source/AudioInsights/Public/Views/SAudioFilterBar.h
@@ -34,6 +34,17 @@ namespace UE::Audio::Insights
 			);
 		}
 
+		void DeleteCategory(const TSharedRef<FFilterCategory> InCategory)
+		{
+			SBasicFilterBar<EnumType>::AllFilterCategories.Remove(InCategory);
+		}
+
+		void DeleteFromFilter(const TSharedRef<FFilterBase<EnumType>> InFilter)
+		{
+			SBasicFilterBar<EnumType>::SetFilterCheckState(InFilter, ECheckBoxState::Unchecked);
+			SBasicFilterBar<EnumType>::AllFrontendFilters.Remove(InFilter);
+		}
+
 	protected:
 		virtual TSharedRef<SWidget> MakeAddFilterMenu() override
 		{
diff --git a/Engine/Plugins/AudioInsights/Source/AudioInsightsEditor/Private/AudioInsightsEditorSettings.cpp b/Engine/Plugins/AudioInsights/Source/AudioInsightsEditor/Private/AudioInsightsEditorSettings.cpp
--- a/Engine/Plugins/AudioInsights/Source/AudioInsightsEditor/Private/AudioInsightsEditorSettings.cpp
+++ b/Engine/Plugins/AudioInsights/Source/AudioInsightsEditor/Private/AudioInsightsEditorSettings.cpp
@@ -2,10 +2,41 @@
 
 #include "AudioInsightsEditorSettings.h"
 
+#include "AudioInsightsLog.h"
+#include "Framework/Notifications/NotificationManager.h"
+#include "Views/AudioEventLogDashboardViewFactory.h"
+#include "Widgets/Notifications/SNotificationList.h"
+
 #include UE_INLINE_GENERATED_CPP_BY_NAME(AudioInsightsEditorSettings)
 
 #define LOCTEXT_NAMESPACE "UAudioInsightsEditorSettings"
 
+namespace AudioInsightsEditorSettingsPrivate
+{
+	FString CreateUniqueName(const FString& DefaultName, TFunctionRef<bool(const FString&)> IsUniqueNameCheck)
+	{
+		FString NewName = DefaultName;
+
+		for (int NumTries = 1; NumTries < TNumericLimits<int32>::Max(); ++NumTries)
+		{
+			if (IsUniqueNameCheck(NewName))
+			{
+				return NewName;
+			}
+
+			NewName = DefaultName + TEXT(" ") + FString::FromInt(NumTries);
+		}
+
+		// unable to create new name
+		return FString();
+	}
+}
+
+UAudioInsightsEditorSettings::UAudioInsightsEditorSettings()
+	: InBuiltAudioLogEventTypes(UE::Audio::Insights::FAudioEventLogDashboardViewFactory::GetInitEventTypeFilters())
+{
+}
+
 FName UAudioInsightsEditorSettings::GetCategoryName() const
 {
 	return "Plugins";
@@ -69,7 +100,10 @@ void UAudioInsightsEditorSettings::PostEditChangeChainProperty(FPropertyChangedC
 
 	if (ActiveMemberNode->GetValue()->GetFName() == GET_MEMBER_NAME_CHECKED(UAudioInsightsEditorSettings, EventLogSettings))
 	{
-		FAudioEventLogSettings::OnReadSettings.Broadcast(EventLogSettings);
+		if (VerifyEventLogInput(*ActiveMemberNode, PropertyChangedEvent))
+		{
+			FAudioEventLogSettings::OnReadSettings.Broadcast(EventLogSettings);
+		}
 	}
 
 	if (ActiveMemberNode->GetValue()->GetFName() == GET_MEMBER_NAME_CHECKED(UAudioInsightsEditorSettings, SoundDashboardSettings))
@@ -78,4 +112,307 @@ void UAudioInsightsEditorSettings::PostEditChangeChainProperty(FPropertyChangedC
 	}
 }
 
+bool UAudioInsightsEditorSettings::VerifyEventLogInput(const TDoubleLinkedList<FProperty*>::TDoubleLinkedListNode& ActiveMemberNode, FPropertyChangedChainEvent& PropertyChangedEvent)
+{
+	check(ActiveMemberNode.GetValue() != nullptr);
+
+	const TDoubleLinkedList<FProperty*>::TDoubleLinkedListNode* CustomCategoryNode = ActiveMemberNode.GetNextNode();
+
+	if (CustomCategoryNode == nullptr 
+		|| CustomCategoryNode->GetValue() == nullptr 
+		|| CustomCategoryNode->GetValue()->GetFName() != GET_MEMBER_NAME_CHECKED(FAudioEventLogSettings, CustomCategoriesToEvents))
+	{
+		// This change is not related to custom events - no need to verify anything
+		return true;
+	}
+
+	if ((PropertyChangedEvent.ChangeType & EPropertyChangeType::ArrayRemove) 
+		|| (PropertyChangedEvent.ChangeType & EPropertyChangeType::ArrayClear))
+	{
+		// If the operation is a delete, no need to verify any new event names
+		// We will want to auto-remove any events still cached in the event filter
+		RefreshEventFilter();
+		return true;
+	}
+
+	const TDoubleLinkedList<FProperty*>::TDoubleLinkedListNode* CustomEventsNode = CustomCategoryNode->GetNextNode();
+	if (CustomEventsNode
+		&& CustomEventsNode->GetValue()
+		&& CustomEventsNode->GetValue()->GetFName() == GET_MEMBER_NAME_CHECKED(FAudioEventLogCustomEvents, EventNames))
+	{
+		return VerifyCustomEventNameEdit(PropertyChangedEvent, CastField<FMapProperty>(CustomCategoryNode->GetValue()), CastField<FSetProperty>(CustomEventsNode->GetValue()));
+	}
+	else
+	{
+		return VerifyCategoryEdit();
+	}
+}
+
+void UAudioInsightsEditorSettings::RefreshEventFilter()
+{
+	// If custom events have been removed, we want to remove them from the event filter too
+	TArray<FString> FiltersToRemove;
+
+	for (TSet<FString>::TConstIterator It = EventLogSettings.EventFilters.CreateConstIterator(); It; ++It)
+	{
+		bool bFoundMatch = false;
+
+		// Check custom categories
+		for (const auto& [FilterCategory, CustomEvents] : EventLogSettings.CustomCategoriesToEvents)
+		{
+			if (CustomEvents.EventNames.Contains(*It))
+			{
+				bFoundMatch = true;
+				break;
+			}
+		}
+
+		if (!bFoundMatch)
+		{
+			// Now check the in-built Event Types in Audio Insights
+			for (const auto& [FilterCategory, InBuiltEvents] : InBuiltAudioLogEventTypes)
+			{
+				if (InBuiltEvents.Contains(*It))
+				{
+					bFoundMatch = true;
+					break;
+				}
+			}
+		}
+		
+		if (!bFoundMatch)
+		{
+			// This event filter no longer exists - mark it as being removed
+			FiltersToRemove.Add(*It);
+		}
+	}
+
+	for (const FString& FilterToRemove : FiltersToRemove)
+	{
+		EventLogSettings.EventFilters.Remove(FilterToRemove);
+	}
+}
+
+bool UAudioInsightsEditorSettings::VerifyCategoryEdit()
+{
+	using namespace AudioInsightsEditorSettingsPrivate;
+
+	// TMap will already check to ensure there are no duplicate Keys
+	// If we detect any empty fields, try to create a unique default name to replace them
+	for (auto& [FilterCategory, Events] : EventLogSettings.CustomCategoriesToEvents)
+	{
+		if (FilterCategory.IsEmpty())
+		{
+			// Create a unique event name if the category is empty
+			const FString NewCategoryName = CreateUniqueName(TEXT("New Category"), [this](const FString& NameToCheck)
+			{
+				return !EventLogSettings.CustomCategoriesToEvents.Contains(NameToCheck);
+			});
+
+			// CreateUniqueName() should be able to create a name for the new category - warn and return out if this has failed
+			ensureMsgf(!NewCategoryName.IsEmpty(), TEXT("Failed to create a unique name for edited custom Audio Event Log category"));
+			if (NewCategoryName.IsEmpty())
+			{
+				return false;
+			}
+
+			// There should only be one empty value at a time (i.e. when a new element is added) - edit the map immediately and break;
+			EventLogSettings.CustomCategoriesToEvents.Add(NewCategoryName, MoveTemp(Events));
+			EventLogSettings.CustomCategoriesToEvents.Remove(FilterCategory);
+
+			break;
+		}
+	}
+
+	return true;
+}
+
+bool UAudioInsightsEditorSettings::VerifyCustomEventNameEdit(const FPropertyChangedChainEvent& PropertyChangedEvent, const FMapProperty* CategoryToEventMapProperty, const FSetProperty* EventNamesSetProperty)
+{
+	using namespace AudioInsightsEditorSettingsPrivate;
+
+	// All events need to have unique names - including the in-built event names in the event log
+	// Check both the custom event names and the in built event names for any duplicates, and update the name if found
+	// Also - do not allow empty event names, auto-generate a unique name in this case
+
+	FString* EditedCategoryName = GetEditedCategoryKey(PropertyChangedEvent, CategoryToEventMapProperty);
+	if (EditedCategoryName == nullptr)
+	{
+		return false;
+	}
+
+	FAudioEventLogCustomEvents* Events = EventLogSettings.CustomCategoriesToEvents.Find(*EditedCategoryName);
+	if (Events == nullptr)
+	{
+		UE_LOG(LogAudioInsights, Warning, TEXT("Could not find category with name %s in EventLogSettings.CustomCategoriesToEvents"), **EditedCategoryName);
+		return false;
+	}
+
+	const EditedEvent EditedEvent = GetEditedEvent(PropertyChangedEvent, EventNamesSetProperty, *Events);
+	if (EditedEvent.EventName == nullptr)
+	{
+		return false;
+	}
+
+	// We have found the event - now validate that it is unique and not empty - edit if neccessary
+	if (EditedEvent.EventName->IsEmpty() || !IsEventNameUnique(*EditedEvent.EventName, *EditedCategoryName, EditedEvent.LogicalIndex))
+	{
+		const FString NewEventName = CreateUniqueName(EditedEvent.EventName->IsEmpty() ? TEXT("New Event") : *EditedEvent.EventName,
+		[this, &EditedCategoryName, &EditedEvent](const FString& NameToCheck)
+		{
+			return IsEventNameUnique(NameToCheck, *EditedCategoryName, EditedEvent.LogicalIndex);
+		});
+
+		// CreateUniqueName() should be able to create a name for the new event - warn and return out if this has failed
+		ensureMsgf(!NewEventName.IsEmpty(), TEXT("Failed to create a unique name for edited custom Audio Event Log event"));
+		if (NewEventName.IsEmpty())
+		{
+			return false;
+		}
+
+		if (!EditedEvent.EventName->IsEmpty())
+		{
+			static const FTextFormat SubTextFormat = FTextFormat::FromString(TEXT("{0} {1}"));
+			const FText SubText = FText::Format(SubTextFormat, LOCTEXT("AudioInsightsEditorSettings_DuplicateCustomNameSubtext", "Duplicate names are not allowed. The event name has been changed to : "), FText::FromString(NewEventName));
+
+			// This was a event name duplication fail - notify the user that the event name is being changed
+			ShowNotification(LOCTEXT("AudioInsightsEditorSettings_DuplicateCustomName", "Event name already exists in the Event Log."), SubText);
+		}
+
+		// Update the set with the new name and delete the old one
+		Events->EventNames.Remove(*EditedEvent.EventName);
+		Events->EventNames.Add(NewEventName);
+	}
+
+	return true;
+}
+
+FString* UAudioInsightsEditorSettings::GetEditedCategoryKey(const FPropertyChangedChainEvent& PropertyChangedEvent, const FMapProperty* CategoryToEventMapProperty) const
+{
+	ensureMsgf(CategoryToEventMapProperty, TEXT("Invalid FMapProperty* CategoryToEventMapProperty passed in"));
+	if (CategoryToEventMapProperty == nullptr)
+	{
+		return nullptr;
+	}
+
+	// First find the specific event that has been edited using the CategoryToEventMapProperty
+	FScriptMapHelper MapHelper(CategoryToEventMapProperty, CategoryToEventMapProperty->ContainerPtrToValuePtr<void>(&EventLogSettings));
+	const int32 ContainerIndex = PropertyChangedEvent.GetArrayIndex(GET_MEMBER_NAME_CHECKED(FAudioEventLogSettings, CustomCategoriesToEvents).ToString());
+	const int32 MapIndex = MapHelper.FindInternalIndex(ContainerIndex);
+
+	if (!MapHelper.IsValidIndex(MapIndex))
+	{
+		UE_LOG(LogAudioInsights, Warning, TEXT("Detected invalid index in CustomCategoriesToEvents TMap when trying to validate custom event log edit"));
+		return nullptr;
+	}
+
+	FString* EditedCategoryName = reinterpret_cast<FString*>(MapHelper.GetKeyPtr(MapIndex));
+	ensure(EditedCategoryName);
+	return EditedCategoryName;
+}
+
+UAudioInsightsEditorSettings::EditedEvent UAudioInsightsEditorSettings::GetEditedEvent(const FPropertyChangedChainEvent& PropertyChangedEvent, const FSetProperty* EditedEventsSetProperty, const FAudioEventLogCustomEvents& Events) const
+{
+	ensureMsgf(EditedEventsSetProperty, TEXT("Invalid FSetProperty* EditedEventsSetProperty passed in"));
+	if (EditedEventsSetProperty == nullptr)
+	{
+		return { nullptr };
+	}
+
+	FScriptSetHelper SetHelper(EditedEventsSetProperty, EditedEventsSetProperty->ContainerPtrToValuePtr<void>(&Events));
+
+	if (PropertyChangedEvent.ChangeType & EPropertyChangeType::ArrayAdd)
+	{
+		// If the change is an add action, we cannot rely on GetArrayIndex to return the correct index
+		// Newly added elements will be empty - locate the index via this method
+		for (FScriptSetHelper::FIterator Iterator = SetHelper.CreateIterator(); Iterator; ++Iterator)
+		{
+			FString* EventName = reinterpret_cast<FString*>(SetHelper.GetElementPtr(Iterator));
+			ensure(EventName);
+
+			if (EventName && EventName->IsEmpty())
+			{
+				return { EventName, Iterator.GetLogicalIndex() };
+			}
+		}
+
+		UE_LOG(LogAudioInsights, Warning, TEXT("Could not locate newly added event name element in EventNames set"));
+		return { nullptr };
+	}
+	else
+	{
+		const int32 EditedEventIndex = PropertyChangedEvent.GetArrayIndex(GET_MEMBER_NAME_CHECKED(FAudioEventLogCustomEvents, EventNames).ToString());
+		const int32 SetIndex = SetHelper.FindInternalIndex(EditedEventIndex);
+		if (!SetHelper.IsValidIndex(SetIndex))
+		{
+			UE_LOG(LogAudioInsights, Warning, TEXT("Detected invalid index in EventNames TSet when trying to validate custom event log edit"));
+			return { nullptr };
+		}
+
+		FString* EventName = reinterpret_cast<FString*>(SetHelper.GetElementPtr(SetIndex));
+		ensure(EventName);
+
+		if (EventName)
+		{
+			return { EventName, EditedEventIndex };
+		}
+	}
+
+	return { nullptr };
+}
+
+bool UAudioInsightsEditorSettings::IsEventNameUnique(const FString& EditedEventName, const FString& EditedFilterCategory, const int32 EditedLogicalIndex) const
+{
+	if (EditedEventName.IsEmpty() || EditedFilterCategory.IsEmpty() || EditedLogicalIndex == INDEX_NONE)
+	{
+		return false;
+	}
+	
+	// Check custom categories
+	for (const auto& [FilterCategory, CustomEvents] : EventLogSettings.CustomCategoriesToEvents)
+	{
+		if (CustomEvents.EventNames.Contains(EditedEventName))
+		{
+			if (EditedFilterCategory == FilterCategory)
+			{
+				const int32 FoundEventIndex = CustomEvents.EventNames.Array().IndexOfByKey(EditedEventName);
+				if (FoundEventIndex != INDEX_NONE && FoundEventIndex == EditedLogicalIndex)
+				{
+					// This is the value currently being edited, skip
+					continue;
+				}
+			}
+
+			// This is a duplicate
+			return false;
+		}
+	}
+
+	// Now check the in-built Event Types in Audio Insights
+	for (const auto& [FilterCategory, InBuiltEvents] : InBuiltAudioLogEventTypes)
+	{
+		if (InBuiltEvents.Contains(EditedEventName))
+		{
+			// No need to check if this is the value edited - users cannot edit in-built categories
+			// This is a duplicate
+			return false;
+		}
+	}
+
+	return true;
+}
+
+void UAudioInsightsEditorSettings::ShowNotification(const FText& TitleText, const FText& SubText) const
+{
+	// Show floating notification
+	FNotificationInfo Info(TitleText);
+
+	Info.SubText = SubText;
+	Info.bFireAndForget = true;
+	Info.FadeOutDuration = 1.0f;
+	Info.ExpireDuration = 4.0f;
+
+	FSlateNotificationManager::Get().AddNotification(Info);
+}
+
 #undef LOCTEXT_NAMESPACE
diff --git a/Engine/Plugins/AudioInsights/Source/AudioInsightsEditor/Public/AudioInsightsEditorSettings.h b/Engine/Plugins/AudioInsights/Source/AudioInsightsEditor/Public/AudioInsightsEditorSettings.h
--- a/Engine/Plugins/AudioInsights/Source/AudioInsightsEditor/Public/AudioInsightsEditorSettings.h
+++ b/Engine/Plugins/AudioInsights/Source/AudioInsightsEditor/Public/AudioInsightsEditorSettings.h
@@ -18,6 +18,8 @@ class UAudioInsightsEditorSettings : public UDeveloperSettings
 	GENERATED_BODY()
 	
 public:
+	UE_API UAudioInsightsEditorSettings();
+
 	UE_API virtual FName GetCategoryName() const override;
 	UE_API virtual FText GetSectionText() const override;
 	UE_API virtual FText GetSectionDescription() const override;
@@ -47,11 +49,31 @@ public:
 	FSoundDashboardSettings SoundDashboardSettings;
 
 private:
+	struct EditedEvent
+	{
+		FString* EventName = nullptr;
+		int32 LogicalIndex = INDEX_NONE;
+	};
+
+	bool VerifyEventLogInput(const TDoubleLinkedList<FProperty*>::TDoubleLinkedListNode& ActiveMemberNode, FPropertyChangedChainEvent& PropertyChangedEvent);
+	
+	void RefreshEventFilter();
+	bool VerifyCategoryEdit();
+	bool VerifyCustomEventNameEdit(const FPropertyChangedChainEvent& PropertyChangedEvent, const FMapProperty* CategoryToEventMapProperty, const FSetProperty* EventNamesSetProperty);
+
+	FString* GetEditedCategoryKey(const FPropertyChangedChainEvent& PropertyChangedEvent, const FMapProperty* CategoryToEventMapProperty) const;
+	EditedEvent GetEditedEvent(const FPropertyChangedChainEvent& PropertyChangedEvent, const FSetProperty* EditedEventsSetProperty, const FAudioEventLogCustomEvents& Events) const;
+	bool IsEventNameUnique(const FString& EditedEventName, const FString& EditedFilterCategory, const int32 EditedLogicalIndex) const;
+
+	void ShowNotification(const FText& TitleText, const FText& SubText) const;
+
 	FDelegateHandle OnRequestReadEventLogSettingsHandle;
 	FDelegateHandle OnRequestWriteEventLogSettingsHandle;
 
 	FDelegateHandle OnRequestReadSoundDashboardSettingsHandle;
 	FDelegateHandle OnRequestWriteSoundDashboardSettingsHandle;
+
+	const TMap<FString, TSet<FString>> InBuiltAudioLogEventTypes;
 };
 
 #undef UE_API
