[Audio Insights] Updates to comments in AudioInsightsSettings after bug hawk analysis

diff --git a/Engine/Plugins/AudioInsights/Source/AudioInsights/Public/Settings/AudioEventLogSettings.h b/Engine/Plugins/AudioInsights/Source/AudioInsights/Public/Settings/AudioEventLogSettings.h
--- a/Engine/Plugins/AudioInsights/Source/AudioInsights/Public/Settings/AudioEventLogSettings.h
+++ b/Engine/Plugins/AudioInsights/Source/AudioInsights/Public/Settings/AudioEventLogSettings.h
@@ -10,7 +10,7 @@
 
 #define LOCTEXT_NAMESPACE "AudioInsights"
 
-/** FAudioEventLogCustomEventType
+/** FAudioEventLogCustomEvents
  *
  * Custom event types that can be filtered by in the Event Log
  */
@@ -20,7 +20,7 @@ struct FAudioEventLogCustomEvents
 	GENERATED_BODY()
 
 	/** 
-	* The names of any custom events associated with the above category.
+	* The names of any custom events that can be sent to the Event Log.
 	* This should match the event name used to send the event to the Event Log.
 	*/
 	UPROPERTY(EditAnywhere, config, Category = EventLog, meta = (DisplayName = "Custom event names"))
diff --git a/Engine/Plugins/AudioInsights/Source/AudioInsightsEditor/Private/AudioInsightsEditorSettings.cpp b/Engine/Plugins/AudioInsights/Source/AudioInsightsEditor/Private/AudioInsightsEditorSettings.cpp
--- a/Engine/Plugins/AudioInsights/Source/AudioInsightsEditor/Private/AudioInsightsEditorSettings.cpp
+++ b/Engine/Plugins/AudioInsights/Source/AudioInsightsEditor/Private/AudioInsightsEditorSettings.cpp
@@ -203,7 +203,7 @@ bool UAudioInsightsEditorSettings::VerifyCategoryEdit()
 	{
 		if (FilterCategory.IsEmpty())
 		{
-			// Create a unique event name if the category is empty
+			// Create a unique category name if the category is empty
 			const FString NewCategoryName = CreateUniqueName(TEXT("New Category"), [this](const FString& NameToCheck)
 			{
 				return !EventLogSettings.CustomCategoriesToEvents.Contains(NameToCheck);
