[Audio Insights] Update Volume and Pitch plot parameter range metadata.

diff --git a/Engine/Plugins/AudioInsights/Source/AudioInsights/Public/Settings/SoundDashboardSettings.h b/Engine/Plugins/AudioInsights/Source/AudioInsights/Public/Settings/SoundDashboardSettings.h
--- a/Engine/Plugins/AudioInsights/Source/AudioInsights/Public/Settings/SoundDashboardSettings.h
+++ b/Engine/Plugins/AudioInsights/Source/AudioInsights/Public/Settings/SoundDashboardSettings.h
@@ -140,7 +140,7 @@ struct FSoundPlotsDashboardPlotRanges
 	UPROPERTY(EditAnywhere, config, Category = "SoundDashboard|CustomPlotRanges", meta = (InlineEditConditionToggle))
 	bool bUseCustomVolumeRange = false;
 
-	UPROPERTY(EditAnywhere, config, Category = "SoundDashboard|CustomPlotRanges", meta = (EditCondition = "bUseCustomVolumeRange", UIMin = "0", UIMax = "1", ClampMin = "0", ClampMax = "1"))
+	UPROPERTY(EditAnywhere, config, Category = "SoundDashboard|CustomPlotRanges", meta = (EditCondition = "bUseCustomVolumeRange", UIMin = "0", UIMax = "4", ClampMin = "0", ClampMax = "4"))
 	FFloatInterval Volume = { 0.0f, 1.0f };
 
 
@@ -161,8 +161,8 @@ struct FSoundPlotsDashboardPlotRanges
 	UPROPERTY(EditAnywhere, config, Category = "SoundDashboard|CustomPlotRanges", meta = (InlineEditConditionToggle))
 	bool bUseCustomPitchRange = false;
 
-	UPROPERTY(EditAnywhere, config, Category = "SoundDashboard|CustomPlotRanges", meta = (EditCondition = "bUseCustomPitchRange", UIMin = "0.25", UIMax = "4", ClampMin = "0", ClampMax = "32", Units = "Multiplier"))
-	FFloatInterval Pitch = { 0.25f, 4.0f };
+	UPROPERTY(EditAnywhere, config, Category = "SoundDashboard|CustomPlotRanges", meta = (EditCondition = "bUseCustomPitchRange", UIMin = "0.125", UIMax = "4", ClampMin = "0", ClampMax = "32", Units = "Multiplier"))
+	FFloatInterval Pitch = { 0.125f, 4.0f };
 
 
 	UPROPERTY(EditAnywhere, config, Category = "SoundDashboard|CustomPlotRanges", meta = (InlineEditConditionToggle))
