[Audio Insights] Refactor VerifyCategoryEdit to use .Find() to get empty categories rather than for loop - fix after bughawk suggestion

diff --git a/Engine/Plugins/AudioInsights/Source/AudioInsightsEditor/Private/AudioInsightsEditorSettings.cpp b/Engine/Plugins/AudioInsights/Source/AudioInsightsEditor/Private/AudioInsightsEditorSettings.cpp
--- a/Engine/Plugins/AudioInsights/Source/AudioInsightsEditor/Private/AudioInsightsEditorSettings.cpp
+++ b/Engine/Plugins/AudioInsights/Source/AudioInsightsEditor/Private/AudioInsightsEditorSettings.cpp
@@ -198,30 +198,26 @@ bool UAudioInsightsEditorSettings::VerifyCategoryEdit()
 	using namespace AudioInsightsEditorSettingsPrivate;
 
 	// TMap will already check to ensure there are no duplicate Keys
-	// If we detect any empty fields, try to create a unique default name to replace them
-	for (auto& [FilterCategory, Events] : EventLogSettings.CustomCategoriesToEvents)
+	// If we detect an empty field, try to create a unique default name to replace them
+	const FString EmptyCategoryName = FString();
+	if (FAudioEventLogCustomEvents* Events = EventLogSettings.CustomCategoriesToEvents.Find(EmptyCategoryName))
 	{
-		if (FilterCategory.IsEmpty())
+		// Create a unique category name if the category is empty
+		const FString NewCategoryName = CreateUniqueName(TEXT("New Category"), [this](const FString& NameToCheck)
 		{
-			// Create a unique category name if the category is empty
-			const FString NewCategoryName = CreateUniqueName(TEXT("New Category"), [this](const FString& NameToCheck)
-			{
-				return !EventLogSettings.CustomCategoriesToEvents.Contains(NameToCheck);
-			});
-
-			// CreateUniqueName() should be able to create a name for the new category - warn and return out if this has failed
-			ensureMsgf(!NewCategoryName.IsEmpty(), TEXT("Failed to create a unique name for edited custom Audio Event Log category"));
-			if (NewCategoryName.IsEmpty())
-			{
-				return false;
-			}
-
-			// There should only be one empty value at a time (i.e. when a new element is added) - edit the map immediately and break;
-			EventLogSettings.CustomCategoriesToEvents.Add(NewCategoryName, MoveTemp(Events));
-			EventLogSettings.CustomCategoriesToEvents.Remove(FilterCategory);
+			return !EventLogSettings.CustomCategoriesToEvents.Contains(NameToCheck);
+		});
 
-			break;
+		// CreateUniqueName() should be able to create a name for the new category - warn and return out if this has failed
+		ensureMsgf(!NewCategoryName.IsEmpty(), TEXT("Failed to create a unique name for edited custom Audio Event Log category"));
+		if (NewCategoryName.IsEmpty())
+		{
+			return false;
 		}
+
+		// There should only be one empty value at a time (i.e. when a new element is added) - edit the map immediately and break;
+		EventLogSettings.CustomCategoriesToEvents.Add(NewCategoryName, MoveTemp(*Events));
+		EventLogSettings.CustomCategoriesToEvents.Remove(EmptyCategoryName);
 	}
 
 	return true;
