[Audio Insights] Updated comment after bug hawk suggestion

diff --git a/Engine/Plugins/AudioInsights/Source/AudioInsightsEditor/Private/AudioInsightsEditorSettings.cpp b/Engine/Plugins/AudioInsights/Source/AudioInsightsEditor/Private/AudioInsightsEditorSettings.cpp
--- a/Engine/Plugins/AudioInsights/Source/AudioInsightsEditor/Private/AudioInsightsEditorSettings.cpp
+++ b/Engine/Plugins/AudioInsights/Source/AudioInsightsEditor/Private/AudioInsightsEditorSettings.cpp
@@ -215,7 +215,7 @@ bool UAudioInsightsEditorSettings::VerifyCategoryEdit()
 			return false;
 		}
 
-		// There should only be one empty value at a time (i.e. when a new element is added) - edit the map immediately and break;
+		// Move the events into the newly named category - delete the empty one
 		EventLogSettings.CustomCategoriesToEvents.Add(NewCategoryName, MoveTemp(*Events));
 		EventLogSettings.CustomCategoriesToEvents.Remove(EmptyCategoryName);
 	}
