[Audio Insights] Remove Audio Insights build dependency on Engine to fix build break.

diff --git a/Engine/Plugins/AudioInsights/Source/AudioInsights/Private/Views/AudioEventLogDashboardViewFactory.cpp b/Engine/Plugins/AudioInsights/Source/AudioInsights/Private/Views/AudioEventLogDashboardViewFactory.cpp
--- a/Engine/Plugins/AudioInsights/Source/AudioInsights/Private/Views/AudioEventLogDashboardViewFactory.cpp
+++ b/Engine/Plugins/AudioInsights/Source/AudioInsights/Private/Views/AudioEventLogDashboardViewFactory.cpp
@@ -1,7 +1,7 @@
 // Copyright Epic Games, Inc. All Rights Reserved.
 #include "Views/AudioEventLogDashboardViewFactory.h"
 
-#include "AudioDeviceTrace.h"
+//#include "AudioDeviceTrace.h"
 #include "AudioEventLogEditorCommands.h"
 #include "AudioInsightsModule.h"
 #include "AudioInsightsStyle.h"
@@ -92,10 +92,10 @@ namespace UE::Audio::Insights
 			{
 				CategoryPlayErrors,
 				{
-					ID::PlayFailedNotPlayable,
-					ID::PlayFailedOutOfRange,
-					ID::PlayFailedDebugFiltered,
-					ID::PlayFailedConcurrency,
+					//ID::PlayFailedNotPlayable,
+					//ID::PlayFailedOutOfRange,
+					//ID::PlayFailedDebugFiltered,
+					//ID::PlayFailedConcurrency,
 				}
 			},
 
