[Audio Insights] Update Priority plot parameter range metadata.

diff --git a/Engine/Plugins/AudioInsights/Source/AudioInsights/Public/Settings/SoundDashboardSettings.h b/Engine/Plugins/AudioInsights/Source/AudioInsights/Public/Settings/SoundDashboardSettings.h
--- a/Engine/Plugins/AudioInsights/Source/AudioInsights/Public/Settings/SoundDashboardSettings.h
+++ b/Engine/Plugins/AudioInsights/Source/AudioInsights/Public/Settings/SoundDashboardSettings.h
@@ -168,8 +168,8 @@ struct FSoundPlotsDashboardPlotRanges
 	UPROPERTY(EditAnywhere, config, Category = "SoundDashboard|CustomPlotRanges", meta = (InlineEditConditionToggle))
 	bool bUseCustomPriorityRange = false;
 
-	UPROPERTY(EditAnywhere, config, Category = "SoundDashboard|CustomPlotRanges", meta = (EditCondition = "bUseCustomPriorityRange", UIMin = "0", UIMax = "255", ClampMin = "0", ClampMax = "255"))
-	FFloatInterval Priority = { 0.0f, 255.0f };
+	UPROPERTY(EditAnywhere, config, Category = "SoundDashboard|CustomPlotRanges", meta = (EditCondition = "bUseCustomPriorityRange", UIMin = "0", UIMax = "100", ClampMin = "0", ClampMax = "100"))
+	FFloatInterval Priority = { 0.0f, 100.0f };
 
 
 	UPROPERTY(EditAnywhere, config, Category = "SoundDashboard|CustomPlotRanges", meta = (InlineEditConditionToggle))
