[Audio Insights] Created TAssetProvider class to reduce code duplication in AudioBus and Submix Asset Providers + added missing hookup to Asset Provider RequestEntriesUpdate() method

diff --git a/Engine/Plugins/AudioInsights/Source/AudioInsights/Private/Providers/AudioBusAssetProvider.cpp b/Engine/Plugins/AudioInsights/Source/AudioInsights/Private/Providers/AudioBusAssetProvider.cpp
--- a/Engine/Plugins/AudioInsights/Source/AudioInsights/Private/Providers/AudioBusAssetProvider.cpp
+++ b/Engine/Plugins/AudioInsights/Source/AudioInsights/Private/Providers/AudioBusAssetProvider.cpp
@@ -3,157 +3,27 @@
 #include "Providers/AudioBusAssetProvider.h"
 
 #if WITH_EDITOR
-#include "IAudioInsightsModule.h"
+
 #include "AssetRegistry/AssetData.h"
-#include "AssetRegistry/AssetRegistryModule.h"
-#include "AudioDeviceManager.h"
-#include "AudioMixerDevice.h"
+#include "Messages/AudioBusTraceMessages.h"
 #include "Sound/AudioBus.h"
 
 namespace UE::Audio::Insights
 {
-	FAudioBusAssetProvider::FAudioBusAssetProvider()
-	{
-		FAssetRegistryModule& AssetRegistryModule = FModuleManager::LoadModuleChecked<FAssetRegistryModule>("AssetRegistry");
-
-		AssetRegistryModule.Get().OnAssetAdded().AddRaw(this, &FAudioBusAssetProvider::OnAssetAdded);
-		AssetRegistryModule.Get().OnAssetRemoved().AddRaw(this, &FAudioBusAssetProvider::OnAssetRemoved);
-		AssetRegistryModule.Get().OnFilesLoaded().AddRaw(this, &FAudioBusAssetProvider::OnFilesLoaded);
-
-		FAudioDeviceManagerDelegates::OnAudioDeviceCreated.AddRaw(this, &FAudioBusAssetProvider::OnActiveAudioDeviceChanged);
-		FAudioDeviceManagerDelegates::OnAudioDeviceDestroyed.AddRaw(this, &FAudioBusAssetProvider::OnActiveAudioDeviceChanged);
-
-		FTraceAuxiliary::OnTraceStarted.AddRaw(this, &FAudioBusAssetProvider::OnTraceStarted);
-	}
-
-	FAudioBusAssetProvider::~FAudioBusAssetProvider()
-	{
-		if (FAssetRegistryModule* AssetRegistryModule = FModuleManager::GetModulePtr<FAssetRegistryModule>("AssetRegistry"))
-		{
-			AssetRegistryModule->Get().OnAssetAdded().RemoveAll(this);
-			AssetRegistryModule->Get().OnAssetRemoved().RemoveAll(this);
-			AssetRegistryModule->Get().OnFilesLoaded().RemoveAll(this);
-		}
-
-		FAudioDeviceManagerDelegates::OnAudioDeviceCreated.RemoveAll(this);
-		FAudioDeviceManagerDelegates::OnAudioDeviceDestroyed.RemoveAll(this);
-
-		FTraceAuxiliary::OnTraceStarted.RemoveAll(this);
-	}
-
-	void FAudioBusAssetProvider::OnAssetAdded(const FAssetData& InAssetData)
+	uint32 FAudioBusAssetProvider::GetAssetId(const TSharedPtr<FAudioBusDashboardEntry>& InEntry) const
 	{
-		if (bAreFilesLoaded && InAssetData.AssetClassPath == FTopLevelAssetPath(UAudioBus::StaticClass()))
-		{
-			constexpr bool bShouldNotify = true;
-			AddAudioBusAsset(InAssetData, bShouldNotify);
-		}
+		return InEntry->AudioBusId;
 	}
 
-	void FAudioBusAssetProvider::OnAssetRemoved(const FAssetData& InAssetData)
+	void FAudioBusAssetProvider::InitAssetEntry(TSharedPtr<FAudioBusDashboardEntry>& InEntry, const FAssetData& InAssetData)
 	{
-		if (InAssetData.AssetClassPath == FTopLevelAssetPath(UAudioBus::StaticClass()))
-		{
-			RemoveAudioBusAsset(InAssetData);
-		}
-	}
-
-	void FAudioBusAssetProvider::OnFilesLoaded()
-	{
-		bAreFilesLoaded = true;
-		UpdateAudioBusAssetNames();
-	}
-
-	void FAudioBusAssetProvider::OnActiveAudioDeviceChanged(::Audio::FDeviceId InDeviceId)
-	{
-		UpdateAudioBusAssetNames();
-	}
-
-	void FAudioBusAssetProvider::OnTraceStarted(FTraceAuxiliary::EConnectionType TraceType, const FString& TraceDestination)
-	{
-		UpdateAudioBusAssetNames();
-	}
-
-	void FAudioBusAssetProvider::RequestEntriesUpdate()
-	{
-		UpdateAudioBusAssetNames();
-	}
-	
-	void FAudioBusAssetProvider::AddAudioBusAsset(const FAssetData& InAssetData, const bool bInShouldNotify)
-	{
-		const bool bIsAudioBusEntryAlreadAdded = AudioBusAssetEntries.ContainsByPredicate(
-			[AssetName = InAssetData.GetObjectPathString()](const TSharedPtr<FAudioBusDashboardEntry>& AudioBusDashboardEntry)
-			{
-				return AudioBusDashboardEntry->Name == AssetName;
-			});
-
-		if (!bIsAudioBusEntryAlreadAdded)
-		{
-			const IAudioInsightsModule& AudioInsightsModule = IAudioInsightsModule::GetEditorChecked();
-			const ::Audio::FDeviceId AudioDeviceId = AudioInsightsModule.GetDeviceId();
-
-			const TObjectPtr<UAudioBus> AudioBus = Cast<UAudioBus>(InAssetData.GetAsset());
-
-			TSharedPtr<FAudioBusDashboardEntry>& AudioBusDashboardEntry = AudioBusAssetEntries.Add_GetRef(MakeShared<FAudioBusDashboardEntry>());
-			AudioBusDashboardEntry->Timestamp = FPlatformTime::Seconds() - GStartTime;
-			AudioBusDashboardEntry->AudioBusId = AudioBus ? AudioBus->GetUniqueID() : INDEX_NONE; //GetTypeHash(InAssetData.GetObjectPathString());
-			AudioBusDashboardEntry->Name = InAssetData.GetObjectPathString();
-			AudioBusDashboardEntry->AudioMeterInfo->NumChannels = AudioBus->GetNumChannels();
-			AudioBusDashboardEntry->AudioBusType = EAudioBusType::AssetBased;
-
-			if (bInShouldNotify)
-			{
-				OnAudioBusAssetAdded.ExecuteIfBound(AudioBusDashboardEntry);
-			}
-		}
-	}
-
-	void FAudioBusAssetProvider::RemoveAudioBusAsset(const FAssetData& InAssetData)
-	{
-		const TSharedPtr<FAudioBusDashboardEntry>* FoundAudioBusDashboardEntry = AudioBusAssetEntries.FindByPredicate(
-			[AssetName = InAssetData.GetObjectPathString()](const TSharedPtr<FAudioBusDashboardEntry>& InAudioBusDashboardEntry)
-			{
-				return InAudioBusDashboardEntry->Name == AssetName;
-			});
-
-		if (FoundAudioBusDashboardEntry)
-		{
-			const TSharedPtr<FAudioBusDashboardEntry> AudioBusDashboardEntry = *FoundAudioBusDashboardEntry;
-
-			if (AudioBusDashboardEntry.IsValid())
-			{
-				const uint32 AudioBusId = AudioBusDashboardEntry->AudioBusId;
-
-				AudioBusAssetEntries.Remove(AudioBusDashboardEntry);
-
-				OnAudioBusAssetRemoved.ExecuteIfBound(AudioBusId);
-			}
-		}
-	}
-
-	void FAudioBusAssetProvider::UpdateAudioBusAssetNames()
-	{
-		// Get all UAudioBus assets
-		TArray<FAssetData> AssetDataArray;
-
-		FAssetRegistryModule& AssetRegistryModule = FModuleManager::LoadModuleChecked<FAssetRegistryModule>("AssetRegistry");
-		AssetRegistryModule.Get().GetAssetsByClass(FTopLevelAssetPath(UAudioBus::StaticClass()), AssetDataArray);
-
-		// Build AudioBusAssetEntries
-		AudioBusAssetEntries.Reset();
-
-		constexpr bool bShouldNotify = false;
-		for (const FAssetData& AssetData : AssetDataArray)
-		{
-			AddAudioBusAsset(AssetData, bShouldNotify);
-		}
-
-		AudioBusAssetEntries.Sort([](const TSharedPtr<FAudioBusDashboardEntry>& A, const TSharedPtr<FAudioBusDashboardEntry>& B)
-		{
-			return A->GetDisplayName().CompareToCaseIgnored(B->GetDisplayName()) < 0;
-		});
-
-		OnAudioBusAssetListUpdated.ExecuteIfBound(AudioBusAssetEntries);
+		const TObjectPtr<UAudioBus> AudioBus = Cast<UAudioBus>(InAssetData.GetAsset());
+		
+		InEntry->Timestamp = FPlatformTime::Seconds() - GStartTime;
+		InEntry->AudioBusId = AudioBus ? AudioBus->GetUniqueID() : INDEX_NONE;
+		InEntry->Name = InAssetData.GetObjectPathString();
+		InEntry->AudioMeterInfo->NumChannels = AudioBus->GetNumChannels();
+		InEntry->AudioBusType = EAudioBusType::AssetBased;
 	}
 } // namespace UE::Audio::Insights
 
diff --git a/Engine/Plugins/AudioInsights/Source/AudioInsights/Private/Providers/AudioBusAssetProvider.h b/Engine/Plugins/AudioInsights/Source/AudioInsights/Private/Providers/AudioBusAssetProvider.h
--- a/Engine/Plugins/AudioInsights/Source/AudioInsights/Private/Providers/AudioBusAssetProvider.h
+++ b/Engine/Plugins/AudioInsights/Source/AudioInsights/Private/Providers/AudioBusAssetProvider.h
@@ -2,48 +2,21 @@
 #pragma once
 
 #if WITH_EDITOR
-#include "Delegates/Delegate.h"
-#include "Messages/AudioBusTraceMessages.h"
-#include "ProfilingDebugging/TraceAuxiliary.h"
 
+#include "Providers/AssetProvider.h"
+
+class UAudioBus;
 struct FAssetData;
 
 namespace UE::Audio::Insights
 {
-	class FAudioBusAssetProvider : public TSharedFromThis<FAudioBusAssetProvider>
-	{
-	public:
-		FAudioBusAssetProvider();
-		virtual ~FAudioBusAssetProvider();
-
-		void RequestEntriesUpdate();
-
-		DECLARE_DELEGATE_OneParam(FOnAudioBusAssetAdded, const TSharedPtr<FAudioBusDashboardEntry>& /*AudioBusDashboardEntry*/);
-		FOnAudioBusAssetAdded OnAudioBusAssetAdded;
-
-		DECLARE_DELEGATE_OneParam(FOnAudioBusAssetRemoved, const uint32 /*AudioBusId*/);
-		FOnAudioBusAssetRemoved OnAudioBusAssetRemoved;
-
-		DECLARE_DELEGATE_OneParam(FOnAudioBusAssetListUpdated, const TArray<TSharedPtr<FAudioBusDashboardEntry>>& /*AudioBusDashboardEntries*/);
-		FOnAudioBusAssetListUpdated OnAudioBusAssetListUpdated;
+	struct FAudioBusDashboardEntry;
 
-	private:
-		void OnAssetAdded(const FAssetData& InAssetData);
-		void OnAssetRemoved(const FAssetData& InAssetData);
-		void OnFilesLoaded();
-
-		void OnActiveAudioDeviceChanged(::Audio::FDeviceId InDeviceId);
-		void OnTraceStarted(FTraceAuxiliary::EConnectionType TraceType, const FString& TraceDestination);
-
-		void AddAudioBusAsset(const FAssetData& InAssetData, const bool bInShouldNotify);
-		void RemoveAudioBusAsset(const FAssetData& InAssetData);
-
-		void UpdateAudioBusAssetNames();
-
-		bool bAreFilesLoaded = false;
-		bool bAssetEntriesNeedRefreshing = false;
-
-		TArray<TSharedPtr<FAudioBusDashboardEntry>> AudioBusAssetEntries;
+	class FAudioBusAssetProvider : public TAssetProvider<UAudioBus, FAudioBusDashboardEntry>
+	{
+	protected:
+		virtual uint32 GetAssetId(const TSharedPtr<FAudioBusDashboardEntry>& InEntry) const override;
+		virtual void InitAssetEntry(TSharedPtr<FAudioBusDashboardEntry>& InEntry, const FAssetData& InAssetData) override;
 	};
 } // namespace UE::Audio::Insights
 
diff --git a/Engine/Plugins/AudioInsights/Source/AudioInsights/Private/Providers/AudioBusTraceProvider.cpp b/Engine/Plugins/AudioInsights/Source/AudioInsights/Private/Providers/AudioBusTraceProvider.cpp
--- a/Engine/Plugins/AudioInsights/Source/AudioInsights/Private/Providers/AudioBusTraceProvider.cpp
+++ b/Engine/Plugins/AudioInsights/Source/AudioInsights/Private/Providers/AudioBusTraceProvider.cpp
@@ -38,18 +38,18 @@ namespace UE::Audio::Insights
 		: TDeviceDataMapTraceProvider<uint32, TSharedPtr<FAudioBusDashboardEntry>>(GetName_Static())
 	{
 #if WITH_EDITOR
-		AudioBusAssetProvider.OnAudioBusAssetAdded.BindRaw(this, &FAudioBusTraceProvider::HandleOnAudioBusAssetAdded);
-		AudioBusAssetProvider.OnAudioBusAssetRemoved.BindRaw(this, &FAudioBusTraceProvider::HandleOnAudioBusAssetRemoved);
-		AudioBusAssetProvider.OnAudioBusAssetListUpdated.BindRaw(this, &FAudioBusTraceProvider::HandleOnAudioBusAssetListUpdated);
+		AudioBusAssetProvider.OnAssetAdded.BindRaw(this, &FAudioBusTraceProvider::HandleOnAudioBusAssetAdded);
+		AudioBusAssetProvider.OnAssetRemoved.BindRaw(this, &FAudioBusTraceProvider::HandleOnAudioBusAssetRemoved);
+		AudioBusAssetProvider.OnAssetListUpdated.BindRaw(this, &FAudioBusTraceProvider::HandleOnAudioBusAssetListUpdated);
 #endif // WITH_EDITOR
 	}
 
 	FAudioBusTraceProvider::~FAudioBusTraceProvider()
 	{
 #if WITH_EDITOR
-		AudioBusAssetProvider.OnAudioBusAssetAdded.Unbind();
-		AudioBusAssetProvider.OnAudioBusAssetRemoved.Unbind();
-		AudioBusAssetProvider.OnAudioBusAssetListUpdated.Unbind();
+		AudioBusAssetProvider.OnAssetAdded.Unbind();
+		AudioBusAssetProvider.OnAssetRemoved.Unbind();
+		AudioBusAssetProvider.OnAssetListUpdated.Unbind();
 #endif // WITH_EDITOR
 	}
 	
@@ -263,6 +263,11 @@ namespace UE::Audio::Insights
 	}
 
 #if WITH_EDITOR
+	void FAudioBusTraceProvider::RequestEntriesUpdate()
+	{
+		AudioBusAssetProvider.RequestEntriesUpdate();
+	}
+
 	void FAudioBusTraceProvider::HandleOnAudioBusAssetAdded(const TSharedPtr<FAudioBusDashboardEntry>& InAudioBusDashboardEntry)
 	{
 		if (const FAudioDeviceManager* DeviceManager = FAudioDeviceManager::Get())
diff --git a/Engine/Plugins/AudioInsights/Source/AudioInsights/Private/Providers/AudioBusTraceProvider.h b/Engine/Plugins/AudioInsights/Source/AudioInsights/Private/Providers/AudioBusTraceProvider.h
--- a/Engine/Plugins/AudioInsights/Source/AudioInsights/Private/Providers/AudioBusTraceProvider.h
+++ b/Engine/Plugins/AudioInsights/Source/AudioInsights/Private/Providers/AudioBusTraceProvider.h
@@ -26,6 +26,7 @@ namespace UE::Audio::Insights
 #if !WITH_EDITOR
 		virtual void InitSessionCachedMessages(TraceServices::IAnalysisSession& InSession) override;
 #endif // !WITH_EDITOR
+		void RequestEntriesUpdate();
 
 		DECLARE_DELEGATE_OneParam(FOnAudioBusChanged, const uint32 /*AudioBusId*/);
 		FOnAudioBusChanged OnAudioBusAdded;
diff --git a/Engine/Plugins/AudioInsights/Source/AudioInsights/Private/Providers/SubmixAssetProvider.cpp b/Engine/Plugins/AudioInsights/Source/AudioInsights/Private/Providers/SubmixAssetProvider.cpp
--- a/Engine/Plugins/AudioInsights/Source/AudioInsights/Private/Providers/SubmixAssetProvider.cpp
+++ b/Engine/Plugins/AudioInsights/Source/AudioInsights/Private/Providers/SubmixAssetProvider.cpp
@@ -3,162 +3,37 @@
 #include "Providers/SubmixAssetProvider.h"
 
 #if WITH_EDITOR
-#include "IAudioInsightsModule.h"
+
 #include "AssetRegistry/AssetData.h"
-#include "AssetRegistry/AssetRegistryModule.h"
 #include "AudioDeviceManager.h"
 #include "AudioMixerDevice.h"
-#include "Sound/SoundSubmix.h"
+#include "IAudioInsightsModule.h"
+#include "Messages/SubmixTraceMessages.h"
 
 namespace UE::Audio::Insights
 {
-	FSubmixAssetProvider::FSubmixAssetProvider()
-	{
-		FAssetRegistryModule& AssetRegistryModule = FModuleManager::LoadModuleChecked<FAssetRegistryModule>("AssetRegistry");
-
-		AssetRegistryModule.Get().OnAssetAdded().AddRaw(this, &FSubmixAssetProvider::OnAssetAdded);
-		AssetRegistryModule.Get().OnAssetRemoved().AddRaw(this, &FSubmixAssetProvider::OnAssetRemoved);
-		AssetRegistryModule.Get().OnFilesLoaded().AddRaw(this, &FSubmixAssetProvider::OnFilesLoaded);
-
-		FAudioDeviceManagerDelegates::OnAudioDeviceCreated.AddRaw(this, &FSubmixAssetProvider::OnActiveAudioDeviceChanged);
-		FAudioDeviceManagerDelegates::OnAudioDeviceDestroyed.AddRaw(this, &FSubmixAssetProvider::OnActiveAudioDeviceChanged);
-
-		FTraceAuxiliary::OnTraceStarted.AddRaw(this, &FSubmixAssetProvider::OnTraceStarted);
-	}
-
-	FSubmixAssetProvider::~FSubmixAssetProvider()
-	{
-		if (FAssetRegistryModule* AssetRegistryModule = FModuleManager::GetModulePtr<FAssetRegistryModule>("AssetRegistry"))
-		{
-			AssetRegistryModule->Get().OnAssetAdded().RemoveAll(this);
-			AssetRegistryModule->Get().OnAssetRemoved().RemoveAll(this);
-			AssetRegistryModule->Get().OnFilesLoaded().RemoveAll(this);
-		}
-
-		FAudioDeviceManagerDelegates::OnAudioDeviceCreated.RemoveAll(this);
-		FAudioDeviceManagerDelegates::OnAudioDeviceDestroyed.RemoveAll(this);
-
-		FTraceAuxiliary::OnTraceStarted.RemoveAll(this);
-	}
-
-	void FSubmixAssetProvider::OnAssetAdded(const FAssetData& InAssetData)
-	{
-		if (bAreFilesLoaded && InAssetData.AssetClassPath == FTopLevelAssetPath(USoundSubmix::StaticClass()))
-		{
-			constexpr bool bShouldNotify = true;
-			AddSubmixAsset(InAssetData, bShouldNotify);
-		}
-	}
-
-	void FSubmixAssetProvider::OnAssetRemoved(const FAssetData& InAssetData)
-	{
-		if (InAssetData.AssetClassPath == FTopLevelAssetPath(USoundSubmix::StaticClass()))
-		{
-			RemoveSubmixAsset(InAssetData);
-		}
-	}
-
-	void FSubmixAssetProvider::OnFilesLoaded()
-	{
-		bAreFilesLoaded = true;
-		UpdateSubmixAssetNames();
-	}
-
-	void FSubmixAssetProvider::OnActiveAudioDeviceChanged(::Audio::FDeviceId InDeviceId)
-	{
-		UpdateSubmixAssetNames();
-	}
-
-	void FSubmixAssetProvider::OnTraceStarted(FTraceAuxiliary::EConnectionType TraceType, const FString& TraceDestination)
+	uint32 FSubmixAssetProvider::GetAssetId(const TSharedPtr<FSubmixDashboardEntry>& InEntry) const
 	{
-		UpdateSubmixAssetNames();
+		return InEntry->SubmixId;
 	}
 
-	void FSubmixAssetProvider::RequestEntriesUpdate()
-	{
-		UpdateSubmixAssetNames();
-	}
-	
-	void FSubmixAssetProvider::AddSubmixAsset(const FAssetData& InAssetData, const bool bInShouldNotify)
+	void FSubmixAssetProvider::InitAssetEntry(TSharedPtr<FSubmixDashboardEntry>& InEntry, const FAssetData& InAssetData)
 	{
-		const bool bIsSubmixEntryAlreadAdded = SubmixAssetEntries.ContainsByPredicate(
-			[AssetName = InAssetData.GetObjectPathString()](const TSharedPtr<FSubmixDashboardEntry>& SubmixDashboardEntry)
-			{
-				return SubmixDashboardEntry->Name == AssetName;
-			});
-
-		if (!bIsSubmixEntryAlreadAdded)
-		{
-			const IAudioInsightsModule& AudioInsightsModule = IAudioInsightsModule::GetEditorChecked();
-			const ::Audio::FDeviceId AudioDeviceId = AudioInsightsModule.GetDeviceId();
-
-			TSharedPtr<FSubmixDashboardEntry>& SubmixDashboardEntry = SubmixAssetEntries.Add_GetRef(MakeShared<FSubmixDashboardEntry>());
-			SubmixDashboardEntry->Timestamp = FPlatformTime::Seconds() - GStartTime;
-			SubmixDashboardEntry->SubmixId  = GetTypeHash(InAssetData.GetObjectPathString());
-			SubmixDashboardEntry->Name      = InAssetData.GetObjectPathString();
-
-			if (const FAudioDeviceManager* AudioDeviceManager = FAudioDeviceManager::Get())
-			{
-				if (const ::Audio::FMixerDevice* MixerDevice = static_cast<const ::Audio::FMixerDevice*>(AudioDeviceManager->GetAudioDeviceRaw(AudioDeviceId)))
-				{
-					SubmixDashboardEntry->AudioMeterInfo->NumChannels = MixerDevice->GetNumDeviceChannels();
-				}
-			}
+		const IAudioInsightsModule& AudioInsightsModule = IAudioInsightsModule::GetEditorChecked();
+		const ::Audio::FDeviceId AudioDeviceId = AudioInsightsModule.GetDeviceId();
 
-			if (bInShouldNotify)
-			{
-				OnSubmixAssetAdded.ExecuteIfBound(SubmixDashboardEntry);
-			}
-		}
-	}
+		InEntry->Timestamp = FPlatformTime::Seconds() - GStartTime;
+		InEntry->SubmixId = GetTypeHash(InAssetData.GetObjectPathString());
+		InEntry->Name = InAssetData.GetObjectPathString();
 
-	void FSubmixAssetProvider::RemoveSubmixAsset(const FAssetData& InAssetData)
-	{
-		const TSharedPtr<FSubmixDashboardEntry>* FoundSubmixDashboardEntry = SubmixAssetEntries.FindByPredicate(
-			[AssetName = InAssetData.GetObjectPathString()](const TSharedPtr<FSubmixDashboardEntry>& InSubmixDashboardEntry)
-			{
-				return InSubmixDashboardEntry->Name == AssetName;
-			});
-
-		if (FoundSubmixDashboardEntry)
+		if (const FAudioDeviceManager* AudioDeviceManager = FAudioDeviceManager::Get())
 		{
-			const TSharedPtr<FSubmixDashboardEntry> SubmixDashboardEntry = *FoundSubmixDashboardEntry;
-
-			if (SubmixDashboardEntry.IsValid())
+			if (const ::Audio::FMixerDevice* MixerDevice = static_cast<const ::Audio::FMixerDevice*>(AudioDeviceManager->GetAudioDeviceRaw(AudioDeviceId)))
 			{
-				const uint32 SubmixId = SubmixDashboardEntry->SubmixId;
-
-				SubmixAssetEntries.Remove(SubmixDashboardEntry);
-
-				OnSubmixAssetRemoved.ExecuteIfBound(SubmixId);
+				InEntry->AudioMeterInfo->NumChannels = MixerDevice->GetNumDeviceChannels();
 			}
 		}
 	}
-
-	void FSubmixAssetProvider::UpdateSubmixAssetNames()
-	{
-		// Get all USoundSubmix assets
-		TArray<FAssetData> AssetDataArray;
-
-		FAssetRegistryModule& AssetRegistryModule = FModuleManager::LoadModuleChecked<FAssetRegistryModule>("AssetRegistry");
-		AssetRegistryModule.Get().GetAssetsByClass(FTopLevelAssetPath(USoundSubmix::StaticClass()), AssetDataArray);
-
-		// Build SubmixAssetEntries
-		SubmixAssetEntries.Reset();
-
-		constexpr bool bShouldNotify = false;
-		for (const FAssetData& AssetData : AssetDataArray)
-		{
-			AddSubmixAsset(AssetData, bShouldNotify);
-		}
-
-		SubmixAssetEntries.Sort([](const TSharedPtr<FSubmixDashboardEntry>& A, const TSharedPtr<FSubmixDashboardEntry>& B)
-		{
-			return A->GetDisplayName().CompareToCaseIgnored(B->GetDisplayName()) < 0;
-		});
-
-		OnSubmixAssetListUpdated.ExecuteIfBound(SubmixAssetEntries);
-	}
 } // namespace UE::Audio::Insights
 
 #endif // WITH_EDITOR
diff --git a/Engine/Plugins/AudioInsights/Source/AudioInsights/Private/Providers/SubmixAssetProvider.h b/Engine/Plugins/AudioInsights/Source/AudioInsights/Private/Providers/SubmixAssetProvider.h
--- a/Engine/Plugins/AudioInsights/Source/AudioInsights/Private/Providers/SubmixAssetProvider.h
+++ b/Engine/Plugins/AudioInsights/Source/AudioInsights/Private/Providers/SubmixAssetProvider.h
@@ -2,48 +2,21 @@
 #pragma once
 
 #if WITH_EDITOR
-#include "Delegates/Delegate.h"
-#include "Messages/SubmixTraceMessages.h"
-#include "ProfilingDebugging/TraceAuxiliary.h"
 
+#include "Providers/AssetProvider.h"
+
+class USoundSubmix;
 struct FAssetData;
 
 namespace UE::Audio::Insights
 {
-	class FSubmixAssetProvider : public TSharedFromThis<FSubmixAssetProvider>
-	{
-	public:
-		FSubmixAssetProvider();
-		virtual ~FSubmixAssetProvider();
-
-		void RequestEntriesUpdate();
-
-		DECLARE_DELEGATE_OneParam(FOnSubmixAssetAdded, const TSharedPtr<FSubmixDashboardEntry>& /*SubmixDashboardEntry*/);
-		FOnSubmixAssetAdded OnSubmixAssetAdded;
-
-		DECLARE_DELEGATE_OneParam(FOnSubmixAssetRemoved, const uint32 /*SubmixId*/);
-		FOnSubmixAssetRemoved OnSubmixAssetRemoved;
-
-		DECLARE_DELEGATE_OneParam(FOnSubmixAssetListUpdated, const TArray<TSharedPtr<FSubmixDashboardEntry>>& /*SubmixDashboardEntries*/);
-		FOnSubmixAssetListUpdated OnSubmixAssetListUpdated;
+	struct FSubmixDashboardEntry;
 
-	private:
-		void OnAssetAdded(const FAssetData& InAssetData);
-		void OnAssetRemoved(const FAssetData& InAssetData);
-		void OnFilesLoaded();
-
-		void OnActiveAudioDeviceChanged(::Audio::FDeviceId InDeviceId);
-		void OnTraceStarted(FTraceAuxiliary::EConnectionType TraceType, const FString& TraceDestination);
-
-		void AddSubmixAsset(const FAssetData& InAssetData, const bool bInShouldNotify);
-		void RemoveSubmixAsset(const FAssetData& InAssetData);
-
-		void UpdateSubmixAssetNames();
-
-		bool bAreFilesLoaded = false;
-		bool bAssetEntriesNeedRefreshing = false;
-
-		TArray<TSharedPtr<FSubmixDashboardEntry>> SubmixAssetEntries;
+	class FSubmixAssetProvider : public TAssetProvider<USoundSubmix, FSubmixDashboardEntry>
+	{
+	protected:
+		virtual uint32 GetAssetId(const TSharedPtr<FSubmixDashboardEntry>& InEntry) const override;
+		virtual void InitAssetEntry(TSharedPtr<FSubmixDashboardEntry>& InEntry, const FAssetData& InAssetData) override;
 	};
 } // namespace UE::Audio::Insights
 
diff --git a/Engine/Plugins/AudioInsights/Source/AudioInsights/Private/Providers/SubmixTraceProvider.cpp b/Engine/Plugins/AudioInsights/Source/AudioInsights/Private/Providers/SubmixTraceProvider.cpp
--- a/Engine/Plugins/AudioInsights/Source/AudioInsights/Private/Providers/SubmixTraceProvider.cpp
+++ b/Engine/Plugins/AudioInsights/Source/AudioInsights/Private/Providers/SubmixTraceProvider.cpp
@@ -38,18 +38,18 @@ namespace UE::Audio::Insights
 		: TDeviceDataMapTraceProvider<uint32, TSharedPtr<FSubmixDashboardEntry>>(GetName_Static())
 	{
 #if WITH_EDITOR
-		SubmixAssetProvider.OnSubmixAssetAdded.BindRaw(this, &FSubmixTraceProvider::HandleOnSubmixAssetAdded);
-		SubmixAssetProvider.OnSubmixAssetRemoved.BindRaw(this, &FSubmixTraceProvider::HandleOnSubmixAssetRemoved);
-		SubmixAssetProvider.OnSubmixAssetListUpdated.BindRaw(this, &FSubmixTraceProvider::HandleOnSubmixAssetListUpdated);
+		SubmixAssetProvider.OnAssetAdded.BindRaw(this, &FSubmixTraceProvider::HandleOnSubmixAssetAdded);
+		SubmixAssetProvider.OnAssetRemoved.BindRaw(this, &FSubmixTraceProvider::HandleOnSubmixAssetRemoved);
+		SubmixAssetProvider.OnAssetListUpdated.BindRaw(this, &FSubmixTraceProvider::HandleOnSubmixAssetListUpdated);
 #endif // WITH_EDITOR
 	}
 
 	FSubmixTraceProvider::~FSubmixTraceProvider()
 	{
 #if WITH_EDITOR
-		SubmixAssetProvider.OnSubmixAssetAdded.Unbind();
-		SubmixAssetProvider.OnSubmixAssetRemoved.Unbind();
-		SubmixAssetProvider.OnSubmixAssetListUpdated.Unbind();
+		SubmixAssetProvider.OnAssetAdded.Unbind();
+		SubmixAssetProvider.OnAssetRemoved.Unbind();
+		SubmixAssetProvider.OnAssetListUpdated.Unbind();
 #endif // WITH_EDITOR
 	}
 	
@@ -245,6 +245,11 @@ namespace UE::Audio::Insights
 	}
 
 #if WITH_EDITOR
+	void FSubmixTraceProvider::RequestEntriesUpdate()
+	{
+		SubmixAssetProvider.RequestEntriesUpdate();
+	}
+
 	void FSubmixTraceProvider::HandleOnSubmixAssetAdded(const TSharedPtr<FSubmixDashboardEntry>& InSubmixDashboardEntry)
 	{
 		if (const FAudioDeviceManager* DeviceManager = FAudioDeviceManager::Get())
diff --git a/Engine/Plugins/AudioInsights/Source/AudioInsights/Private/Providers/SubmixTraceProvider.h b/Engine/Plugins/AudioInsights/Source/AudioInsights/Private/Providers/SubmixTraceProvider.h
--- a/Engine/Plugins/AudioInsights/Source/AudioInsights/Private/Providers/SubmixTraceProvider.h
+++ b/Engine/Plugins/AudioInsights/Source/AudioInsights/Private/Providers/SubmixTraceProvider.h
@@ -26,6 +26,8 @@ namespace UE::Audio::Insights
 #if !WITH_EDITOR
 		virtual void InitSessionCachedMessages(TraceServices::IAnalysisSession& InSession) override;
 #else
+		void RequestEntriesUpdate();
+
 		DECLARE_DELEGATE(FOnSubmixAssetsUpdated);
 		FOnSubmixAssetsUpdated OnSubmixAssetAdded;
 		FOnSubmixAssetsUpdated OnSubmixAssetRemoved;
diff --git a/Engine/Plugins/AudioInsights/Source/AudioInsights/Private/Views/AudioBusDashboardViewFactory.cpp b/Engine/Plugins/AudioInsights/Source/AudioInsights/Private/Views/AudioBusDashboardViewFactory.cpp
--- a/Engine/Plugins/AudioInsights/Source/AudioInsights/Private/Views/AudioBusDashboardViewFactory.cpp
+++ b/Engine/Plugins/AudioInsights/Source/AudioInsights/Private/Views/AudioBusDashboardViewFactory.cpp
@@ -304,6 +304,15 @@ namespace UE::Audio::Insights
 				FilteredEntriesListView->SetSelectionMode(ESelectionMode::Single);
 			}
 		}
+#if WITH_EDITOR
+		else
+		{
+			if (AudioBusProvider.IsValid())
+			{
+				AudioBusProvider->RequestEntriesUpdate();
+			}
+		}
+#endif // WITH_EDITOR
 
 		// Update audio meters (if any has been checked previously)
 		for (const TSharedPtr<IDashboardDataViewEntry>& Entry : DataViewEntries)
diff --git a/Engine/Plugins/AudioInsights/Source/AudioInsights/Private/Views/SubmixDashboardViewFactory.cpp b/Engine/Plugins/AudioInsights/Source/AudioInsights/Private/Views/SubmixDashboardViewFactory.cpp
--- a/Engine/Plugins/AudioInsights/Source/AudioInsights/Private/Views/SubmixDashboardViewFactory.cpp
+++ b/Engine/Plugins/AudioInsights/Source/AudioInsights/Private/Views/SubmixDashboardViewFactory.cpp
@@ -236,6 +236,15 @@ namespace UE::Audio::Insights
 				FilteredEntriesListView->SetSelectionMode(ESelectionMode::Single);
 			}
 		}
+#if WITH_EDITOR
+		else
+		{
+			if (SubmixProvider.IsValid())
+			{
+				SubmixProvider->RequestEntriesUpdate();
+			}
+		}
+#endif // WITH_EDITOR
 
 		// Update audio meters (if any has been checked previously)
 		for (const TSharedPtr<IDashboardDataViewEntry>& Entry : DataViewEntries)
diff --git a/Engine/Plugins/AudioInsights/Source/AudioInsights/Public/Providers/AssetProvider.h b/Engine/Plugins/AudioInsights/Source/AudioInsights/Public/Providers/AssetProvider.h
new file mode 100644
--- /dev/null
+++ b/Engine/Plugins/AudioInsights/Source/AudioInsights/Public/Providers/AssetProvider.h
@@ -0,0 +1,180 @@
+// Copyright Epic Games, Inc. All Rights Reserved.
+#pragma once
+
+#if WITH_EDITOR
+#include "AssetRegistry/AssetData.h"
+#include "AssetRegistry/AssetRegistryModule.h"
+#include "AudioDeviceManager.h"
+#include "AudioMixerDevice.h"
+#include "Delegates/Delegate.h"
+#include "IAudioInsightsModule.h"
+#include "ProfilingDebugging/TraceAuxiliary.h"
+
+#define UE_API AUDIOINSIGHTS_API
+
+namespace UE::Audio::Insights
+{
+	template <typename TAsset, typename TDashboardEntry>
+	class TAssetProvider
+	{
+	public:
+		UE_API TAssetProvider()
+		{
+			FAssetRegistryModule& AssetRegistryModule = FModuleManager::LoadModuleChecked<FAssetRegistryModule>("AssetRegistry");
+
+			AssetRegistryModule.Get().OnAssetAdded().AddRaw(this, &TAssetProvider::HandleOnAssetAdded);
+			AssetRegistryModule.Get().OnAssetRemoved().AddRaw(this, &TAssetProvider::HandleOnAssetRemoved);
+			AssetRegistryModule.Get().OnFilesLoaded().AddRaw(this, &TAssetProvider::HandleOnFilesLoaded);
+
+			FAudioDeviceManagerDelegates::OnAudioDeviceCreated.AddRaw(this, &TAssetProvider::HandleOnActiveAudioDeviceChanged);
+			FAudioDeviceManagerDelegates::OnAudioDeviceDestroyed.AddRaw(this, &TAssetProvider::HandleOnActiveAudioDeviceChanged);
+
+			FTraceAuxiliary::OnTraceStarted.AddRaw(this, &TAssetProvider::HandleOnTraceStarted);
+		}
+
+		UE_API virtual ~TAssetProvider()
+		{
+			if (FAssetRegistryModule* AssetRegistryModule = FModuleManager::GetModulePtr<FAssetRegistryModule>("AssetRegistry"))
+			{
+				AssetRegistryModule->Get().OnAssetAdded().RemoveAll(this);
+				AssetRegistryModule->Get().OnAssetRemoved().RemoveAll(this);
+				AssetRegistryModule->Get().OnFilesLoaded().RemoveAll(this);
+			}
+
+			FAudioDeviceManagerDelegates::OnAudioDeviceCreated.RemoveAll(this);
+			FAudioDeviceManagerDelegates::OnAudioDeviceDestroyed.RemoveAll(this);
+
+			FTraceAuxiliary::OnTraceStarted.RemoveAll(this);
+		}
+
+		UE_API void RequestEntriesUpdate()
+		{
+			UpdateAssetNames();
+		}
+
+		DECLARE_DELEGATE_OneParam(FOnAssetAdded, const TSharedPtr<TDashboardEntry>& /*DashboardEntry*/);
+		FOnAssetAdded OnAssetAdded;
+
+		DECLARE_DELEGATE_OneParam(FOnAssetRemoved, const uint32 /*AssetId*/);
+		FOnAssetRemoved OnAssetRemoved;
+
+		DECLARE_DELEGATE_OneParam(FOnAssetListUpdated, const TArray<TSharedPtr<TDashboardEntry>>& /*DashboardEntries*/);
+		FOnAssetListUpdated OnAssetListUpdated;
+
+	protected:
+		UE_API virtual uint32 GetAssetId(const TSharedPtr<TDashboardEntry>& InEntry) const = 0;
+		UE_API virtual void InitAssetEntry(TSharedPtr<TDashboardEntry>& InEntry, const FAssetData& InAssetData) = 0;
+
+	private:
+		void UpdateAssetNames()
+		{
+			// Get all assets
+			TArray<FAssetData> AssetDataArray;
+
+			FAssetRegistryModule& AssetRegistryModule = FModuleManager::LoadModuleChecked<FAssetRegistryModule>("AssetRegistry");
+			AssetRegistryModule.Get().GetAssetsByClass(FTopLevelAssetPath(TAsset::StaticClass()), AssetDataArray);
+
+			// Build AssetEntries
+			AssetEntries.Reset();
+
+			constexpr bool bShouldNotify = false;
+			for (const FAssetData& AssetData : AssetDataArray)
+			{
+				AddAsset(AssetData, bShouldNotify);
+			}
+
+			AssetEntries.Sort([](const TSharedPtr<TDashboardEntry>& A, const TSharedPtr<TDashboardEntry>& B)
+			{
+				return A->GetDisplayName().CompareToCaseIgnored(B->GetDisplayName()) < 0;
+			});
+
+			OnAssetListUpdated.ExecuteIfBound(AssetEntries);
+		}
+
+		void AddAsset(const FAssetData& InAssetData, const bool bInShouldNotify)
+		{
+			const bool bIsAssetEntryAlreadAdded = AssetEntries.ContainsByPredicate(
+				[AssetName = InAssetData.GetObjectPathString()](const TSharedPtr<TDashboardEntry>& DashboardEntry)
+				{
+					return DashboardEntry->Name == AssetName;
+				});
+
+			if (!bIsAssetEntryAlreadAdded)
+			{
+				TSharedPtr<TDashboardEntry>& DashboardEntry = AssetEntries.Add_GetRef(MakeShared<TDashboardEntry>());
+
+				InitAssetEntry(DashboardEntry, InAssetData);
+
+				if (bInShouldNotify)
+				{
+					OnAssetAdded.ExecuteIfBound(DashboardEntry);
+				}
+			}
+		}
+
+		void RemoveAsset(const FAssetData& InAssetData)
+		{
+			const TSharedPtr<TDashboardEntry>* FoundDashboardEntry = AssetEntries.FindByPredicate(
+				[AssetName = InAssetData.GetObjectPathString()](const TSharedPtr<TDashboardEntry>& InDashboardEntry)
+				{
+					return InDashboardEntry->Name == AssetName;
+				});
+
+			if (FoundDashboardEntry)
+			{
+				const TSharedPtr<TDashboardEntry> DashboardEntry = *FoundDashboardEntry;
+
+				if (DashboardEntry.IsValid())
+				{
+					const uint32 AssetId = GetAssetId(DashboardEntry);
+
+					AssetEntries.Remove(DashboardEntry);
+
+					OnAssetRemoved.ExecuteIfBound(AssetId);
+				}
+			}
+		}
+
+		void HandleOnAssetAdded(const FAssetData& InAssetData)
+		{
+			if (bAreFilesLoaded && InAssetData.AssetClassPath == FTopLevelAssetPath(TAsset::StaticClass()))
+			{
+				constexpr bool bShouldNotify = true;
+				AddAsset(InAssetData, bShouldNotify);
+			}
+		}
+
+		void HandleOnAssetRemoved(const FAssetData& InAssetData)
+		{
+			if (InAssetData.AssetClassPath == FTopLevelAssetPath(TAsset::StaticClass()))
+			{
+				RemoveAsset(InAssetData);
+			}
+		}
+
+		void HandleOnFilesLoaded()
+		{
+			bAreFilesLoaded = true;
+			UpdateAssetNames();
+		}
+
+		void HandleOnActiveAudioDeviceChanged(::Audio::FDeviceId InDeviceId)
+		{
+			UpdateAssetNames();
+		}
+
+		void HandleOnTraceStarted(FTraceAuxiliary::EConnectionType TraceType, const FString& TraceDestination)
+		{
+			UpdateAssetNames();
+		}
+
+		bool bAreFilesLoaded = false;
+		bool bAssetEntriesNeedRefreshing = false;
+
+		TArray<TSharedPtr<TDashboardEntry>> AssetEntries;
+	};
+} // namespace UE::Audio::Insights
+
+#undef UE_API
+
+#endif // WITH_EDITOR
