[Audio Insights] FAudioBusAssetProvider::InitAssetEntry: added missing nullptr check detected by BugHawk

diff --git a/Engine/Plugins/AudioInsights/Source/AudioInsights/Private/Providers/AudioBusAssetProvider.cpp b/Engine/Plugins/AudioInsights/Source/AudioInsights/Private/Providers/AudioBusAssetProvider.cpp
--- a/Engine/Plugins/AudioInsights/Source/AudioInsights/Private/Providers/AudioBusAssetProvider.cpp
+++ b/Engine/Plugins/AudioInsights/Source/AudioInsights/Private/Providers/AudioBusAssetProvider.cpp
@@ -22,7 +22,7 @@ namespace UE::Audio::Insights
 		InEntry->Timestamp = FPlatformTime::Seconds() - GStartTime;
 		InEntry->AudioBusId = AudioBus ? AudioBus->GetUniqueID() : INDEX_NONE;
 		InEntry->Name = InAssetData.GetObjectPathString();
-		InEntry->AudioMeterInfo->NumChannels = AudioBus->GetNumChannels();
+		InEntry->AudioMeterInfo->NumChannels = AudioBus ? AudioBus->GetNumChannels() : 0;
 		InEntry->AudioBusType = EAudioBusType::AssetBased;
 	}
 } // namespace UE::Audio::Insights
