[Audio Insights] Clean up mute/solo filtered user experience

diff --git a/Engine/Plugins/AudioInsights/Source/AudioInsights/Private/Views/SoundDashboardViewFactory.cpp b/Engine/Plugins/AudioInsights/Source/AudioInsights/Private/Views/SoundDashboardViewFactory.cpp
--- a/Engine/Plugins/AudioInsights/Source/AudioInsights/Private/Views/SoundDashboardViewFactory.cpp
+++ b/Engine/Plugins/AudioInsights/Source/AudioInsights/Private/Views/SoundDashboardViewFactory.cpp
@@ -889,17 +889,18 @@ namespace UE::Audio::Insights
 	{
 		using namespace FSoundDashboardViewFactoryPrivate;
 
+		const bool bIsNotFilteredMuteOrSolo = !(bIsMuteFilteredMode || bIsSoloFilteredMode);
 		TFunction<FSlateColor()> MuteColorFunc = [this]() { return bIsMuteFilteredMode ? FSlateColor(FColor::Green) : FSlateColor::UseForeground(); };
-		TFunction<FText()> MuteTextFunc = [&FilteredList = FilteredEntriesListView]()
+		TFunction<FText()> MuteTextFunc = [bNotMuteSolo = bIsNotFilteredMuteOrSolo, &FilteredList = FilteredEntriesListView]()
 			{
-				return FilteredList.IsValid() && FilteredList->GetNumItemsSelected() > 0 ?
+				return FilteredList.IsValid() && bNotMuteSolo && FilteredList->GetNumItemsSelected() > 0 ?
 					LOCTEXT("SoundDashboard_MuteButtonTextSelect", "Mute Selected") :
 					LOCTEXT("SoundDashboard_MuteButtonTextFilter", "Mute Filtered");
 			};
 		TFunction<FSlateColor()> SoloColorFunc = [this]() { return bIsSoloFilteredMode ? FSlateColor(FColor::Green) : FSlateColor::UseForeground(); };
-		TFunction<FText()> SoloTextFunc = [&FilteredList = FilteredEntriesListView]()
+		TFunction<FText()> SoloTextFunc = [bNotMuteSolo = bIsNotFilteredMuteOrSolo, &FilteredList = FilteredEntriesListView]()
 			{
-				return FilteredList.IsValid() && FilteredList->GetNumItemsSelected() > 0 ?
+				return FilteredList.IsValid() && bNotMuteSolo && FilteredList->GetNumItemsSelected() > 0 ?
 					LOCTEXT("SoundDashboard_SoloButtonTextSelect", "Solo Selected") :
 					LOCTEXT("SoundDashboard_SoloButtonTextFilter", "Solo Filtered");
 			};
@@ -917,7 +918,7 @@ namespace UE::Audio::Insights
 				{
 					if (FilteredEntriesListView.IsValid())
 					{
-						if (FilteredEntriesListView->GetNumItemsSelected() > 0)
+						if (FilteredEntriesListView->GetNumItemsSelected() > 0 && !bIsMuteFilteredMode)
 						{
 							ToggleMuteSoloEntries(FilteredEntriesListView->GetSelectedItems(), EMuteSoloMode::Mute);
 						}
@@ -946,7 +947,7 @@ namespace UE::Audio::Insights
 				{
 					if (FilteredEntriesListView.IsValid())
 					{
-						if (FilteredEntriesListView->GetNumItemsSelected() > 0)
+						if (FilteredEntriesListView->GetNumItemsSelected() > 0 && !bIsSoloFilteredMode)
 						{
 							ToggleMuteSoloEntries(FilteredEntriesListView->GetSelectedItems(), EMuteSoloMode::Solo);
 						}
@@ -971,9 +972,11 @@ namespace UE::Audio::Insights
 				SNew(SButton)
 				.ButtonStyle(&FAppStyle::Get().GetWidgetStyle<FButtonStyle>("SimpleButton"))
 				.ToolTipText(LOCTEXT("SoundsDashboard_ClearMutesAndSolosTooltipText", "Clears all assigned mute/solo states."))
-				.OnClicked_Lambda([]()
+				.OnClicked_Lambda([this]()
 				{
 					ClearMutesAndSolos();
+					bIsMuteFilteredMode = false;
+					bIsSoloFilteredMode = false;
 
 					return FReply::Handled();
 				})
