[Audio Insights] Removed Audio Insights Editor Submix/AudioBus/AudioMeters dashboards in favor of using the ones that work for both standalone and editor

diff --git a/Engine/Plugins/AudioInsights/Source/AudioInsightsEditor/Private/Analyzers/AudioMeterAnalyzer.cpp b/Engine/Plugins/AudioInsights/Source/AudioInsightsEditor/Private/Analyzers/AudioMeterAnalyzer.cpp
deleted file mode 100644
diff --git a/Engine/Plugins/AudioInsights/Source/AudioInsightsEditor/Private/Analyzers/AudioMeterAnalyzer.h b/Engine/Plugins/AudioInsights/Source/AudioInsightsEditor/Private/Analyzers/AudioMeterAnalyzer.h
deleted file mode 100644
diff --git a/Engine/Plugins/AudioInsights/Source/AudioInsightsEditor/Private/Analyzers/AudioMeterSubmixAnalyzer.cpp b/Engine/Plugins/AudioInsights/Source/AudioInsightsEditor/Private/Analyzers/AudioMeterSubmixAnalyzer.cpp
deleted file mode 100644
diff --git a/Engine/Plugins/AudioInsights/Source/AudioInsightsEditor/Private/Analyzers/AudioMeterSubmixAnalyzer.h b/Engine/Plugins/AudioInsights/Source/AudioInsightsEditor/Private/Analyzers/AudioMeterSubmixAnalyzer.h
deleted file mode 100644
diff --git a/Engine/Plugins/AudioInsights/Source/AudioInsightsEditor/Private/AudioInsightsEditorModule.cpp b/Engine/Plugins/AudioInsights/Source/AudioInsightsEditor/Private/AudioInsightsEditorModule.cpp
--- a/Engine/Plugins/AudioInsights/Source/AudioInsightsEditor/Private/AudioInsightsEditorModule.cpp
+++ b/Engine/Plugins/AudioInsights/Source/AudioInsightsEditor/Private/AudioInsightsEditorModule.cpp
@@ -9,13 +9,13 @@
 #include "Templates/SharedPointer.h"
 #include "UObject/NameTypes.h"
 #include "Views/AudioAnalyzerRackDashboardViewFactory.h"
-#include "Views/AudioBusesDashboardViewFactory.h"
+#include "Views/AudioBusDashboardViewFactory.h"
 #include "Views/AudioEventLogDashboardViewFactory.h"
-#include "Views/AudioMetersDashboardViewFactory.h"
+#include "Views/AudioMetersPanelDashboardViewFactory.h"
 #include "Views/LogDashboardViewFactory.h"
 #include "Views/SoundDashboardViewFactory.h"
 #include "Views/SoundPlotsDashboardViewFactory.h"
-#include "Views/SubmixesDashboardViewFactory.h"
+#include "Views/SubmixDashboardViewFactory.h"
 #include "Views/ViewportDashboardViewFactory.h"
 #include "Views/VirtualLoopDashboardViewFactory.h"
 #include "WorkspaceMenuStructure.h"
@@ -45,9 +45,9 @@ namespace UE::Audio::Insights
 			DashboardFactory->RegisterViewFactory(MakeShared<FAudioEventLogDashboardViewFactory>());
 			DashboardFactory->RegisterViewFactory(SoundsDashboard);
 			DashboardFactory->RegisterViewFactory(MakeShared<FVirtualLoopDashboardViewFactory>());
-			DashboardFactory->RegisterViewFactory(MakeShared<FSubmixesDashboardViewFactory>());
-			DashboardFactory->RegisterViewFactory(MakeShared<FAudioBusesDashboardViewFactory>());
-			DashboardFactory->RegisterViewFactory(MakeShared<FAudioMetersDashboardViewFactory>());
+			DashboardFactory->RegisterViewFactory(MakeShared<FSubmixDashboardViewFactory>());
+			DashboardFactory->RegisterViewFactory(MakeShared<FAudioBusDashboardViewFactory>());
+			DashboardFactory->RegisterViewFactory(MakeShared<FAudioMetersPanelDashboardViewFactory>());
 			DashboardFactory->RegisterViewFactory(PlotsDashboard);
 			DashboardFactory->RegisterViewFactory(MakeShared<FAudioAnalyzerRackDashboardViewFactory>());
 
diff --git a/Engine/Plugins/AudioInsights/Source/AudioInsightsEditor/Private/Messages/AudioBusMessages.h b/Engine/Plugins/AudioInsights/Source/AudioInsightsEditor/Private/Messages/AudioBusMessages.h
deleted file mode 100644
diff --git a/Engine/Plugins/AudioInsights/Source/AudioInsightsEditor/Private/Messages/SoundSubmixMessages.h b/Engine/Plugins/AudioInsights/Source/AudioInsightsEditor/Private/Messages/SoundSubmixMessages.h
deleted file mode 100644
diff --git a/Engine/Plugins/AudioInsights/Source/AudioInsightsEditor/Private/Providers/AudioBusProvider.cpp b/Engine/Plugins/AudioInsights/Source/AudioInsightsEditor/Private/Providers/AudioBusProvider.cpp
deleted file mode 100644
diff --git a/Engine/Plugins/AudioInsights/Source/AudioInsightsEditor/Private/Providers/AudioBusProvider.h b/Engine/Plugins/AudioInsights/Source/AudioInsightsEditor/Private/Providers/AudioBusProvider.h
deleted file mode 100644
diff --git a/Engine/Plugins/AudioInsights/Source/AudioInsightsEditor/Private/Providers/SoundSubmixProvider.cpp b/Engine/Plugins/AudioInsights/Source/AudioInsightsEditor/Private/Providers/SoundSubmixProvider.cpp
deleted file mode 100644
diff --git a/Engine/Plugins/AudioInsights/Source/AudioInsightsEditor/Private/Providers/SoundSubmixProvider.h b/Engine/Plugins/AudioInsights/Source/AudioInsightsEditor/Private/Providers/SoundSubmixProvider.h
deleted file mode 100644
diff --git a/Engine/Plugins/AudioInsights/Source/AudioInsightsEditor/Private/Views/AudioAnalyzerRackDashboardViewFactory.cpp b/Engine/Plugins/AudioInsights/Source/AudioInsightsEditor/Private/Views/AudioAnalyzerRackDashboardViewFactory.cpp
--- a/Engine/Plugins/AudioInsights/Source/AudioInsightsEditor/Private/Views/AudioAnalyzerRackDashboardViewFactory.cpp
+++ b/Engine/Plugins/AudioInsights/Source/AudioInsightsEditor/Private/Views/AudioAnalyzerRackDashboardViewFactory.cpp
@@ -7,7 +7,6 @@
 #include "AudioInsightsStyle.h"
 #include "Sound/AudioSettings.h"
 #include "Sound/SoundSubmix.h"
-#include "SubmixesDashboardViewFactory.h"
 #include "UObject/WeakObjectPtr.h"
 #include "Widgets/Colors/SColorBlock.h"
 
@@ -48,7 +47,6 @@ namespace UE::Audio::Insights
 				SubmixAudioAnalyzerRack = MakeShared<FSubmixAudioAnalyzerRack>(MainSubmix);
 
 				FEditorDashboardFactory::OnActiveAudioDeviceChanged.AddSP(this, &FAudioAnalyzerRackDashboardViewFactory::HandleOnActiveAudioDeviceChanged);
-				FSubmixesDashboardViewFactory::OnSubmixSelectionChanged.AddSP(this, &FAudioAnalyzerRackDashboardViewFactory::HandleOnSubmixSelectionChanged);
 			}
 		}
 		
@@ -68,14 +66,6 @@ namespace UE::Audio::Insights
 			SubmixAudioAnalyzerRack->RebuildAudioAnalyzerRack(MainSubmix);
 		}
 	}
-
-	void FAudioAnalyzerRackDashboardViewFactory::HandleOnSubmixSelectionChanged(const TWeakObjectPtr<USoundSubmix> InSoundSubmix)
-	{
-		if (SubmixAudioAnalyzerRack.IsValid())
-		{
-			SubmixAudioAnalyzerRack->RebuildAudioAnalyzerRack(InSoundSubmix);
-		}
-	}
 } // namespace UE::Audio::Insights
 
 #undef LOCTEXT_NAMESPACE
diff --git a/Engine/Plugins/AudioInsights/Source/AudioInsightsEditor/Private/Views/AudioAnalyzerRackDashboardViewFactory.h b/Engine/Plugins/AudioInsights/Source/AudioInsightsEditor/Private/Views/AudioAnalyzerRackDashboardViewFactory.h
--- a/Engine/Plugins/AudioInsights/Source/AudioInsightsEditor/Private/Views/AudioAnalyzerRackDashboardViewFactory.h
+++ b/Engine/Plugins/AudioInsights/Source/AudioInsightsEditor/Private/Views/AudioAnalyzerRackDashboardViewFactory.h
@@ -21,7 +21,6 @@ namespace UE::Audio::Insights
 
 	private:
 		void HandleOnActiveAudioDeviceChanged();
-		void HandleOnSubmixSelectionChanged(const TWeakObjectPtr<USoundSubmix> InSoundSubmix);
 
 		TSharedPtr<FSubmixAudioAnalyzerRack> SubmixAudioAnalyzerRack;
 		TObjectPtr<USoundSubmix> MainSubmix;
diff --git a/Engine/Plugins/AudioInsights/Source/AudioInsightsEditor/Private/Views/AudioBusesDashboardViewFactory.cpp b/Engine/Plugins/AudioInsights/Source/AudioInsightsEditor/Private/Views/AudioBusesDashboardViewFactory.cpp
deleted file mode 100644
diff --git a/Engine/Plugins/AudioInsights/Source/AudioInsightsEditor/Private/Views/AudioBusesDashboardViewFactory.h b/Engine/Plugins/AudioInsights/Source/AudioInsightsEditor/Private/Views/AudioBusesDashboardViewFactory.h
deleted file mode 100644
diff --git a/Engine/Plugins/AudioInsights/Source/AudioInsightsEditor/Private/Views/AudioMeterView.cpp b/Engine/Plugins/AudioInsights/Source/AudioInsightsEditor/Private/Views/AudioMeterView.cpp
deleted file mode 100644
diff --git a/Engine/Plugins/AudioInsights/Source/AudioInsightsEditor/Private/Views/AudioMeterView.h b/Engine/Plugins/AudioInsights/Source/AudioInsightsEditor/Private/Views/AudioMeterView.h
deleted file mode 100644
diff --git a/Engine/Plugins/AudioInsights/Source/AudioInsightsEditor/Private/Views/AudioMetersDashboardViewFactory.cpp b/Engine/Plugins/AudioInsights/Source/AudioInsightsEditor/Private/Views/AudioMetersDashboardViewFactory.cpp
deleted file mode 100644
diff --git a/Engine/Plugins/AudioInsights/Source/AudioInsightsEditor/Private/Views/AudioMetersDashboardViewFactory.h b/Engine/Plugins/AudioInsights/Source/AudioInsightsEditor/Private/Views/AudioMetersDashboardViewFactory.h
deleted file mode 100644
diff --git a/Engine/Plugins/AudioInsights/Source/AudioInsightsEditor/Private/Views/SubmixesDashboardViewFactory.cpp b/Engine/Plugins/AudioInsights/Source/AudioInsightsEditor/Private/Views/SubmixesDashboardViewFactory.cpp
deleted file mode 100644
diff --git a/Engine/Plugins/AudioInsights/Source/AudioInsightsEditor/Private/Views/SubmixesDashboardViewFactory.h b/Engine/Plugins/AudioInsights/Source/AudioInsightsEditor/Private/Views/SubmixesDashboardViewFactory.h
deleted file mode 100644
