Removing inline static delegate definition in AudioInsightsEditorDashboardFactory to enable module validation before using it.

diff --git a/Engine/Plugins/AudioInsights/Source/AudioInsightsEditor/Private/AudioInsightsEditorDashboardFactory.h b/Engine/Plugins/AudioInsights/Source/AudioInsightsEditor/Private/AudioInsightsEditorDashboardFactory.h
--- a/Engine/Plugins/AudioInsights/Source/AudioInsightsEditor/Private/AudioInsightsEditorDashboardFactory.h
+++ b/Engine/Plugins/AudioInsights/Source/AudioInsightsEditor/Private/AudioInsightsEditorDashboardFactory.h
@@ -29,7 +29,7 @@ namespace UE::Audio::Insights
 		virtual ::Audio::FDeviceId GetDeviceId() const override;
 
 		DECLARE_MULTICAST_DELEGATE(FOnActiveAudioDeviceChanged);
-		inline static FOnActiveAudioDeviceChanged OnActiveAudioDeviceChanged;
+		FOnActiveAudioDeviceChanged OnActiveAudioDeviceChanged;
 
 	private:
 		void OnWorldRegisteredToAudioDevice(const UWorld* InWorld, ::Audio::FDeviceId InDeviceId);
diff --git a/Engine/Plugins/AudioInsights/Source/AudioInsightsEditor/Private/AudioInsightsEditorModule.cpp b/Engine/Plugins/AudioInsights/Source/AudioInsightsEditor/Private/AudioInsightsEditorModule.cpp
--- a/Engine/Plugins/AudioInsights/Source/AudioInsightsEditor/Private/AudioInsightsEditorModule.cpp
+++ b/Engine/Plugins/AudioInsights/Source/AudioInsightsEditor/Private/AudioInsightsEditorModule.cpp
@@ -94,14 +94,14 @@ namespace UE::Audio::Insights
 		return InsightsModule.GetTraceModule();
 	}
 
-	TSharedRef<FEditorDashboardFactory> FAudioInsightsEditorModule::GetDashboardFactory()
+	TSharedPtr<FEditorDashboardFactory> FAudioInsightsEditorModule::GetDashboardFactory()
 	{
-		return DashboardFactory->AsShared();
+		return DashboardFactory;
 	}
 
-	const TSharedRef<FEditorDashboardFactory> FAudioInsightsEditorModule::GetDashboardFactory() const
+	const TSharedPtr<FEditorDashboardFactory> FAudioInsightsEditorModule::GetDashboardFactory() const
 	{
-		return DashboardFactory->AsShared();
+		return DashboardFactory;
 	}
 
 	TSharedRef<SDockTab> FAudioInsightsEditorModule::CreateDashboardTabWidget(const FSpawnTabArgs& Args)
diff --git a/Engine/Plugins/AudioInsights/Source/AudioInsightsEditor/Private/AudioInsightsEditorModule.h b/Engine/Plugins/AudioInsights/Source/AudioInsightsEditor/Private/AudioInsightsEditorModule.h
--- a/Engine/Plugins/AudioInsights/Source/AudioInsightsEditor/Private/AudioInsightsEditorModule.h
+++ b/Engine/Plugins/AudioInsights/Source/AudioInsightsEditor/Private/AudioInsightsEditorModule.h
@@ -29,8 +29,8 @@ namespace UE::Audio::Insights
 		virtual void UnregisterDashboardViewFactory(FName InName) override;
 		virtual ::Audio::FDeviceId GetDeviceId() const override;
 
-		TSharedRef<FEditorDashboardFactory> GetDashboardFactory();
-		const TSharedRef<FEditorDashboardFactory> GetDashboardFactory() const;
+		TSharedPtr<FEditorDashboardFactory> GetDashboardFactory();
+		const TSharedPtr<FEditorDashboardFactory> GetDashboardFactory() const;
 
 		static bool IsModuleLoaded();
 		static FAudioInsightsEditorModule& GetChecked();
diff --git a/Engine/Plugins/AudioInsights/Source/AudioInsightsEditor/Private/Views/AudioAnalyzerRackDashboardViewFactory.cpp b/Engine/Plugins/AudioInsights/Source/AudioInsightsEditor/Private/Views/AudioAnalyzerRackDashboardViewFactory.cpp
--- a/Engine/Plugins/AudioInsights/Source/AudioInsightsEditor/Private/Views/AudioAnalyzerRackDashboardViewFactory.cpp
+++ b/Engine/Plugins/AudioInsights/Source/AudioInsightsEditor/Private/Views/AudioAnalyzerRackDashboardViewFactory.cpp
@@ -4,6 +4,7 @@
 
 #include "Analyzers/SubmixAudioAnalyzerRack.h"
 #include "AudioInsightsEditorDashboardFactory.h"
+#include "AudioInsightsEditorModule.h"
 #include "AudioInsightsStyle.h"
 #include "Sound/AudioSettings.h"
 #include "Sound/SoundSubmix.h"
@@ -14,6 +15,15 @@
 
 namespace UE::Audio::Insights
 {
+	FAudioAnalyzerRackDashboardViewFactory::~FAudioAnalyzerRackDashboardViewFactory()
+	{
+		TSharedPtr<FEditorDashboardFactory> DashboardFactory = FAudioInsightsEditorModule::GetChecked().GetDashboardFactory();
+		if (DashboardFactory.IsValid())
+		{
+			DashboardFactory->OnActiveAudioDeviceChanged.RemoveAll(this);
+		}
+	}
+
 	FName FAudioAnalyzerRackDashboardViewFactory::GetName() const
 	{
 		return "AudioAnalyzerRack";
@@ -46,7 +56,12 @@ namespace UE::Audio::Insights
 			{
 				SubmixAudioAnalyzerRack = MakeShared<FSubmixAudioAnalyzerRack>(MainSubmix);
 
-				FEditorDashboardFactory::OnActiveAudioDeviceChanged.AddSP(this, &FAudioAnalyzerRackDashboardViewFactory::HandleOnActiveAudioDeviceChanged);
+				TSharedPtr<FEditorDashboardFactory> DashboardFactory = FAudioInsightsEditorModule::GetChecked().GetDashboardFactory();
+				if (DashboardFactory.IsValid())
+				{
+					DashboardFactory->OnActiveAudioDeviceChanged.RemoveAll(this);
+					DashboardFactory->OnActiveAudioDeviceChanged.AddSP(this, &FAudioAnalyzerRackDashboardViewFactory::HandleOnActiveAudioDeviceChanged);
+				}
 			}
 		}
 		
diff --git a/Engine/Plugins/AudioInsights/Source/AudioInsightsEditor/Private/Views/AudioAnalyzerRackDashboardViewFactory.h b/Engine/Plugins/AudioInsights/Source/AudioInsightsEditor/Private/Views/AudioAnalyzerRackDashboardViewFactory.h
--- a/Engine/Plugins/AudioInsights/Source/AudioInsightsEditor/Private/Views/AudioAnalyzerRackDashboardViewFactory.h
+++ b/Engine/Plugins/AudioInsights/Source/AudioInsightsEditor/Private/Views/AudioAnalyzerRackDashboardViewFactory.h
@@ -13,6 +13,8 @@ namespace UE::Audio::Insights
 	class FAudioAnalyzerRackDashboardViewFactory : public IDashboardViewFactory, public TSharedFromThis<FAudioAnalyzerRackDashboardViewFactory>
 	{
 	public:
+		virtual ~FAudioAnalyzerRackDashboardViewFactory() override;
+
 		virtual FName GetName() const override;
 		virtual FText GetDisplayName() const override;
 		virtual EDefaultDashboardTabStack GetDefaultTabStack() const override;
