[Audio Insights] BugHawk suggestion to let delegates clean themselves up rather than being explicit about it.

diff --git a/Engine/Plugins/AudioInsights/Source/AudioInsightsEditor/Private/Views/AudioAnalyzerRackDashboardViewFactory.cpp b/Engine/Plugins/AudioInsights/Source/AudioInsightsEditor/Private/Views/AudioAnalyzerRackDashboardViewFactory.cpp
--- a/Engine/Plugins/AudioInsights/Source/AudioInsightsEditor/Private/Views/AudioAnalyzerRackDashboardViewFactory.cpp
+++ b/Engine/Plugins/AudioInsights/Source/AudioInsightsEditor/Private/Views/AudioAnalyzerRackDashboardViewFactory.cpp
@@ -15,15 +15,6 @@
 
 namespace UE::Audio::Insights
 {
-	FAudioAnalyzerRackDashboardViewFactory::~FAudioAnalyzerRackDashboardViewFactory()
-	{
-		TSharedPtr<FEditorDashboardFactory> DashboardFactory = FAudioInsightsEditorModule::GetChecked().GetDashboardFactory();
-		if (DashboardFactory.IsValid())
-		{
-			DashboardFactory->OnActiveAudioDeviceChanged.RemoveAll(this);
-		}
-	}
-
 	FName FAudioAnalyzerRackDashboardViewFactory::GetName() const
 	{
 		return "AudioAnalyzerRack";
diff --git a/Engine/Plugins/AudioInsights/Source/AudioInsightsEditor/Private/Views/AudioAnalyzerRackDashboardViewFactory.h b/Engine/Plugins/AudioInsights/Source/AudioInsightsEditor/Private/Views/AudioAnalyzerRackDashboardViewFactory.h
--- a/Engine/Plugins/AudioInsights/Source/AudioInsightsEditor/Private/Views/AudioAnalyzerRackDashboardViewFactory.h
+++ b/Engine/Plugins/AudioInsights/Source/AudioInsightsEditor/Private/Views/AudioAnalyzerRackDashboardViewFactory.h
@@ -13,8 +13,6 @@ namespace UE::Audio::Insights
 	class FAudioAnalyzerRackDashboardViewFactory : public IDashboardViewFactory, public TSharedFromThis<FAudioAnalyzerRackDashboardViewFactory>
 	{
 	public:
-		virtual ~FAudioAnalyzerRackDashboardViewFactory() override;
-
 		virtual FName GetName() const override;
 		virtual FText GetDisplayName() const override;
 		virtual EDefaultDashboardTabStack GetDefaultTabStack() const override;
