[Audio Insights] TAssetProvider: don't call UpdateAssetNames if the engine is getting shutdown.

diff --git a/Engine/Plugins/AudioInsights/Source/AudioInsights/Public/Providers/AssetProvider.h b/Engine/Plugins/AudioInsights/Source/AudioInsights/Public/Providers/AssetProvider.h
--- a/Engine/Plugins/AudioInsights/Source/AudioInsights/Public/Providers/AssetProvider.h
+++ b/Engine/Plugins/AudioInsights/Source/AudioInsights/Public/Providers/AssetProvider.h
@@ -26,8 +26,8 @@ namespace UE::Audio::Insights
 			AssetRegistryModule.Get().OnAssetRemoved().AddRaw(this, &TAssetProvider::HandleOnAssetRemoved);
 			AssetRegistryModule.Get().OnFilesLoaded().AddRaw(this, &TAssetProvider::HandleOnFilesLoaded);
 
-			FAudioDeviceManagerDelegates::OnAudioDeviceCreated.AddRaw(this, &TAssetProvider::HandleOnActiveAudioDeviceChanged);
-			FAudioDeviceManagerDelegates::OnAudioDeviceDestroyed.AddRaw(this, &TAssetProvider::HandleOnActiveAudioDeviceChanged);
+			FAudioDeviceManagerDelegates::OnAudioDeviceCreated.AddRaw(this, &TAssetProvider::HandleOnActiveAudioDeviceCreated);
+			FAudioDeviceManagerDelegates::OnAudioDeviceDestroyed.AddRaw(this, &TAssetProvider::HandleOnActiveAudioDeviceDestroyed);
 
 			FTraceAuxiliary::OnTraceStarted.AddRaw(this, &TAssetProvider::HandleOnTraceStarted);
 		}
@@ -158,11 +158,19 @@ namespace UE::Audio::Insights
 			UpdateAssetNames();
 		}
 
-		void HandleOnActiveAudioDeviceChanged(::Audio::FDeviceId InDeviceId)
+		void HandleOnActiveAudioDeviceCreated(::Audio::FDeviceId InDeviceId)
 		{
 			UpdateAssetNames();
 		}
 
+		void HandleOnActiveAudioDeviceDestroyed(::Audio::FDeviceId InDeviceId)
+		{
+			if (GIsRunning)
+			{
+				UpdateAssetNames();
+			}
+		}
+
 		void HandleOnTraceStarted(FTraceAuxiliary::EConnectionType TraceType, const FString& TraceDestination)
 		{
 			UpdateAssetNames();
