[Audio Insights] FAudioMetersPanelDashboardViewFactory: Fix for crash dereferencing a nullptr

diff --git a/Engine/Plugins/AudioInsights/Source/AudioInsights/Private/Views/AudioMetersPanelDashboardViewFactory.cpp b/Engine/Plugins/AudioInsights/Source/AudioInsights/Private/Views/AudioMetersPanelDashboardViewFactory.cpp
--- a/Engine/Plugins/AudioInsights/Source/AudioInsights/Private/Views/AudioMetersPanelDashboardViewFactory.cpp
+++ b/Engine/Plugins/AudioInsights/Source/AudioInsights/Private/Views/AudioMetersPanelDashboardViewFactory.cpp
@@ -65,12 +65,15 @@ namespace UE::Audio::Insights
 			.NameText(InName)
 			.NameTextColor(InNameColor));
 
-		AudioMeterWidgetsContainer->AddSlot()
-		.AutoWidth()
-		.Padding(10.0f, 0.0f, 10.0f, 0.0f)
-		[
-			AudioMeterView
-		];
+		if (AudioMeterWidgetsContainer.IsValid())
+		{
+			AudioMeterWidgetsContainer->AddSlot()
+			.AutoWidth()
+			.Padding(10.0f, 0.0f, 10.0f, 0.0f)
+			[
+				AudioMeterView
+			];
+		}
 
 		if (AudioMeterWidgetsScrollBox.IsValid())
 		{
@@ -82,7 +85,11 @@ namespace UE::Audio::Insights
 	{
 		if (TSharedRef<SAudioMeterView>* FoundAudioMeterView = AudioMeterViews.Find(InEntryId))
 		{
-			AudioMeterWidgetsContainer->RemoveSlot(*FoundAudioMeterView);
+			if (AudioMeterWidgetsContainer.IsValid())
+			{
+				AudioMeterWidgetsContainer->RemoveSlot(*FoundAudioMeterView);
+			}
+
 			AudioMeterViews.Remove(InEntryId);
 		}
 
