[Audio Insights] FAudioMetersPanelDashboardViewFactory::AddAudioMeterView: Add to AudioMeterWidgetsContainer any AudioMeterView that could have been collected in HandleOnAddAudioMeter while the dashboard was not opened

diff --git a/Engine/Plugins/AudioInsights/Source/AudioInsights/Private/Views/AudioMetersPanelDashboardViewFactory.cpp b/Engine/Plugins/AudioInsights/Source/AudioInsights/Private/Views/AudioMetersPanelDashboardViewFactory.cpp
--- a/Engine/Plugins/AudioInsights/Source/AudioInsights/Private/Views/AudioMetersPanelDashboardViewFactory.cpp
+++ b/Engine/Plugins/AudioInsights/Source/AudioInsights/Private/Views/AudioMetersPanelDashboardViewFactory.cpp
@@ -44,15 +44,36 @@ namespace UE::Audio::Insights
 		return FSlateStyle::Get().CreateIcon("AudioInsights.Icon.Submix");
 	}
 
+	void FAudioMetersPanelDashboardViewFactory::AddAudioMeterView(const TSharedRef<SAudioMeterView>& InAudioMeterView)
+	{
+		if (AudioMeterWidgetsContainer.IsValid())
+		{
+			AudioMeterWidgetsContainer->AddSlot()
+			.AutoWidth()
+			.Padding(10.0f, 0.0f, 10.0f, 0.0f)
+			[
+				InAudioMeterView
+			];
+		}
+	}
+
 	TSharedRef<SWidget> FAudioMetersPanelDashboardViewFactory::MakeWidget(TSharedRef<SDockTab> OwnerTab, const FSpawnTabArgs& SpawnTabArgs)
 	{
 		if (!AudioMeterWidgetsScrollBox.IsValid())
 		{
+			SAssignNew(AudioMeterWidgetsContainer, SHorizontalBox);
+
+			// Add any AudioMeterView that could have been added in HandleOnAddAudioMeter while the dashboard was not opened
+			for (const auto& [EntryId, AudioMeterView] : AudioMeterViews)
+			{
+				AddAudioMeterView(AudioMeterView);
+			}
+
 			SAssignNew(AudioMeterWidgetsScrollBox, SScrollBox)
 			.Orientation(Orient_Horizontal)
 			+ SScrollBox::Slot()
 			[
-				SAssignNew(AudioMeterWidgetsContainer, SHorizontalBox)
+				AudioMeterWidgetsContainer.ToSharedRef()
 			];
 		}
 
@@ -65,15 +86,7 @@ namespace UE::Audio::Insights
 			.NameText(InName)
 			.NameTextColor(InNameColor));
 
-		if (AudioMeterWidgetsContainer.IsValid())
-		{
-			AudioMeterWidgetsContainer->AddSlot()
-			.AutoWidth()
-			.Padding(10.0f, 0.0f, 10.0f, 0.0f)
-			[
-				AudioMeterView
-			];
-		}
+		AddAudioMeterView(AudioMeterView);
 
 		if (AudioMeterWidgetsScrollBox.IsValid())
 		{
diff --git a/Engine/Plugins/AudioInsights/Source/AudioInsights/Public/Views/AudioMetersPanelDashboardViewFactory.h b/Engine/Plugins/AudioInsights/Source/AudioInsights/Public/Views/AudioMetersPanelDashboardViewFactory.h
--- a/Engine/Plugins/AudioInsights/Source/AudioInsights/Public/Views/AudioMetersPanelDashboardViewFactory.h
+++ b/Engine/Plugins/AudioInsights/Source/AudioInsights/Public/Views/AudioMetersPanelDashboardViewFactory.h
@@ -31,6 +31,8 @@ namespace UE::Audio::Insights
 		void HandleOnRemoveAudioMeter(const uint32 InEntryId);
 		void HandleOnUpdateAudioMeterInfo(const uint32 InEntryId, const TSharedRef<FAudioMeterInfo>& InAudioMeterInfo, const FText& InName);
 
+		void AddAudioMeterView(const TSharedRef<SAudioMeterView>& InAudioMeterView);
+
 		TSharedPtr<SScrollBox> AudioMeterWidgetsScrollBox;
 		TSharedPtr<SHorizontalBox> AudioMeterWidgetsContainer;
 		TMap<uint32, TSharedRef<SAudioMeterView>> AudioMeterViews;
