[Audio Insights] Cache Management part 01 - Central caching system

diff --git a/Engine/Plugins/AudioInsights/Source/AudioInsights/Private/AudioInsightsModule.cpp b/Engine/Plugins/AudioInsights/Source/AudioInsights/Private/AudioInsightsModule.cpp
--- a/Engine/Plugins/AudioInsights/Source/AudioInsights/Private/AudioInsightsModule.cpp
+++ b/Engine/Plugins/AudioInsights/Source/AudioInsights/Private/AudioInsightsModule.cpp
@@ -45,6 +45,10 @@ namespace UE::Audio::Insights
 			
 			FDashboardAssetCommands::Register();
 
+#if WITH_EDITOR
+			CacheManager = MakeUnique<FAudioInsightsCacheManager>();
+#endif // WITH_EDITOR
+
 #if !WITH_EDITOR
 			IModularFeatures::Get().RegisterModularFeature(UE::Insights::Timing::TimingViewExtenderFeatureName, &AudioInsightsTimingViewExtender);
 
@@ -103,6 +107,10 @@ namespace UE::Audio::Insights
 
 			IModularFeatures::Get().UnregisterModularFeature(IRewindDebuggerRuntimeExtension::ModularFeatureName, RewindDebuggerExtension.Get());
 			RewindDebuggerExtension.Reset();
+
+#if WITH_EDITOR
+			CacheManager.Reset();
+#endif // WITH_EDITOR
 		}
 	}
 
@@ -151,6 +159,13 @@ namespace UE::Audio::Insights
 		return *TraceModule;
 	}
 
+#if WITH_EDITOR
+	FAudioInsightsCacheManager& FAudioInsightsModule::GetCacheManager()
+	{
+		return *CacheManager;
+	}
+#endif // WITH_EDITOR
+
 	TSharedRef<SDockTab> FAudioInsightsModule::CreateDashboardTabWidget(const FSpawnTabArgs& Args)
 	{
 		return DashboardFactory->MakeDockTabWidget(Args);
diff --git a/Engine/Plugins/AudioInsights/Source/AudioInsights/Private/AudioInsightsModule.h b/Engine/Plugins/AudioInsights/Source/AudioInsights/Private/AudioInsightsModule.h
--- a/Engine/Plugins/AudioInsights/Source/AudioInsights/Private/AudioInsightsModule.h
+++ b/Engine/Plugins/AudioInsights/Source/AudioInsights/Private/AudioInsightsModule.h
@@ -40,11 +40,16 @@ namespace UE::Audio::Insights
 		static FAudioInsightsModule* GetModulePtr();
 		virtual IAudioInsightsTraceModule& GetTraceModule() override;
 
+#if WITH_EDITOR
+		virtual class FAudioInsightsCacheManager& GetCacheManager() override;
+#endif // WITH_EDITOR
+
 #if !WITH_EDITOR
 		TSharedPtr<FAudioInsightsComponent> GetAudioInsightsComponent() { return AudioInsightsComponent; };
 #endif // !WITH_EDITOR
 
 		FAudioInsightsTimingViewExtender& GetTimingViewExtender() { return AudioInsightsTimingViewExtender; };
+		const FAudioInsightsTimingViewExtender& GetTimingViewExtender() const { return AudioInsightsTimingViewExtender; };
 
 		TSharedRef<FDashboardFactory> GetDashboardFactory();
 		const TSharedRef<FDashboardFactory> GetDashboardFactory() const;
@@ -56,6 +61,10 @@ namespace UE::Audio::Insights
 		TUniquePtr<FTraceModule> TraceModule;
 		TUniquePtr<FRewindDebuggerAudioInsightsRuntime> RewindDebuggerExtension;
 
+#if WITH_EDITOR
+		TUniquePtr<class FAudioInsightsCacheManager> CacheManager;
+#endif // WITH_EDITOR
+
 #if !WITH_EDITOR
 		TSharedPtr<FAudioInsightsComponent> AudioInsightsComponent;
 #endif // !WITH_EDITOR
diff --git a/Engine/Plugins/AudioInsights/Source/AudioInsights/Private/AudioInsightsTraceModule.cpp b/Engine/Plugins/AudioInsights/Source/AudioInsights/Private/AudioInsightsTraceModule.cpp
--- a/Engine/Plugins/AudioInsights/Source/AudioInsights/Private/AudioInsightsTraceModule.cpp
+++ b/Engine/Plugins/AudioInsights/Source/AudioInsights/Private/AudioInsightsTraceModule.cpp
@@ -1,6 +1,8 @@
 // Copyright Epic Games, Inc. All Rights Reserved.
 #include "AudioInsightsTraceModule.h"
 
+#include "AudioInsightsModule.h"
+#include "Cache/AudioInsightsCacheManager.h"
 #include "CoreGlobals.h"
 #include "Insights/IInsightsManager.h"
 #include "Insights/IUnrealInsightsModule.h"
@@ -276,6 +278,11 @@ namespace UE::Audio::Insights
 			InSession.AddAnalyzer(Provider->ConstructAnalyzer(InSession));
 		}
 
+#if WITH_EDITOR
+		FAudioInsightsCacheManager& CacheManager = FAudioInsightsModule::GetChecked().GetCacheManager();
+		CacheManager.ClearCache();
+#endif // WITH_EDITOR
+
 		OnAnalysisStarting.Broadcast(FPlatformTime::Seconds());
 	}
 
diff --git a/Engine/Plugins/AudioInsights/Source/AudioInsights/Private/Cache/AudioCachedMessageChunk.cpp b/Engine/Plugins/AudioInsights/Source/AudioInsights/Private/Cache/AudioCachedMessageChunk.cpp
new file mode 100644
--- /dev/null
+++ b/Engine/Plugins/AudioInsights/Source/AudioInsights/Private/Cache/AudioCachedMessageChunk.cpp
@@ -0,0 +1,70 @@
+// Copyright Epic Games, Inc. All Rights Reserved.
+#include "Cache/AudioCachedMessageChunk.h"
+
+namespace UE::Audio::Insights
+{
+	void FAudioCachedMessageChunk::AddMessageToChunk(const TSharedPtr<IAudioCachedMessage> Message)
+    {
+		if (!Message.IsValid())
+		{
+			return;
+		}
+
+		const double MessageTimestamp = Message->Timestamp;
+		if (RangeStartTimestamp == INVALID_TIMESTAMP)
+		{
+			RangeStartTimestamp = MessageTimestamp;
+		}
+
+		if (RangeEndTimestamp < MessageTimestamp)
+		{
+			RangeEndTimestamp = MessageTimestamp;
+		}
+
+		CurrentChunkSize += Message->GetSizeOf();
+
+		TSharedPtr<TraceServices::TPagedArray<TSharedPtr<IAudioCachedMessage>>>& CachedMessages = Messages.FindOrAdd(Message->GetMessageName(), MakeShared<TraceServices::TPagedArray<TSharedPtr<IAudioCachedMessage>>>(Allocator, 4096));
+		CachedMessages->EmplaceBack(Message);
+    }
+
+	void FAudioCachedMessageChunk::ClearChunk()
+	{
+		Messages.Reset();
+		CurrentChunkSize = 0u;
+	}
+
+	uint32 FAudioCachedMessageChunk::GetCurrentChunkSize() const
+	{
+		return CurrentChunkSize;
+	}
+
+	bool FAudioCachedMessageChunk::IsChunkFull() const
+	{
+		return CurrentChunkSize >= MaxChunkSize;
+	}
+
+	bool FAudioCachedMessageChunk::HasAnyData() const
+	{
+		return CurrentChunkSize > 0u;
+	}
+
+	const TMap<FName, TSharedPtr<FAudioCachedMessageChunk::TCachePagedArray>>& FAudioCachedMessageChunk::GetAllChunkMessages() const
+	{
+		return Messages;
+	}
+
+	bool FAudioCachedMessageChunk::TimestampIsInChunkRange(const double Timestamp) const
+	{
+		return FMath::IsWithin(Timestamp, RangeStartTimestamp, RangeEndTimestamp);
+	}
+
+	double FAudioCachedMessageChunk::GetChunkTimeRangeStart() const
+	{
+		return RangeStartTimestamp;
+	}
+
+	double FAudioCachedMessageChunk::GetChunkTimeRangeEnd() const
+	{
+		return RangeEndTimestamp;
+	}
+} // namespace UE::Audio::Insights
diff --git a/Engine/Plugins/AudioInsights/Source/AudioInsights/Private/Cache/AudioInsightsCacheManager.cpp b/Engine/Plugins/AudioInsights/Source/AudioInsights/Private/Cache/AudioInsightsCacheManager.cpp
new file mode 100644
--- /dev/null
+++ b/Engine/Plugins/AudioInsights/Source/AudioInsights/Private/Cache/AudioInsightsCacheManager.cpp
@@ -0,0 +1,87 @@
+// Copyright Epic Games, Inc. All Rights Reserved.
+#include "Cache/AudioInsightsCacheManager.h"
+
+#include "AudioInsightsModule.h"
+
+namespace UE::Audio::Insights
+{
+	namespace CacheManagerPrivate
+	{
+		constexpr uint32 MaxChunkSize = 1 << 20; // 1mB
+		constexpr uint32 NumChunks = 32u; // 32 mB total
+	}
+
+	FAudioInsightsCacheManager::FAudioInsightsCacheManager()
+		: Allocator(8 << 20) // 8mB for storage of pointers to messages
+		, Cache(CacheManagerPrivate::NumChunks)
+	{
+		check(CacheManagerPrivate::NumChunks > 0u);
+
+		CreateCache();
+	}
+
+	void FAudioInsightsCacheManager::AddMessageToCache(const TSharedPtr<IAudioCachedMessage> Message)
+	{
+		check(WriteIndex < Cache.Capacity());
+
+		FAudioCachedMessageChunk& Chunk = *Cache[WriteIndex].Get();
+
+		if (Chunk.IsChunkFull())
+		{
+			// Move to next chunk
+			WriteIndex = Cache.GetNextIndex(WriteIndex);
+
+			FAudioCachedMessageChunk& NextChunk = *Cache[WriteIndex].Get();
+
+			// If this new chunk already has data, clear it before we begin writing
+			if (NextChunk.HasAnyData())
+			{
+				StartIndex = Cache.GetNextIndex(WriteIndex);
+				NextChunk.ClearChunk();
+			}
+
+			NextChunk.AddMessageToChunk(Message);
+		}
+		else
+		{
+			Chunk.AddMessageToChunk(Message);
+		}
+	}
+
+	void FAudioInsightsCacheManager::ClearCache()
+	{
+		CreateCache();
+	}
+
+	float FAudioInsightsCacheManager::GetCacheDuration() const
+	{
+		return static_cast<float>(GetCacheEndTimeStamp() - GetCacheStartTimeStamp());
+	}
+
+	double FAudioInsightsCacheManager::GetCacheStartTimeStamp() const
+	{
+		check(StartIndex < Cache.Capacity());
+		check(Cache[StartIndex].IsValid());
+
+		return Cache[StartIndex]->GetChunkTimeRangeStart();
+	}
+
+	double FAudioInsightsCacheManager::GetCacheEndTimeStamp() const
+	{
+		check(WriteIndex < Cache.Capacity());
+		check(Cache[WriteIndex].IsValid());
+
+		return Cache[WriteIndex]->GetChunkTimeRangeEnd();
+	}
+
+	void FAudioInsightsCacheManager::CreateCache()
+	{
+		for (uint32 Index = 0u; Index < Cache.Capacity(); ++Index)
+		{
+			Cache[Index] = MakeUnique<FAudioCachedMessageChunk>(Allocator, CacheManagerPrivate::MaxChunkSize);
+		}
+
+		WriteIndex = 0u;
+		StartIndex = 0u;
+	}
+} // namespace UE::Audio::Insights
diff --git a/Engine/Plugins/AudioInsights/Source/AudioInsights/Private/IAudioInsightsModule.cpp b/Engine/Plugins/AudioInsights/Source/AudioInsights/Private/IAudioInsightsModule.cpp
--- a/Engine/Plugins/AudioInsights/Source/AudioInsights/Private/IAudioInsightsModule.cpp
+++ b/Engine/Plugins/AudioInsights/Source/AudioInsights/Private/IAudioInsightsModule.cpp
@@ -9,6 +9,14 @@ IAudioInsightsTraceModule& IAudioInsightsModule::GetTraceModule()
 	return AudioInsightsModule.GetTraceModule();
 }
 
+#if WITH_EDITOR
+UE::Audio::Insights::FAudioInsightsCacheManager& IAudioInsightsModule::GetCacheManager()
+{
+	IAudioInsightsModule& AudioInsightsModule = static_cast<IAudioInsightsModule&>(UE::Audio::Insights::FAudioInsightsModule::GetChecked());
+	return AudioInsightsModule.GetCacheManager();
+}
+#endif // WITH_EDITOR
+
 IAudioInsightsModule& IAudioInsightsModule::GetChecked()
 {
 	return static_cast<IAudioInsightsModule&>(UE::Audio::Insights::FAudioInsightsModule::GetChecked());
diff --git a/Engine/Plugins/AudioInsights/Source/AudioInsights/Public/Cache/AudioCachedMessageChunk.h b/Engine/Plugins/AudioInsights/Source/AudioInsights/Public/Cache/AudioCachedMessageChunk.h
new file mode 100644
--- /dev/null
+++ b/Engine/Plugins/AudioInsights/Source/AudioInsights/Public/Cache/AudioCachedMessageChunk.h
@@ -0,0 +1,84 @@
+// Copyright Epic Games, Inc. All Rights Reserved.
+#pragma once
+
+#include "Cache/IAudioCachedMessage.h"
+#include "Common/PagedArray.h"
+#include "TraceServices/Containers/SlabAllocator.h"
+#include "TraceServices/Model/AnalysisSession.h"
+
+#define UE_API AUDIOINSIGHTS_API
+
+namespace UE::Audio::Insights
+{
+	class FAudioCachedMessageChunk
+	{
+	public:
+		using TCachePagedArray = TraceServices::TPagedArray<TSharedPtr<class IAudioCachedMessage>>;
+
+		UE_API FAudioCachedMessageChunk(TraceServices::FSlabAllocator& InAllocator, const uint32 InChunkSize)
+			: Allocator(InAllocator)
+			, MaxChunkSize(InChunkSize)
+		{ }
+
+		virtual ~FAudioCachedMessageChunk() = default;
+
+		UE_API void AddMessageToChunk(const TSharedPtr<class IAudioCachedMessage> Message);
+		UE_API void ClearChunk();
+
+		UE_API uint32 GetCurrentChunkSize() const;
+		UE_API bool IsChunkFull() const;
+		UE_API bool HasAnyData() const;
+		
+		UE_API const TMap<FName, TSharedPtr<TCachePagedArray>>& GetAllChunkMessages() const;
+
+		UE_API bool TimestampIsInChunkRange(const double Timestamp) const;
+		UE_API double GetChunkTimeRangeStart() const;
+		UE_API double GetChunkTimeRangeEnd() const;
+
+		// Finds the closest message to a specific timestamp
+		// Will iterate backwards from param Timestamp until the first message of type MessageID is found
+		// Param ID can be overriden to filter the results further
+		template<TIsCacheableMessage T>
+		T* FindClosestMessage(const FName& MessageID, const double Timestamp, const uint32 ID = INDEX_NONE) const
+		{
+			const TSharedPtr<TCachePagedArray>* MessageCache = Messages.Find(MessageID);
+			if (MessageCache == nullptr || !MessageCache->IsValid() || (*MessageCache)->Num() == 0u)
+			{
+				return nullptr;
+			}
+
+			const TCachePagedArray& MessagesArray = *(MessageCache->Get());
+			const int32 ClosestMessageToTimeStampIndex = TraceServices::PagedArrayAlgo::BinarySearchClosestBy(MessagesArray, Timestamp,
+			[](const TSharedPtr<class IAudioCachedMessage>& Msg) 
+			{
+				return Msg.IsValid() ? Msg->Timestamp : INVALID_TIMESTAMP;
+			});
+
+			// Iterate backwards from TimeMarker until we find the matching ID
+			for (auto MessageItr = MessagesArray.GetIteratorFromItem(ClosestMessageToTimeStampIndex); MessageItr; --MessageItr)
+			{
+				if (ID == INDEX_NONE || (*MessageItr)->GetID() == ID)
+				{
+					return static_cast<T*>((*MessageItr).Get());
+				}
+			}
+
+			return nullptr;
+		}
+
+	private:
+		static constexpr double INVALID_TIMESTAMP = -1.0;
+
+		TMap<FName, TSharedPtr<TCachePagedArray>> Messages;
+
+		double RangeStartTimestamp = INVALID_TIMESTAMP;
+		double RangeEndTimestamp = INVALID_TIMESTAMP;
+
+		TraceServices::FSlabAllocator& Allocator;
+
+		uint32 CurrentChunkSize = 0u;
+		const uint32 MaxChunkSize = 0u;
+	};
+} // namespace UE::Audio::Insights
+
+#undef UE_API
\ No newline at end of file
diff --git a/Engine/Plugins/AudioInsights/Source/AudioInsights/Public/Cache/AudioInsightsCacheManager.h b/Engine/Plugins/AudioInsights/Source/AudioInsights/Public/Cache/AudioInsightsCacheManager.h
new file mode 100644
--- /dev/null
+++ b/Engine/Plugins/AudioInsights/Source/AudioInsights/Public/Cache/AudioInsightsCacheManager.h
@@ -0,0 +1,179 @@
+// Copyright Epic Games, Inc. All Rights Reserved.
+#pragma once 
+
+#include "Cache/AudioCachedMessageChunk.h"
+#include "Cache/IAudioCachedMessage.h"
+#include "Containers/CircularBuffer.h"
+#include "TraceServices/Containers/SlabAllocator.h"
+
+#define UE_API AUDIOINSIGHTS_API
+
+namespace UE::Audio::Insights
+{
+	class FAudioInsightsCacheManager
+	{
+	public:
+		UE_API FAudioInsightsCacheManager();
+		virtual ~FAudioInsightsCacheManager() = default;
+
+		UE_API void AddMessageToCache(const TSharedPtr<IAudioCachedMessage> Message);
+		UE_API void ClearCache();
+
+		UE_API float GetCacheDuration() const;
+		UE_API double GetCacheStartTimeStamp() const;
+		UE_API double GetCacheEndTimeStamp() const;
+
+		// Find the closest message to a timestamp in the cache for a specific message type
+		// Will look backwards from the chosen timestamp until a message is found
+		// @param ID: When overloaded, the search algorithm will only return messages that have a matching ID.
+		//			  If left at INDEX_NONE, the search algorithm will return the closest message irrespective of ID
+		template<TIsCacheableMessage T>
+		const T* FindClosestMessage(const FName& MessageID, const double Timestamp, const uint32 ID = INDEX_NONE) const
+		{
+			uint32 ReadIndex = INDEX_NONE;
+			for (uint32 Index = 0u; Index < Cache.Capacity(); ++Index)
+			{
+				const FAudioCachedMessageChunk& Chunk = *Cache[Index].Get();
+
+				if (Chunk.TimestampIsInChunkRange(Timestamp))
+				{
+					ReadIndex = Index;
+					break;
+				}
+			}
+
+			if (ReadIndex == INDEX_NONE)
+			{
+				// No chunks contain data from the current timestamp - return null
+				return nullptr;
+			}
+
+			// Iterate backwards until we find a matching message
+			for (uint32 Index = ReadIndex; ; )
+			{
+				const FAudioCachedMessageChunk& Chunk = *Cache[Index].Get();
+				if (const T* FoundMessage = Chunk.FindClosestMessage<T>(MessageID, Timestamp, ID))
+				{
+					return FoundMessage;
+				}
+				else if (Index == StartIndex)
+				{
+					break;
+				}
+				else
+				{
+					Index = Cache.GetPreviousIndex(Index);
+				}
+			}
+
+			return nullptr;
+		}
+
+		// Iterates the cache over a chosen timestamp range for a specific message ID
+		// Will stop iterating if the end of the cache is reached
+		// @param ID: When overloaded, the search algorithm will only return messages that have a matching ID.
+		//			  If left at INDEX_NONE, the search algorithm will return the closest message irrespective of ID
+		template<TIsCacheableMessage T>
+		void IterateOverRange(const FName& MessageID, const double Start, const double End, TFunctionRef<void(const T& Message)> OnMessageFunc, const uint32 ID = INDEX_NONE) const
+		{
+			uint32 ReadIndex = StartIndex;
+			double LatestChunkStartTimestamp = -1.0;
+			for (uint32 Index = 0u; Index < Cache.Capacity(); ++Index)
+			{
+				const FAudioCachedMessageChunk& Chunk = *Cache[ReadIndex].Get();
+
+				const double ChunkStartTime = Chunk.GetChunkTimeRangeStart();
+				if (ChunkStartTime > End)
+				{
+					break;
+				}
+
+				if (Chunk.GetChunkTimeRangeEnd() < Start)
+				{
+					if (ChunkStartTime < LatestChunkStartTimestamp)
+					{
+						break;
+					}
+					else
+					{
+						ReadIndex = Cache.GetNextIndex(ReadIndex);
+						continue;
+					}
+				}
+
+				LatestChunkStartTimestamp = ChunkStartTime;
+
+				const TSharedPtr<TraceServices::TPagedArray<TSharedPtr<IAudioCachedMessage>>>* FoundMessages = Chunk.GetAllChunkMessages().Find(MessageID);
+				if (FoundMessages == nullptr || !FoundMessages->IsValid())
+				{
+					ReadIndex = Cache.GetNextIndex(ReadIndex);
+					continue;
+				}
+
+				bool bReachedEndOfRange = false;
+				for (const TSharedPtr<IAudioCachedMessage>& Message : **FoundMessages)
+				{
+					if (!Message.IsValid() || Message->Timestamp < Start)
+					{
+						continue;
+					}
+
+					if (Message->Timestamp > End)
+					{
+						bReachedEndOfRange = true;
+						break;
+					}
+
+					if (ID == INDEX_NONE || Message->GetID() == ID)
+					{
+						OnMessageFunc(static_cast<const T&>(*Message));
+					}
+				}
+
+				if (bReachedEndOfRange)
+				{
+					break;
+				}
+				else
+				{
+					ReadIndex = Cache.GetNextIndex(ReadIndex);
+				}
+			}
+		}
+
+		// Iterates over the cache from the start to a chosen timestamp for a specific message ID
+		// Will stop iterating if the end of the cache is reached
+		template<TIsCacheableMessage T>
+		void IterateTo(const FName& MessageID, const double End, TFunctionRef<void(const T& Message)> OnMessageFunc, const uint32 ID = INDEX_NONE) const
+		{
+			constexpr double Start = 0.0;
+			IterateOverRange<T>(MessageID, Start, End, OnMessageFunc, ID);
+		}
+
+		// Iterates over the cache from a chosen timestamp to the end of the cache for a specific message ID
+		template<TIsCacheableMessage T>
+		void IterateFrom(const FName& MessageID, const double Start, TFunctionRef<void(const T& Message)> OnMessageFunc, const uint32 ID = INDEX_NONE) const
+		{
+			IterateOverRange<T>(MessageID, Start, GetCacheEndTimeStamp(), OnMessageFunc, ID);
+		}
+
+		// Iterates over the entire cache for a specific message ID
+		template<TIsCacheableMessage T>
+		void IterateOverAll(const FName& MessageID, TFunctionRef<void(const T& Message)> OnMessageFunc, const uint32 ID = INDEX_NONE) const
+		{
+			constexpr double Start = 0.0;
+			IterateOverRange<T>(MessageID, Start, GetCacheEndTimeStamp(), OnMessageFunc, ID);
+		}
+
+	private:
+		void CreateCache();
+
+		TraceServices::FSlabAllocator Allocator;
+		TCircularBuffer<TUniquePtr<FAudioCachedMessageChunk>> Cache;
+
+		uint32 WriteIndex = 0u;
+		uint32 StartIndex = 0u;
+	};
+} // namespace UE::Audio::Insights
+
+#undef UE_API
\ No newline at end of file
diff --git a/Engine/Plugins/AudioInsights/Source/AudioInsights/Public/Cache/IAudioCachedMessage.h b/Engine/Plugins/AudioInsights/Source/AudioInsights/Public/Cache/IAudioCachedMessage.h
new file mode 100644
--- /dev/null
+++ b/Engine/Plugins/AudioInsights/Source/AudioInsights/Public/Cache/IAudioCachedMessage.h
@@ -0,0 +1,29 @@
+// Copyright Epic Games, Inc. All Rights Reserved.
+#pragma once 
+
+#include "Templates/SharedPointer.h"
+#include "Trace/Analyzer.h"
+
+#define UE_API AUDIOINSIGHTS_API
+
+namespace UE::Audio::Insights
+{
+	class IAudioCachedMessage
+	{
+	public:
+		virtual ~IAudioCachedMessage() = default;
+
+		virtual uint32 GetSizeOf() const = 0;
+		virtual uint32 GetID() const = 0;
+		virtual const FName GetMessageName() const = 0;
+
+		double Timestamp = 0.0;
+	};
+} // namespace UE::Audio::Insights
+
+// Cacheable messages must inherit from IAudioCachedMessage and have a constructor that takes in a FOnEventContext argument
+template<typename T>
+concept TIsCacheableMessage = std::is_base_of_v<UE::Audio::Insights::IAudioCachedMessage, T> 
+							  && requires (const UE::Trace::IAnalyzer::FOnEventContext& Context) { T(Context); };
+
+#undef UE_API
\ No newline at end of file
diff --git a/Engine/Plugins/AudioInsights/Source/AudioInsights/Public/IAudioInsightsModule.h b/Engine/Plugins/AudioInsights/Source/AudioInsights/Public/IAudioInsightsModule.h
--- a/Engine/Plugins/AudioInsights/Source/AudioInsights/Public/IAudioInsightsModule.h
+++ b/Engine/Plugins/AudioInsights/Source/AudioInsights/Public/IAudioInsightsModule.h
@@ -16,6 +16,7 @@ class IAudioInsightsTraceModule;
 namespace UE::Audio::Insights
 {
 	class IDashboardViewFactory;
+	class FAudioInsightsCacheManager;
 	class FTraceProviderBase;
 } // namespace UE::Audio::Insights
 
@@ -29,6 +30,10 @@ public:
 
 	UE_API virtual IAudioInsightsTraceModule& GetTraceModule() override;
 
+#if WITH_EDITOR
+	UE_API virtual UE::Audio::Insights::FAudioInsightsCacheManager& GetCacheManager() override;
+#endif
+
 	virtual TSharedRef<SDockTab> CreateDashboardTabWidget(const FSpawnTabArgs& Args) = 0;
 
 	static UE_API IAudioInsightsModule& GetChecked();
diff --git a/Engine/Plugins/AudioInsights/Source/AudioInsights/Public/IAudioInsightsModuleInterface.h b/Engine/Plugins/AudioInsights/Source/AudioInsights/Public/IAudioInsightsModuleInterface.h
--- a/Engine/Plugins/AudioInsights/Source/AudioInsights/Public/IAudioInsightsModuleInterface.h
+++ b/Engine/Plugins/AudioInsights/Source/AudioInsights/Public/IAudioInsightsModuleInterface.h
@@ -11,6 +11,7 @@ class IAudioInsightsTraceModule;
 namespace UE::Audio::Insights
 {
 	class IDashboardViewFactory;
+	class FAudioInsightsCacheManager;
 	class FTraceProviderBase;
 } // namespace UE::Audio::Insights
 
@@ -23,4 +24,8 @@ public:
 	virtual ::Audio::FDeviceId GetDeviceId() const = 0;
 
 	virtual IAudioInsightsTraceModule& GetTraceModule() = 0;
+
+#if WITH_EDITOR
+	virtual UE::Audio::Insights::FAudioInsightsCacheManager& GetCacheManager() = 0;
+#endif // WITH_EDITOR
 };
diff --git a/Engine/Plugins/AudioInsights/Source/AudioInsightsEditor/Private/AudioInsightsEditorModule.cpp b/Engine/Plugins/AudioInsights/Source/AudioInsightsEditor/Private/AudioInsightsEditorModule.cpp
--- a/Engine/Plugins/AudioInsights/Source/AudioInsightsEditor/Private/AudioInsightsEditorModule.cpp
+++ b/Engine/Plugins/AudioInsights/Source/AudioInsightsEditor/Private/AudioInsightsEditorModule.cpp
@@ -94,6 +94,12 @@ namespace UE::Audio::Insights
 		return InsightsModule.GetTraceModule();
 	}
 
+	FAudioInsightsCacheManager& FAudioInsightsEditorModule::GetCacheManager()
+	{
+		IAudioInsightsModule& InsightsModule = IAudioInsightsModule::GetChecked();
+		return InsightsModule.GetCacheManager();
+	}
+
 	TSharedPtr<FEditorDashboardFactory> FAudioInsightsEditorModule::GetDashboardFactory()
 	{
 		return DashboardFactory;
diff --git a/Engine/Plugins/AudioInsights/Source/AudioInsightsEditor/Private/AudioInsightsEditorModule.h b/Engine/Plugins/AudioInsights/Source/AudioInsightsEditor/Private/AudioInsightsEditorModule.h
--- a/Engine/Plugins/AudioInsights/Source/AudioInsightsEditor/Private/AudioInsightsEditorModule.h
+++ b/Engine/Plugins/AudioInsights/Source/AudioInsightsEditor/Private/AudioInsightsEditorModule.h
@@ -36,6 +36,10 @@ namespace UE::Audio::Insights
 		static FAudioInsightsEditorModule& GetChecked();
 		virtual IAudioInsightsTraceModule& GetTraceModule() override;
 
+#if WITH_EDITOR
+		virtual class FAudioInsightsCacheManager& GetCacheManager() override;
+#endif // WITH_EDITOR
+
 	private:
 		void RegisterMenus();
 		TSharedRef<SDockTab> CreateDashboardTabWidget(const FSpawnTabArgs& Args);
