[Audio Insights] Cache Management part 02 - Timing authority now tracks whether the timemarker is being controlled by either the latest message time or a manual playhead. Added delegate when this is reset.

diff --git a/Engine/Plugins/AudioInsights/Source/AudioInsights/Private/AudioInsightsTimingViewExtender.cpp b/Engine/Plugins/AudioInsights/Source/AudioInsights/Private/AudioInsightsTimingViewExtender.cpp
--- a/Engine/Plugins/AudioInsights/Source/AudioInsights/Private/AudioInsightsTimingViewExtender.cpp
+++ b/Engine/Plugins/AudioInsights/Source/AudioInsights/Private/AudioInsightsTimingViewExtender.cpp
@@ -33,6 +33,12 @@ namespace UE::Audio::Insights
 		}
 	}
 
+	void FAudioInsightsTimingViewExtender::ResetTimeControlMethod()
+	{
+		TimeControlMethod = ETimeControlMethod::LatestMessageTime;
+		OnTimeControlMethodReset.Broadcast();
+	}
+
 	double FAudioInsightsTimingViewExtender::GetCurrentDurationSeconds() const
 	{
 		return TraceDurationSeconds;
@@ -40,6 +46,8 @@ namespace UE::Audio::Insights
 
 	void FAudioInsightsTimingViewExtender::SetTimeMarker(const double Timestamp)
 	{
+		TimeControlMethod = ETimeControlMethod::Manual;
+
 		// If you set the time marker on the TimingView, it will Broadcast OnTimingViewTimeMarkerChanged. Otherwise we gotta do it ourselves.
 		if (TimingView != nullptr)
 		{
diff --git a/Engine/Plugins/AudioInsights/Source/AudioInsights/Private/AudioInsightsTimingViewExtender.h b/Engine/Plugins/AudioInsights/Source/AudioInsights/Private/AudioInsightsTimingViewExtender.h
--- a/Engine/Plugins/AudioInsights/Source/AudioInsights/Private/AudioInsightsTimingViewExtender.h
+++ b/Engine/Plugins/AudioInsights/Source/AudioInsights/Private/AudioInsightsTimingViewExtender.h
@@ -14,6 +14,12 @@ class IAnalysisSession;
 
 namespace UE::Audio::Insights
 {
+	enum class ETimeControlMethod
+	{
+		LatestMessageTime = 0,
+		Manual
+	};
+
 	class  FAudioInsightsTimingViewExtender : public UE::Insights::Timing::ITimingViewExtender
 	{
 	public:
@@ -22,6 +28,9 @@ namespace UE::Audio::Insights
 		virtual void OnEndSession(UE::Insights::Timing::ITimingViewSession& InSession) override;
 		virtual void Tick(UE::Insights::Timing::ITimingViewSession& InSession, const TraceServices::IAnalysisSession& InAnalysisSession) override;
 
+		ETimeControlMethod GetTimeControlMethod() const { return TimeControlMethod; }
+		void ResetTimeControlMethod();
+
 		double GetCurrentDurationSeconds() const;
 
 		void SetTimeMarker(const double Timestamp);
@@ -29,12 +38,16 @@ namespace UE::Audio::Insights
 		DECLARE_MULTICAST_DELEGATE_OneParam(FOnTimingViewTimeMarkerChanged, double /*TimeMarker*/);
 		FOnTimingViewTimeMarkerChanged OnTimingViewTimeMarkerChanged;
 
+		DECLARE_MULTICAST_DELEGATE(FOnTimeControlMethodReset);
+		FOnTimeControlMethodReset OnTimeControlMethodReset;
+
 	private:
 		void OnTimeMarkerChanged(UE::Insights::Timing::ETimeChangedFlags InFlags, double InTimeMarker);
 
 		const TraceServices::IAnalysisSession* AnalysisSession = nullptr;
 		UE::Insights::Timing::ITimingViewSession* TimingView = nullptr;
 
+		ETimeControlMethod TimeControlMethod = ETimeControlMethod::LatestMessageTime;
 		double TraceDurationSeconds = 0.0;
 	};
 } // namespace UE::Audio::Insights
