[Audio Insights] Fix build - Cache manager

diff --git a/Engine/Plugins/AudioInsights/Source/AudioInsights/Public/IAudioInsightsModuleInterface.h b/Engine/Plugins/AudioInsights/Source/AudioInsights/Public/IAudioInsightsModuleInterface.h
--- a/Engine/Plugins/AudioInsights/Source/AudioInsights/Public/IAudioInsightsModuleInterface.h
+++ b/Engine/Plugins/AudioInsights/Source/AudioInsights/Public/IAudioInsightsModuleInterface.h
@@ -2,6 +2,7 @@
 #pragma once
 
 #include "AudioDefines.h"
+#include "Cache/AudioInsightsCacheManager.h"
 #include "Modules/ModuleInterface.h"
 #include "Templates/SharedPointer.h"
 #include "UObject/NameTypes.h"
@@ -11,7 +12,6 @@ class IAudioInsightsTraceModule;
 namespace UE::Audio::Insights
 {
 	class IDashboardViewFactory;
-	class FAudioInsightsCacheManager;
 	class FTraceProviderBase;
 } // namespace UE::Audio::Insights
 
