[Audio Insights] Cache Management part 03 - Analyzers can now cache messages. Providers can now request the current time control method.

diff --git a/Engine/Plugins/AudioInsights/Source/AudioInsights/Private/AudioInsightsTraceProviderBase.cpp b/Engine/Plugins/AudioInsights/Source/AudioInsights/Private/AudioInsightsTraceProviderBase.cpp
--- a/Engine/Plugins/AudioInsights/Source/AudioInsights/Private/AudioInsightsTraceProviderBase.cpp
+++ b/Engine/Plugins/AudioInsights/Source/AudioInsights/Private/AudioInsightsTraceProviderBase.cpp
@@ -14,6 +14,7 @@ namespace UE::Audio::Insights
 		FAudioInsightsModule& AudioInsightsModule = FAudioInsightsModule::GetChecked();
 		FAudioInsightsTimingViewExtender& AudioInsightsTimingViewExtender = AudioInsightsModule.GetTimingViewExtender();
 		AudioInsightsTimingViewExtender.OnTimingViewTimeMarkerChanged.AddRaw(this, &FTraceProviderBase::OnTimingViewTimeMarkerChanged);
+		AudioInsightsTimingViewExtender.OnTimeControlMethodReset.AddRaw(this, &FTraceProviderBase::OnTimeControlMethodReset);
 	}
 
 	FTraceProviderBase::~FTraceProviderBase()
@@ -23,6 +24,7 @@ namespace UE::Audio::Insights
 			FAudioInsightsModule& AudioInsightsModule = FAudioInsightsModule::GetChecked();
 			FAudioInsightsTimingViewExtender& AudioInsightsTimingViewExtender = AudioInsightsModule.GetTimingViewExtender();
 			AudioInsightsTimingViewExtender.OnTimingViewTimeMarkerChanged.RemoveAll(this);
+			AudioInsightsTimingViewExtender.OnTimeControlMethodReset.RemoveAll(this);
 		}
 	}
 
@@ -56,4 +58,21 @@ namespace UE::Audio::Insights
 		ensureAlwaysMsgf(bEventSucceeded, TEXT("'%s' TraceProvider's Analyzer message with RouteId '%u' event not handled"), *ProviderName, RouteId);
 		return bEventSucceeded;
 	}
+
+	bool FTraceProviderBase::FTraceAnalyzerBase::ShouldIgnoreNewEvents() const
+	{
+		return Provider->GetTimeControlMethod() != ETimeControlMethod::LatestMessageTime;
+	}
+
+	ETimeControlMethod FTraceProviderBase::GetTimeControlMethod() const
+	{
+		const FAudioInsightsModule& AudioInsightsModule = FAudioInsightsModule::GetChecked();
+		return AudioInsightsModule.GetTimingViewExtender().GetTimeControlMethod();
+	}
+
+	void FTraceProviderBase::ResetTimeControlMethod()
+	{
+		FAudioInsightsModule& AudioInsightsModule = FAudioInsightsModule::GetChecked();
+		AudioInsightsModule.GetTimingViewExtender().ResetTimeControlMethod();
+	}
 } // namespace UE::Audio::Insights
diff --git a/Engine/Plugins/AudioInsights/Source/AudioInsights/Public/AudioInsightsTraceProviderBase.h b/Engine/Plugins/AudioInsights/Source/AudioInsights/Public/AudioInsightsTraceProviderBase.h
--- a/Engine/Plugins/AudioInsights/Source/AudioInsights/Public/AudioInsightsTraceProviderBase.h
+++ b/Engine/Plugins/AudioInsights/Source/AudioInsights/Public/AudioInsightsTraceProviderBase.h
@@ -15,11 +15,15 @@
 
 #if WITH_EDITOR
 #include "AudioDeviceManager.h"
+#include "Cache/AudioInsightsCacheManager.h"
+#include "Cache/IAudioCachedMessage.h"
 #include "ProfilingDebugging/TraceAuxiliary.h"
 #endif // WITH_EDITOR
 
 namespace UE::Audio::Insights
 {
+	enum class ETimeControlMethod;
+
 	class FTraceProviderBase : public TraceServices::IProvider, public TraceServices::IEditableProvider
 	{
 	public:
@@ -73,6 +77,10 @@ namespace UE::Audio::Insights
 #endif // !WITH_EDITOR
 
 		virtual void OnTimingViewTimeMarkerChanged(double TimeMarker) { ++LastMessageId; };
+		virtual void OnTimeControlMethodReset() { };
+
+		ETimeControlMethod GetTimeControlMethod() const;
+		void ResetTimeControlMethod();
 
 	protected:
 		class FTraceAnalyzerBase : public Trace::IAnalyzer
@@ -87,12 +95,28 @@ namespace UE::Audio::Insights
 			AUDIOINSIGHTS_API virtual bool OnEventSuccess(uint16 RouteId, EStyle Style, const FOnEventContext& Context);
 			AUDIOINSIGHTS_API virtual bool OnEventFailure(uint16 RouteId, EStyle Style, const FOnEventContext& Context);
 
+			AUDIOINSIGHTS_API bool ShouldIgnoreNewEvents() const;
+
 			template <typename TProviderType>
 			TProviderType& GetProvider()
 			{
 				return *StaticCastSharedRef<TProviderType>(Provider);
 			}
 
+#if WITH_EDITOR
+			template <TIsCacheableMessage T>
+			void CacheMessage(const FOnEventContext& Context, TAnalyzerMessageQueue<T>& OutMessageQueue)
+			{
+				FAudioInsightsCacheManager& CacheManager = IAudioInsightsModule::GetChecked().GetCacheManager();
+
+				TSharedPtr<T> Message = MakeShared<T>(Context);
+				CacheManager.AddMessageToCache(Message);
+
+				T MessageCopy = *Message;
+				OutMessageQueue.Enqueue(MoveTemp(MessageCopy));
+			}
+#endif // WITH_EDITOR
+
 		private:
 			TSharedRef<FTraceProviderBase> Provider;
 		};
