[Audio Insights] Cache management part 04: Integrate event log with the new cache manager

diff --git a/Engine/Plugins/AudioInsights/Source/AudioInsights/Private/Messages/AudioEventLogTraceMessages.cpp b/Engine/Plugins/AudioInsights/Source/AudioInsights/Private/Messages/AudioEventLogTraceMessages.cpp
--- a/Engine/Plugins/AudioInsights/Source/AudioInsights/Private/Messages/AudioEventLogTraceMessages.cpp
+++ b/Engine/Plugins/AudioInsights/Source/AudioInsights/Private/Messages/AudioEventLogTraceMessages.cpp
@@ -45,4 +45,25 @@ namespace UE::Audio::Insights
 			CategoryType = EAudioEventLogSoundCategory::None;
 		}
     }
+
+	uint32 FAudioEventLogMessage::GetSizeOf() const
+	{
+		// Start with the fixed size of the struct
+		uint32 MemorySize = sizeof(FAudioEventLogMessage);
+
+		// Add any dynamically sized members
+		MemorySize += EventName.GetAllocatedSize();
+		MemorySize += AssetPath.GetAllocatedSize();
+		MemorySize += ActorLabel.GetAllocatedSize();
+		MemorySize += ActorIconName.GetAllocatedSize();
+		MemorySize += CategoryName.GetAllocatedSize();
+
+		return MemorySize;
+	}
+
+	const FName FAudioEventLogMessage::GetMessageName() const
+	{
+		static const FLazyName EventLogName = "EventLog";
+		return EventLogName;
+	}
 }
\ No newline at end of file
diff --git a/Engine/Plugins/AudioInsights/Source/AudioInsights/Private/Providers/AudioEventLogTraceProvider.cpp b/Engine/Plugins/AudioInsights/Source/AudioInsights/Private/Providers/AudioEventLogTraceProvider.cpp
--- a/Engine/Plugins/AudioInsights/Source/AudioInsights/Private/Providers/AudioEventLogTraceProvider.cpp
+++ b/Engine/Plugins/AudioInsights/Source/AudioInsights/Private/Providers/AudioEventLogTraceProvider.cpp
@@ -1,6 +1,7 @@
 // Copyright Epic Games, Inc. All Rights Reserved.
 #include "Providers/AudioEventLogTraceProvider.h"
 
+#include "AudioInsightsTimingViewExtender.h"
 #include "Trace/Analyzer.h"
 #include "TraceServices/Model/AnalysisSession.h"
 #include "TraceServices/ModuleService.h"
@@ -41,6 +42,13 @@ namespace UE::Audio::Insights
 			{
 				LLM_SCOPE_BYNAME(TEXT("Insights/FAudioEventLogTraceAnalyzer"));
 
+#if WITH_EDITOR
+				if (ShouldIgnoreNewEvents())
+				{
+					return true;
+				}
+#endif // WITH_EDITOR
+
 				FAudioEventLogMessages& Messages = GetProvider<FAudioEventLogTraceProvider>().TraceMessages;
 				switch (RouteId)
 				{
@@ -48,7 +56,11 @@ namespace UE::Audio::Insights
 					case RouteId_SoundStart:
 					case RouteId_SoundStop:
 					{
+#if WITH_EDITOR
+						CacheMessage<FAudioEventLogMessage>(Context, Messages.AudioEventLogMessages);
+#else
 						Messages.AudioEventLogMessages.Enqueue(FAudioEventLogMessage(Context));
+#endif // WITH_EDITOR
 						break;
 					}
 
@@ -91,6 +103,13 @@ namespace UE::Audio::Insights
 
 	bool FAudioEventLogTraceProvider::ProcessMessages()
 	{
+#if WITH_EDITOR
+		if (GetTimeControlMethod() != ETimeControlMethod::LatestMessageTime)
+		{
+			return true;
+		}
+#endif // WITH_EDITOR
+
 		auto AddEntryFunc = [this](const FAudioEventLogMessage& Msg)
 		{
 			TSharedPtr<FAudioEventLogDashboardEntry>* ToReturn = nullptr;
diff --git a/Engine/Plugins/AudioInsights/Source/AudioInsights/Private/Views/AudioEventLogDashboardViewFactory.cpp b/Engine/Plugins/AudioInsights/Source/AudioInsights/Private/Views/AudioEventLogDashboardViewFactory.cpp
--- a/Engine/Plugins/AudioInsights/Source/AudioInsights/Private/Views/AudioEventLogDashboardViewFactory.cpp
+++ b/Engine/Plugins/AudioInsights/Source/AudioInsights/Private/Views/AudioEventLogDashboardViewFactory.cpp
@@ -534,6 +534,9 @@ namespace UE::Audio::Insights
 		if (!SelectedItem.IsValid())
 		{
 			FilteredEntriesListView->ClearHighlightedItems();
+
+			FAudioInsightsModule& AudioInsightsModule = FAudioInsightsModule::GetChecked();
+			AudioInsightsModule.GetTimingViewExtender().ResetTimeControlMethod();
 			return;
 		}
 
diff --git a/Engine/Plugins/AudioInsights/Source/AudioInsights/Public/Messages/AudioEventLogTraceMessages.h b/Engine/Plugins/AudioInsights/Source/AudioInsights/Public/Messages/AudioEventLogTraceMessages.h
--- a/Engine/Plugins/AudioInsights/Source/AudioInsights/Public/Messages/AudioEventLogTraceMessages.h
+++ b/Engine/Plugins/AudioInsights/Source/AudioInsights/Public/Messages/AudioEventLogTraceMessages.h
@@ -2,6 +2,7 @@
 #pragma once
 
 #include "AudioDefines.h"
+#include "Cache/IAudioCachedMessage.h"
 #include "Messages/AnalyzerMessageQueue.h"
 #include "SoundTraceMessages.h"
 #include "Trace/Analyzer.h"
@@ -24,11 +25,15 @@ namespace UE::Audio::Insights
 		SoundCueTemplate
 	};
 
-	struct FAudioEventLogMessage
+	struct FAudioEventLogMessage : public IAudioCachedMessage
 	{
 		FAudioEventLogMessage() = default;
 		FAudioEventLogMessage(const Trace::IAnalyzer::FOnEventContext& InContext);
 
+		virtual uint32 GetSizeOf() const override;
+		virtual uint32 GetID() const override { return PlayOrder; }
+		virtual const FName GetMessageName() const override;
+
 		::Audio::FDeviceId DeviceId = INDEX_NONE;
 		uint32 MessageID = INDEX_NONE;
 		uint32 PlayOrder = INDEX_NONE;
@@ -40,8 +45,6 @@ namespace UE::Audio::Insights
 		FString CategoryName;
 
 		EAudioEventLogSoundCategory CategoryType;
-
-		double Timestamp = 0.0;
 	};
 
 	class FAudioEventLogDashboardEntry : public FSoundAssetDashboardEntry
