[Audio Insights] Bughawk fixes for cache manager:
* Remove dll import defines where not required
* Update comments to reflect function process
* Make sure GetCacheDuration returns a sensible value when the chunk is empty
* Reset chunk durations on ClearChunk

diff --git a/Engine/Plugins/AudioInsights/Source/AudioInsights/Private/Cache/AudioCachedMessageChunk.cpp b/Engine/Plugins/AudioInsights/Source/AudioInsights/Private/Cache/AudioCachedMessageChunk.cpp
--- a/Engine/Plugins/AudioInsights/Source/AudioInsights/Private/Cache/AudioCachedMessageChunk.cpp
+++ b/Engine/Plugins/AudioInsights/Source/AudioInsights/Private/Cache/AudioCachedMessageChunk.cpp
@@ -31,6 +31,9 @@ namespace UE::Audio::Insights
 	{
 		Messages.Reset();
 		CurrentChunkSize = 0u;
+
+		RangeStartTimestamp = INVALID_TIMESTAMP;
+		RangeEndTimestamp = INVALID_TIMESTAMP;
 	}
 
 	uint32 FAudioCachedMessageChunk::GetCurrentChunkSize() const
diff --git a/Engine/Plugins/AudioInsights/Source/AudioInsights/Private/Cache/AudioInsightsCacheManager.cpp b/Engine/Plugins/AudioInsights/Source/AudioInsights/Private/Cache/AudioInsightsCacheManager.cpp
--- a/Engine/Plugins/AudioInsights/Source/AudioInsights/Private/Cache/AudioInsightsCacheManager.cpp
+++ b/Engine/Plugins/AudioInsights/Source/AudioInsights/Private/Cache/AudioInsightsCacheManager.cpp
@@ -55,7 +55,7 @@ namespace UE::Audio::Insights
 
 	float FAudioInsightsCacheManager::GetCacheDuration() const
 	{
-		return static_cast<float>(GetCacheEndTimeStamp() - GetCacheStartTimeStamp());
+		return FMath::Max(static_cast<float>(GetCacheEndTimeStamp() - GetCacheStartTimeStamp()), 0.0f);
 	}
 
 	double FAudioInsightsCacheManager::GetCacheStartTimeStamp() const
diff --git a/Engine/Plugins/AudioInsights/Source/AudioInsights/Public/Cache/AudioInsightsCacheManager.h b/Engine/Plugins/AudioInsights/Source/AudioInsights/Public/Cache/AudioInsightsCacheManager.h
--- a/Engine/Plugins/AudioInsights/Source/AudioInsights/Public/Cache/AudioInsightsCacheManager.h
+++ b/Engine/Plugins/AudioInsights/Source/AudioInsights/Public/Cache/AudioInsightsCacheManager.h
@@ -72,7 +72,7 @@ namespace UE::Audio::Insights
 		// Iterates the cache over a chosen timestamp range for a specific message ID
 		// Will stop iterating if the end of the cache is reached
 		// @param ID: When overloaded, the search algorithm will only return messages that have a matching ID.
-		//			  If left at INDEX_NONE, the search algorithm will return the closest message irrespective of ID
+		//			  If left at INDEX_NONE, the search algorithm will return all matching messages irrespective of ID
 		template<TIsCacheableMessage T>
 		void IterateOverRange(const FName& MessageID, const double Start, const double End, TFunctionRef<void(const T& Message)> OnMessageFunc, const uint32 ID = INDEX_NONE) const
 		{
diff --git a/Engine/Plugins/AudioInsights/Source/AudioInsights/Public/Cache/IAudioCachedMessage.h b/Engine/Plugins/AudioInsights/Source/AudioInsights/Public/Cache/IAudioCachedMessage.h
--- a/Engine/Plugins/AudioInsights/Source/AudioInsights/Public/Cache/IAudioCachedMessage.h
+++ b/Engine/Plugins/AudioInsights/Source/AudioInsights/Public/Cache/IAudioCachedMessage.h
@@ -4,8 +4,6 @@
 #include "Templates/SharedPointer.h"
 #include "Trace/Analyzer.h"
 
-#define UE_API AUDIOINSIGHTS_API
-
 namespace UE::Audio::Insights
 {
 	class IAudioCachedMessage
@@ -24,6 +22,4 @@ namespace UE::Audio::Insights
 // Cacheable messages must inherit from IAudioCachedMessage and have a constructor that takes in a FOnEventContext argument
 template<typename T>
 concept TIsCacheableMessage = std::is_base_of_v<UE::Audio::Insights::IAudioCachedMessage, T> 
-							  && requires (const UE::Trace::IAnalyzer::FOnEventContext& Context) { T(Context); };
-
-#undef UE_API
\ No newline at end of file
+							  && requires (const UE::Trace::IAnalyzer::FOnEventContext& Context) { T(Context); };
\ No newline at end of file
