[Audio Insights] Cache Management part 05: Integrate the Sounds tab with the new cache manager

diff --git a/Engine/Plugins/AudioInsights/Source/AudioInsights/Private/Messages/SoundTraceMessages.cpp b/Engine/Plugins/AudioInsights/Source/AudioInsights/Private/Messages/SoundTraceMessages.cpp
--- a/Engine/Plugins/AudioInsights/Source/AudioInsights/Private/Messages/SoundTraceMessages.cpp
+++ b/Engine/Plugins/AudioInsights/Source/AudioInsights/Private/Messages/SoundTraceMessages.cpp
@@ -3,6 +3,25 @@
 
 namespace UE::Audio::Insights
 {
+	namespace SoundMessageNames
+	{
+		const FName SoundStart = "SoundStart";
+		const FName SoundIsAlivePing = "SoundIsAlivePing";
+		const FName SoundWaveStart = "SoundWaveStart";
+		const FName SoundWaveIsAlivePing = "SoundWaveIsAlivePing";
+		const FName SoundStop = "SoundStop";
+
+		const FName PriorityParam = "Priority";
+		const FName DistanceParam = "Distance";
+		const FName DistanceAttenuationParam = "DistanceAttenuation";
+		const FName HPFFreqParam = "HPFFreq";
+		const FName LPFFreqParam = "LPFFreq";
+		const FName EnvelopeParam = "Envelope";
+		const FName PitchParam = "Pitch";
+		const FName VolumeParam = "Volume";
+		const FName RelativeRenderCostParam = "RelativeRenderCost";
+	};
+
 	FSoundMessageBase::FSoundMessageBase(const Trace::IAnalyzer::FOnEventContext& InContext)
 	{
 		const Trace::IAnalyzer::FEventData& EventData = InContext.EventData;
@@ -50,9 +69,31 @@ namespace UE::Audio::Insights
 		EventData.GetString("ActorLabel", ActorLabel);
 	}
 
+	uint32 FSoundStartMessage::GetSizeOf() const
+	{
+		// Start with the fixed size of the struct
+		uint32 MemorySize = sizeof(FSoundStartMessage);
+
+		// Add any dynamically sized members
+		MemorySize += Name.GetAllocatedSize();
+		MemorySize += ActorLabel.GetAllocatedSize();
+
+		return MemorySize;
+	}
+
+	uint32 FSoundWaveStartMessage::GetSizeOf() const
+	{
+		return FSoundStartMessage::GetSizeOf() + sizeof(ActiveSoundPlayOrder);
+	}
+
+	uint32 FSoundStopMessage::GetSizeOf() const
+	{
+		return sizeof(FSoundStopMessage);
+	}
+
 	FSoundDashboardEntry::FSoundDashboardEntry()
 	{
-		constexpr uint32 DataPointsCapacity = 256;
+		constexpr uint32 DataPointsCapacity = 256u;
 
 		PriorityDataPoints.SetCapacity(DataPointsCapacity);
 		DistanceDataPoints.SetCapacity(DataPointsCapacity);
@@ -65,10 +106,22 @@ namespace UE::Audio::Insights
 		RelativeRenderCostDataPoints.SetCapacity(DataPointsCapacity);
 	}
 
-#if !WITH_EDITOR
+	void FSoundDashboardEntry::ResetDataBuffers(const uint32 DataPointsCapacity)
+	{
+		PriorityDataPoints.Reset(DataPointsCapacity);
+		DistanceDataPoints.Reset(DataPointsCapacity);
+		DistanceAttenuationDataPoints.Reset(DataPointsCapacity);
+		LPFFreqDataPoints.Reset(DataPointsCapacity);
+		HPFFreqDataPoints.Reset(DataPointsCapacity);
+		AmplitudeDataPoints.Reset(DataPointsCapacity);
+		VolumeDataPoints.Reset(DataPointsCapacity);
+		PitchDataPoints.Reset(DataPointsCapacity);
+		RelativeRenderCostDataPoints.Reset(DataPointsCapacity);
+	}
+
 	void FSoundDashboardEntry::SwitchBuffersToSingleValues()
 	{
-		constexpr uint32 DataPointsCapacity = 1;
+		constexpr uint32 DataPointsCapacity = 1u;
 
 		PriorityDataPoints.Empty();
 		PriorityDataPoints.SetCapacity(DataPointsCapacity);
@@ -97,5 +150,4 @@ namespace UE::Audio::Insights
 		RelativeRenderCostDataPoints.Empty();
 		RelativeRenderCostDataPoints.SetCapacity(DataPointsCapacity);
 	}
-#endif // !WITH_EDITOR
 }
\ No newline at end of file
diff --git a/Engine/Plugins/AudioInsights/Source/AudioInsights/Private/Providers/SoundTraceProvider.cpp b/Engine/Plugins/AudioInsights/Source/AudioInsights/Private/Providers/SoundTraceProvider.cpp
--- a/Engine/Plugins/AudioInsights/Source/AudioInsights/Private/Providers/SoundTraceProvider.cpp
+++ b/Engine/Plugins/AudioInsights/Source/AudioInsights/Private/Providers/SoundTraceProvider.cpp
@@ -1,14 +1,15 @@
 // Copyright Epic Games, Inc. All Rights Reserved.
 #include "Providers/SoundTraceProvider.h"
 
+#include "Async/ParallelFor.h"
 #include "AudioInsightsModule.h"
+#include "Cache/AudioInsightsCacheManager.h"
 #include "Settings/SoundDashboardSettings.h"
 #include "Trace/Analyzer.h"
 #include "TraceServices/Model/AnalysisSession.h"
 
 #if !WITH_EDITOR
 #include "Common/PagedArray.h"
-#include "Async/ParallelFor.h"
 #endif // !WITH_EDITOR
 
 #define LOCTEXT_NAMESPACE "AudioInsights"
@@ -222,86 +223,147 @@
 			{
 				LLM_SCOPE_BYNAME(TEXT("Insights/FSoundTraceAnalyzer"));
 
+#if WITH_EDITOR
+				if (ShouldIgnoreNewEvents())
+				{
+					return true;
+				}
+#endif // WITH_EDITOR
+
 				FSoundMessages& Messages = GetProvider<FSoundTraceProvider>().TraceMessages;
 
 				switch (RouteId)
 				{
 					case RouteId_ActiveSoundStart:
 					{
+#if WITH_EDITOR
+						CacheMessage<FSoundStartMessage>(Context, Messages.ActiveSoundStartMessages);
+#else
 						Messages.ActiveSoundStartMessages.Enqueue(FSoundStartMessage{ Context });
+#endif // WITH_EDITOR
 						break;
 					}
 
 					case RouteId_SoundWaveStart:
 					{
+#if WITH_EDITOR
+						CacheMessage<FSoundWaveStartMessage>(Context, Messages.SoundWaveStartMessages);
+#else
 						Messages.SoundWaveStartMessages.Enqueue(FSoundWaveStartMessage(Context));
+#endif // WITH_EDITOR
 						break;
 					}
 
 					case RouteId_SoundIsAlivePing:
 					{
+#if WITH_EDITOR
+						CacheMessage<FSoundIsAlivePingMessage>(Context, Messages.ActiveSoundIsAlivePingMessages);
+#else
 						Messages.ActiveSoundIsAlivePingMessages.Enqueue(FSoundIsAlivePingMessage(Context));
+#endif // WITH_EDITOR
 						break;
 					}
 
 					case RouteId_SoundWaveIsAlivePing:
 					{
+#if WITH_EDITOR
+						CacheMessage<FSoundWaveIsAlivePingMessage>(Context, Messages.SoundWaveIsAlivePingMessages);
+#else
 						Messages.SoundWaveIsAlivePingMessages.Enqueue(FSoundWaveIsAlivePingMessage(Context));
+#endif // WITH_EDITOR
 						break;
 					}
 
 					case RouteId_Stop:
 					{
+#if WITH_EDITOR
+						CacheMessage<FSoundStopMessage>(Context, Messages.StopMessages);
+#else
 						Messages.StopMessages.Enqueue(FSoundStopMessage{ Context });
+#endif // WITH_EDITOR
 						break;
 					}
 
 					case RouteId_Priority:
 					{
+#if WITH_EDITOR
+						CacheMessage<FSoundPriorityMessage>(Context, Messages.PriorityMessages);
+#else
 						Messages.PriorityMessages.Enqueue(FSoundPriorityMessage{ Context });
+#endif // WITH_EDITOR
 						break;
 					}
 
 					case RouteId_Distance:
 					{
+#if WITH_EDITOR
+						CacheMessage<FSoundDistanceMessage>(Context, Messages.DistanceMessages);
+#else
 						Messages.DistanceMessages.Enqueue(FSoundDistanceMessage{ Context });
+#endif // WITH_EDITOR
 						break;
 					}
 
 					case RouteId_DistanceAttenuation:
 					{
+#if WITH_EDITOR
+						CacheMessage<FSoundDistanceAttenuationMessage>(Context, Messages.DistanceAttenuationMessages);
+#else
 						Messages.DistanceAttenuationMessages.Enqueue(FSoundDistanceAttenuationMessage{ Context });
+#endif // WITH_EDITOR
 						break;
 					}
 
 					case RouteId_Filters:
 					{
+#if WITH_EDITOR
+						CacheMessage<FSoundLPFFreqMessage>(Context, Messages.LPFFreqMessages);
+						CacheMessage<FSoundHPFFreqMessage>(Context, Messages.HPFFreqMessages);
+#else
 						Messages.LPFFreqMessages.Enqueue(FSoundLPFFreqMessage{ Context });
 						Messages.HPFFreqMessages.Enqueue(FSoundHPFFreqMessage{ Context });
+#endif // WITH_EDITOR
+
 						break;
 					}
 
 					case RouteId_Amplitude:
 					{
+#if WITH_EDITOR
+						CacheMessage<FSoundEnvelopeMessage>(Context, Messages.AmplitudeMessages);
+#else
 						Messages.AmplitudeMessages.Enqueue(FSoundEnvelopeMessage{ Context });
+#endif // WITH_EDITOR
 						break;
 					}
 
 					case RouteId_Volume:
 					{
+#if WITH_EDITOR
+						CacheMessage<FSoundVolumeMessage>(Context, Messages.VolumeMessages);
+#else
 						Messages.VolumeMessages.Enqueue(FSoundVolumeMessage{ Context });
+#endif // WITH_EDITOR
 						break;
 					}
 
 					case RouteId_Pitch:
 					{
+#if WITH_EDITOR
+						CacheMessage<FSoundPitchMessage>(Context, Messages.PitchMessages);
+#else
 						Messages.PitchMessages.Enqueue(FSoundPitchMessage{ Context });
+#endif // WITH_EDITOR
 						break;
 					}
 
 					case RouteId_RelativeRenderCost:
 					{
+#if WITH_EDITOR
+						CacheMessage<FSoundRelativeRenderCostMessage>(Context, Messages.RelativeRenderCostMessages);
+#else
 						Messages.RelativeRenderCostMessages.Enqueue(FSoundRelativeRenderCostMessage{ Context });
+#endif // WITH_EDITOR
 						break;
 					}
 
@@ -359,6 +421,11 @@
 
 	bool FSoundTraceProvider::ProcessMessages()
 	{
+		if (GetTimeControlMethod() != ETimeControlMethod::LatestMessageTime)
+		{
+			return true;
+		}
+
 		// Helper lambdas
 		TSharedPtr<FSoundDashboardEntry> ActiveSoundEntryReturn = nullptr;
 		auto GetActiveSoundEntryFromActiveStartMessage = [this, &ActiveSoundEntryReturn](const FSoundStartMessage& Msg) -> TSharedPtr<FSoundDashboardEntry>*
@@ -1057,13 +1124,14 @@
 
 	void FSoundTraceProvider::OnTimingViewTimeMarkerChanged(double TimeMarker)
 	{
-#if !WITH_EDITOR
 		using namespace FSoundTraceProviderPrivate;
 
+#if !WITH_EDITOR
 		if (!SessionCachedMessages.IsValid())
 		{
 			return;
 		}
+#endif // !WITH_EDITOR
 
 		auto InitActiveSoundEntry = [this](const FSoundStartMessage& Msg, FSoundDashboardEntry& OutActiveSoundEntry)
 		{
@@ -1107,35 +1175,43 @@
 
 		TSet<uint32> ActiveSoundIDs;
 
-		// Collect all the start messages registered until this point in time 
-		for (const FSoundStartMessage& StartCachedMessage : SessionCachedMessages->StartCachedMessages)
+		// Collect all the start messages registered until this point in time
+		auto ProcessStartMessage = [this, &ActiveSoundIDs, &InitActiveSoundEntry](const FSoundStartMessage& StartCachedMessage)
 		{
-			if (StartCachedMessage.Timestamp > TimeMarker)
-			{
-				break;
-			}
-
 			TSharedPtr<FSoundDashboardEntry> ActiveSoundEntry = nullptr;
 			GetOrCreateActiveSoundEntry(StartCachedMessage, ActiveSoundEntry);
 
 			if (!ActiveSoundEntry.IsValid())
 			{
-				continue;
+				return;
 			}
 
 			ActiveSoundIDs.Add(StartCachedMessage.PlayOrder);
-
 			InitActiveSoundEntry(StartCachedMessage, *ActiveSoundEntry);
+		};
+
+#if WITH_EDITOR
+		const FAudioInsightsCacheManager& CacheManager = FAudioInsightsModule::GetChecked().GetCacheManager();
+		CacheManager.IterateTo<FSoundStartMessage>(SoundMessageNames::SoundStart, TimeMarker, [&ProcessStartMessage](const FSoundStartMessage& StartCachedMessage)
+		{
+			ProcessStartMessage(StartCachedMessage);
+		});
+#else
+		for (const FSoundStartMessage& StartCachedMessage : SessionCachedMessages->StartCachedMessages)
+		{
+			if (StartCachedMessage.Timestamp > TimeMarker)
+			{
+				break;
+			}
+
+			ProcessStartMessage(StartCachedMessage);
 		}
+#endif // WITH_EDITOR
+			
 
-		// Collect all the is alive ping messages registered until this point in time 
-		for (const FSoundIsAlivePingMessage& IsAlivePingCachedMessage : SessionCachedMessages->SoundIsAlivePingCachedMessages)
+		// Collect all the is alive ping messages registered until this point in time
+		auto ProcessSoundIsAlivePingMessage = [this, &ActiveSoundIDs, &InitActiveSoundEntry](const FSoundIsAlivePingMessage& IsAlivePingCachedMessage)
 		{
-			if (IsAlivePingCachedMessage.Timestamp > TimeMarker)
-			{
-				break;
-			}
-
 			ActiveSoundIDs.Add(IsAlivePingCachedMessage.PlayOrder);
 
 			TSharedPtr<FSoundDashboardEntry> ActiveSoundEntry = nullptr;
@@ -1143,27 +1219,38 @@
 
 			if (!ActiveSoundEntry.IsValid())
 			{
-				continue;
+				return;
 			}
 
 			InitActiveSoundEntry(IsAlivePingCachedMessage, *ActiveSoundEntry);
+		};
+
+#if WITH_EDITOR
+		CacheManager.IterateTo<FSoundIsAlivePingMessage>(SoundMessageNames::SoundIsAlivePing, TimeMarker, [&ProcessSoundIsAlivePingMessage](const FSoundIsAlivePingMessage& IsAlivePingCachedMessage)
+		{
+			ProcessSoundIsAlivePingMessage(IsAlivePingCachedMessage);
+		});
+#else
+		for (const FSoundIsAlivePingMessage& IsAlivePingCachedMessage : SessionCachedMessages->SoundIsAlivePingCachedMessages)
+		{
+			if (IsAlivePingCachedMessage.Timestamp > TimeMarker)
+			{
+				break;
+			}
+
+			ProcessSoundIsAlivePingMessage(IsAlivePingCachedMessage);
 		}
+#endif // WITH_EDITOR
 
 		////////////////////////////////////////
 		// Add all soundwave start messages registered until this point in time 
 		TSharedPtr<FSoundDashboardEntry> ParentSoundEntryReturn = nullptr;
-		for (const FSoundWaveStartMessage& SoundWaveStartCachedMessage : SessionCachedMessages->SoundWaveStartCachedMessages)
+		auto ProcessSoundWaveStartMessage = [this, &ActiveSoundIDs, &ParentSoundEntryReturn, &CreateChildEntry](const FSoundWaveStartMessage& SoundWaveStartCachedMessage)
 		{
-			if (SoundWaveStartCachedMessage.Timestamp > TimeMarker)
-			{
-				break;
-			}
-
-			GetActiveSoundEntryFromIDs(SoundWaveStartCachedMessage.ActiveSoundPlayOrder, ParentSoundEntryReturn);
-
 			if (ParentSoundEntryReturn == nullptr || !ParentSoundEntryReturn->IsValid() || !FSoundTraceProviderPrivate::EntryTypeHasSoundWaveEntries(ParentSoundEntryReturn->EntryType))
 			{
-				continue;
+				GetActiveSoundEntryFromIDs(SoundWaveStartCachedMessage.ActiveSoundPlayOrder, ParentSoundEntryReturn);
+				return;
 			}
 
 			ActiveSoundIDs.Add(SoundWaveStartCachedMessage.PlayOrder);
@@ -1175,27 +1262,39 @@
 
 			if (bEntryExists)
 			{
-				continue;
+				return;
 			}
 
 			CreateChildEntry(SoundWaveStartCachedMessage, ParentSoundEntryReturn);
+		};
+
+#if WITH_EDITOR
+		CacheManager.IterateTo<FSoundWaveStartMessage>(SoundMessageNames::SoundWaveStart, TimeMarker, [&ProcessSoundWaveStartMessage](const FSoundWaveStartMessage& SoundWaveStartCachedMessage)
+		{
+			ProcessSoundWaveStartMessage(SoundWaveStartCachedMessage);
+		});
+#else
+		for (const FSoundWaveStartMessage& SoundWaveStartCachedMessage : SessionCachedMessages->SoundWaveStartCachedMessages)
+		{
+			if (SoundWaveStartCachedMessage.Timestamp > TimeMarker)
+			{
+				break;
+			}
+
+			ProcessSoundWaveStartMessage(SoundWaveStartCachedMessage);
 		}
+#endif // WITH_EDITOR
 
 		////////////////////////////////////////
 		// Add all soundwave ping messages registered until this point in time 
 		TSharedPtr<FSoundDashboardEntry> ParentSoundIsAlivePingEntryReturn = nullptr;
-		for (const FSoundWaveIsAlivePingMessage& SoundWaveIsAlivePingCachedMessage : SessionCachedMessages->SoundWaveIsAlivePingCachedMessages)
+		auto ProcessSoundWaveIsAlivePingMessage = [this, &ActiveSoundIDs, &ParentSoundIsAlivePingEntryReturn, &CreateChildEntry](const FSoundWaveIsAlivePingMessage& SoundWaveIsAlivePingCachedMessage)
 		{
-			if (SoundWaveIsAlivePingCachedMessage.Timestamp > TimeMarker)
-			{
-				break;
-			}
-
 			GetActiveSoundEntryFromIDs(SoundWaveIsAlivePingCachedMessage.ActiveSoundPlayOrder, ParentSoundIsAlivePingEntryReturn);
 
 			if (ParentSoundIsAlivePingEntryReturn == nullptr || !ParentSoundIsAlivePingEntryReturn->IsValid() || !FSoundTraceProviderPrivate::EntryTypeHasSoundWaveEntries(ParentSoundIsAlivePingEntryReturn->EntryType))
 			{
-				continue;
+				return;
 			}
 
 			ActiveSoundIDs.Add(SoundWaveIsAlivePingCachedMessage.PlayOrder);
@@ -1208,21 +1307,33 @@
 
 			if (bEntryExists)
 			{
-				continue;
+				return;
 			}
 
 			CreateChildEntry(SoundWaveIsAlivePingCachedMessage, ParentSoundIsAlivePingEntryReturn);
+		};
+
+#if WITH_EDITOR
+		CacheManager.IterateTo<FSoundWaveIsAlivePingMessage>(SoundMessageNames::SoundWaveIsAlivePing, TimeMarker, [&ProcessSoundWaveIsAlivePingMessage](const FSoundWaveIsAlivePingMessage& SoundWaveIsAlivePingCachedMessage)
+		{
+			ProcessSoundWaveIsAlivePingMessage(SoundWaveIsAlivePingCachedMessage);
+		});
+#else
+		for (const FSoundWaveIsAlivePingMessage& SoundWaveIsAlivePingCachedMessage : SessionCachedMessages->SoundWaveIsAlivePingCachedMessages)
+		{
+			if (SoundWaveIsAlivePingCachedMessage.Timestamp > TimeMarker)
+			{
+				break;
+			}
+
+			ProcessSoundWaveIsAlivePingMessage(SoundWaveIsAlivePingCachedMessage);
 		}
+#endif // WITH_EDITOR
 
 		// Selectively remove start messages collected in the step above by knowing which sounds were stopped.
 		// With this we will know what are the active sounds at this point in time.
-		for (const FSoundStopMessage& StopCachedMessage : SessionCachedMessages->StopCachedMessages)
+		auto ProcessSoundStopMessageMessage = [this](const FSoundStopMessage& StopCachedMessage)
 		{
-			if (StopCachedMessage.Timestamp > TimeMarker)
-			{
-				break;
-			}
-
 			TSharedPtr<FSoundDashboardEntry> ActiveSoundEntry = nullptr;
 			GetActiveSoundEntryFromIDs(StopCachedMessage.PlayOrder, ActiveSoundEntry);
 
@@ -1237,20 +1348,36 @@
 					RemoveActiveSoundEntry(ActiveSoundEntry);
 				}
 			}
+		};
+#if WITH_EDITOR
+		CacheManager.IterateTo<FSoundStopMessage>(SoundMessageNames::SoundStop, TimeMarker, [&ProcessSoundStopMessageMessage](const FSoundStopMessage& StopCachedMessage)
+		{
+			ProcessSoundStopMessageMessage(StopCachedMessage);
+		});
+#else
+		for (const FSoundStopMessage& StopCachedMessage : SessionCachedMessages->StopCachedMessages)
+		{
+			if (StopCachedMessage.Timestamp > TimeMarker)
+			{
+				break;
+			}
+
+			ProcessSoundStopMessageMessage(StopCachedMessage);
 		}
+#endif // WITH_EDITOR
 
-		TArray<uint32> PlayOrderArray;
-		ActiveSoundToEntryKeysMap.GetKeys(PlayOrderArray);
+		TArray<uint32> ActiveSoundPlayOrderArray;
+		ActiveSoundToEntryKeysMap.GetKeys(ActiveSoundPlayOrderArray);
 
 		// Iterate through the remaining active sounds we have in the dashboard - check whether we detected a start or IsAlive ping message for that entry
 		// If not, we assume this entry is no longer playing and remove it from the dashboard.
-		for (int ReverseIndex = PlayOrderArray.Num() - 1; ReverseIndex >= 0; --ReverseIndex)
+		for (int ReverseIndex = ActiveSoundPlayOrderArray.Num() - 1; ReverseIndex >= 0; --ReverseIndex)
 		{
-			const uint32 PlayOrder = PlayOrderArray[ReverseIndex];
-			if (!ActiveSoundIDs.Contains(PlayOrder))
+			const uint32 ActiveSoundPlayOrder = ActiveSoundPlayOrderArray[ReverseIndex];
+			if (!ActiveSoundIDs.Contains(ActiveSoundPlayOrder))
 			{
 				TSharedPtr<FSoundDashboardEntry> ActiveSoundEntry;
-				GetActiveSoundEntryFromIDs(PlayOrder, ActiveSoundEntry);
+				GetActiveSoundEntryFromIDs(ActiveSoundPlayOrder, ActiveSoundEntry);
 
 				if (ActiveSoundEntry.IsValid())
 				{
@@ -1259,6 +1386,194 @@
 			}
 		}
 		
+		// Regenerate the active sounds play order array
+		ActiveSoundToEntryKeysMap.GetKeys(ActiveSoundPlayOrderArray);
+
+#if WITH_EDITOR
+		CollectParamsForTimestamp(CacheManager, ActiveSoundPlayOrderArray, TimeMarker);
+#else
+		CollectParamsForTimestampStandalone(ActiveSoundPlayOrderArray, TimeMarker);
+#endif //!WITH_EDITOR
+
+		// Call parent method to update LastMessageId
+		FTraceProviderBase::OnTimingViewTimeMarkerChanged(TimeMarker);
+	}
+
+	void FSoundTraceProvider::OnTimeControlMethodReset()
+	{
+		Reset();
+	}
+
+#if WITH_EDITOR
+	void FSoundTraceProvider::CollectParamsForTimestamp(const FAudioInsightsCacheManager& CacheManager, const TArray<uint32>& ActiveSoundPlayOrderArray, const double TimeMarker)
+	{
+		TMap<uint32, TSharedPtr<FSoundDashboardEntry>> FoundEntries;
+		FoundEntries.Reserve(ActiveSoundPlayOrderArray.Num());
+
+		auto GetEntry = [this, &FoundEntries](const uint32 InActiveSoundPlayOrderID)
+		{
+			if (FoundEntries.Contains(InActiveSoundPlayOrderID))
+			{
+				return;
+			}
+
+			TSharedPtr<FSoundDashboardEntry> ActiveSoundEntry = nullptr;
+			GetActiveSoundEntryFromIDs(InActiveSoundPlayOrderID, ActiveSoundEntry);
+
+			if (ActiveSoundEntry.IsValid())
+			{
+				constexpr uint32 DataPointsCapacity = 512u;
+				ActiveSoundEntry->ResetDataBuffers(DataPointsCapacity);
+
+				FoundEntries.Add(InActiveSoundPlayOrderID, ActiveSoundEntry);
+			}
+		};
+
+		for (const uint32 PlayOrder : ActiveSoundPlayOrderArray)
+		{
+			GetEntry(PlayOrder);
+		}
+
+		enum class ESoundParamMessageType : uint8
+		{
+			Priority = 0u,
+			Distance,
+			DistanceAttenuation,
+			LPFFreq,
+			HPFFreq,
+			Envelope,
+			Volume,
+			Pitch,
+			RelativeRenderCost,
+
+			/////
+			MAX
+		};
+
+		constexpr uint32 NumMessageTypesToProcess = static_cast<uint32>(ESoundParamMessageType::MAX);
+
+		// Using ParallelFor to speed-up the cached messages retrieval, using a traditional for loop is unacceptably slower, specially in large traces.
+		ParallelFor(NumMessageTypesToProcess,
+		[&FoundEntries, &CacheManager, TimeMarker, this](const int32 Index)
+		{
+			// 5s is the current window size of the plots widget
+			constexpr double LookbackTime = 5.0;
+			constexpr double Margin = 0.2;
+			const double StartTime = FMath::Max(0.0f, TimeMarker - LookbackTime - Margin);
+
+			switch (static_cast<ESoundParamMessageType>(Index))
+			{
+			case ESoundParamMessageType::Priority:
+				CacheManager.IterateOverRange<FSoundPriorityMessage>(SoundMessageNames::PriorityParam, StartTime, TimeMarker, [&FoundEntries](const FSoundPriorityMessage& Message)
+				{
+					if (TSharedPtr<FSoundDashboardEntry>* SoundEntry = FoundEntries.Find(Message.ActiveSoundPlayOrder))
+					{
+						(*SoundEntry)->PriorityDataPoints.Push({ Message.Timestamp, Message.Priority });
+						(*SoundEntry)->Timestamp = Message.Timestamp;
+					}
+				});
+				break;
+
+			case ESoundParamMessageType::Distance:
+				CacheManager.IterateOverRange<FSoundDistanceMessage>(SoundMessageNames::DistanceParam, StartTime, TimeMarker, [&FoundEntries](const FSoundDistanceMessage& Message)
+				{
+					if (TSharedPtr<FSoundDashboardEntry>* SoundEntry = FoundEntries.Find(Message.ActiveSoundPlayOrder))
+					{
+						(*SoundEntry)->DistanceDataPoints.Push({ Message.Timestamp, Message.Distance });
+						(*SoundEntry)->Timestamp = Message.Timestamp;
+					}
+				});
+				break;
+
+			case ESoundParamMessageType::DistanceAttenuation:
+				CacheManager.IterateOverRange<FSoundDistanceAttenuationMessage>(SoundMessageNames::DistanceAttenuationParam, StartTime, TimeMarker, [&FoundEntries](const FSoundDistanceAttenuationMessage& Message)
+				{
+					if (TSharedPtr<FSoundDashboardEntry>* SoundEntry = FoundEntries.Find(Message.ActiveSoundPlayOrder))
+					{
+						(*SoundEntry)->DistanceAttenuationDataPoints.Push({ Message.Timestamp, Message.DistanceAttenuation });
+						(*SoundEntry)->Timestamp = Message.Timestamp;
+					}
+				});
+				break;
+
+			case ESoundParamMessageType::LPFFreq:
+				CacheManager.IterateOverRange<FSoundLPFFreqMessage>(SoundMessageNames::LPFFreqParam, StartTime, TimeMarker, [&FoundEntries](const FSoundLPFFreqMessage& Message)
+				{
+					if (TSharedPtr<FSoundDashboardEntry>* SoundEntry = FoundEntries.Find(Message.ActiveSoundPlayOrder))
+					{
+						(*SoundEntry)->LPFFreqDataPoints.Push({ Message.Timestamp, Message.LPFFrequency });
+						(*SoundEntry)->Timestamp = Message.Timestamp;
+					}
+				});
+				break;
+
+			case ESoundParamMessageType::HPFFreq:
+				CacheManager.IterateOverRange<FSoundHPFFreqMessage>(SoundMessageNames::HPFFreqParam, StartTime, TimeMarker, [&FoundEntries](const FSoundHPFFreqMessage& Message)
+				{
+					if (TSharedPtr<FSoundDashboardEntry>* SoundEntry = FoundEntries.Find(Message.ActiveSoundPlayOrder))
+					{
+						(*SoundEntry)->HPFFreqDataPoints.Push({ Message.Timestamp, Message.HPFFrequency });
+						(*SoundEntry)->Timestamp = Message.Timestamp;
+					}
+				});
+				break;
+
+			case ESoundParamMessageType::Envelope:
+				CacheManager.IterateOverRange<FSoundEnvelopeMessage>(SoundMessageNames::EnvelopeParam, StartTime, TimeMarker, [&FoundEntries](const FSoundEnvelopeMessage& Message)
+				{
+					if (TSharedPtr<FSoundDashboardEntry>* SoundEntry = FoundEntries.Find(Message.ActiveSoundPlayOrder))
+					{
+						(*SoundEntry)->AmplitudeDataPoints.Push({ Message.Timestamp, Message.Envelope });
+						(*SoundEntry)->Timestamp = Message.Timestamp;
+					}
+				});
+				break;
+
+			case ESoundParamMessageType::Volume:
+				CacheManager.IterateOverRange<FSoundVolumeMessage>(SoundMessageNames::VolumeParam, StartTime, TimeMarker, [&FoundEntries](const FSoundVolumeMessage& Message)
+				{
+					if (TSharedPtr<FSoundDashboardEntry>* SoundEntry = FoundEntries.Find(Message.ActiveSoundPlayOrder))
+					{
+						(*SoundEntry)->VolumeDataPoints.Push({ Message.Timestamp, Message.Volume });
+						(*SoundEntry)->Timestamp = Message.Timestamp;
+					}
+				});
+				break;
+
+			case ESoundParamMessageType::Pitch:
+				CacheManager.IterateOverRange<FSoundPitchMessage>(SoundMessageNames::PitchParam, StartTime, TimeMarker, [&FoundEntries](const FSoundPitchMessage& Message)
+				{
+					if (TSharedPtr<FSoundDashboardEntry>* SoundEntry = FoundEntries.Find(Message.ActiveSoundPlayOrder))
+					{
+						(*SoundEntry)->PitchDataPoints.Push({ Message.Timestamp, Message.Pitch });
+						(*SoundEntry)->Timestamp = Message.Timestamp;
+					}
+				});
+				break;
+
+			case ESoundParamMessageType::RelativeRenderCost:
+				CacheManager.IterateOverRange<FSoundRelativeRenderCostMessage>(SoundMessageNames::RelativeRenderCostParam, StartTime, TimeMarker, [&FoundEntries](const FSoundRelativeRenderCostMessage& Message)
+				{
+					if (TSharedPtr<FSoundDashboardEntry>* SoundEntry = FoundEntries.Find(Message.ActiveSoundPlayOrder))
+					{
+						(*SoundEntry)->RelativeRenderCostDataPoints.Push({ Message.Timestamp, Message.RelativeRenderCost });
+						(*SoundEntry)->Timestamp = Message.Timestamp;
+					}
+				});
+				break;
+
+			default:
+				break;
+			}
+		});
+
+		OnProcessPlotData.Broadcast();
+	}
+#else
+	void FSoundTraceProvider::CollectParamsForTimestampStandalone(const TArray<uint32>& ActiveSoundPlayOrderArray, const double TimeMarker)
+	{
+		using namespace FSoundTraceProviderPrivate;
+
 		const FDeviceData* DeviceData = FindFilteredDeviceData();
 		if (DeviceData)
 		{
@@ -1276,75 +1591,73 @@
 				FSoundRelativeRenderCostMessage RelativeRenderCostMessage;
 			};
 
-			ActiveSoundToEntryKeysMap.GetKeys(PlayOrderArray);
-
 			TArray<CachedEntryInfo> CachedEntryInfos;
-			CachedEntryInfos.SetNumUninitialized(PlayOrderArray.Num());
+			CachedEntryInfos.SetNumUninitialized(ActiveSoundPlayOrderArray.Num());
 
 			// Using ParallelFor to speed-up the cached messages retrieval, using a traditional for loop is unacceptably slower, specially in large traces.
-			ParallelFor(PlayOrderArray.Num(), 
-			[&PlayOrderArray, &CachedEntryInfos, TimeMarker, this](const int32 Index)
+			ParallelFor(ActiveSoundPlayOrderArray.Num(),
+			[&ActiveSoundPlayOrderArray, &CachedEntryInfos, TimeMarker, this](const int32 Index)
 			{
-				const uint32 PlayOrder = PlayOrderArray[Index];
+				const uint32 ActiveSoundPlayOrder = ActiveSoundPlayOrderArray[Index];
 
 				// Priority
-				const FSoundPriorityMessage* FoundPriorityCachedMessage = FindClosestMessageToTimestamp(SessionCachedMessages->PriorityCachedMessages, TimeMarker, PlayOrder);
+				const FSoundPriorityMessage* FoundPriorityCachedMessage = FindClosestMessageToTimestamp(SessionCachedMessages->PriorityCachedMessages, TimeMarker, ActiveSoundPlayOrder);
 				if (FoundPriorityCachedMessage)
 				{
 					CachedEntryInfos[Index].PriorityMessage = *FoundPriorityCachedMessage;
 				}
 
 				// Distance
-				const FSoundDistanceMessage* FoundDistanceCachedMessage = FindClosestMessageToTimestamp(SessionCachedMessages->DistanceCachedMessages, TimeMarker, PlayOrder);
+				const FSoundDistanceMessage* FoundDistanceCachedMessage = FindClosestMessageToTimestamp(SessionCachedMessages->DistanceCachedMessages, TimeMarker, ActiveSoundPlayOrder);
 				if (FoundDistanceCachedMessage)
 				{
 					CachedEntryInfos[Index].DistanceMessage = *FoundDistanceCachedMessage;
 				}
 
 				// Distance Attenuation
-				const FSoundDistanceAttenuationMessage* FoundDistanceAttenuationCachedMessage = FindClosestMessageToTimestamp(SessionCachedMessages->DistanceAttenuationCachedMessages, TimeMarker, PlayOrder);
+				const FSoundDistanceAttenuationMessage* FoundDistanceAttenuationCachedMessage = FindClosestMessageToTimestamp(SessionCachedMessages->DistanceAttenuationCachedMessages, TimeMarker, ActiveSoundPlayOrder);
 				if (FoundDistanceAttenuationCachedMessage)
 				{
 					CachedEntryInfos[Index].DistanceAttenuationMessage = *FoundDistanceAttenuationCachedMessage;
 				}
 
 				// LPF Freq
-				const FSoundLPFFreqMessage* FoundLPFFreqCachedMessage = FindClosestMessageToTimestamp(SessionCachedMessages->LPFFreqCachedMessages, TimeMarker, PlayOrder);
+				const FSoundLPFFreqMessage* FoundLPFFreqCachedMessage = FindClosestMessageToTimestamp(SessionCachedMessages->LPFFreqCachedMessages, TimeMarker, ActiveSoundPlayOrder);
 				if (FoundLPFFreqCachedMessage)
 				{
 					CachedEntryInfos[Index].LPFFreqMessage = *FoundLPFFreqCachedMessage;
 				}
 
 				// HPF Freq
-				const FSoundHPFFreqMessage* FoundHPFFreqCachedMessage = FindClosestMessageToTimestamp(SessionCachedMessages->HPFFreqCachedMessages, TimeMarker, PlayOrder);
+				const FSoundHPFFreqMessage* FoundHPFFreqCachedMessage = FindClosestMessageToTimestamp(SessionCachedMessages->HPFFreqCachedMessages, TimeMarker, ActiveSoundPlayOrder);
 				if (FoundHPFFreqCachedMessage)
 				{
 					CachedEntryInfos[Index].HPFFreqMessage = *FoundHPFFreqCachedMessage;
 				}
 
 				// Amplitude
-				const FSoundEnvelopeMessage* FoundAmplitudeCachedMessage = FindClosestMessageToTimestamp(SessionCachedMessages->AmplitudeCachedMessages, TimeMarker, PlayOrder);
+				const FSoundEnvelopeMessage* FoundAmplitudeCachedMessage = FindClosestMessageToTimestamp(SessionCachedMessages->AmplitudeCachedMessages, TimeMarker, ActiveSoundPlayOrder);
 				if (FoundAmplitudeCachedMessage)
 				{
 					CachedEntryInfos[Index].AmplitudeMessage = *FoundAmplitudeCachedMessage;
 				}
 
 				// Volume
-				const FSoundVolumeMessage* FoundVolumeCachedMessage = FindClosestMessageToTimestamp(SessionCachedMessages->VolumeCachedMessages, TimeMarker, PlayOrder);
+				const FSoundVolumeMessage* FoundVolumeCachedMessage = FindClosestMessageToTimestamp(SessionCachedMessages->VolumeCachedMessages, TimeMarker, ActiveSoundPlayOrder);
 				if (FoundVolumeCachedMessage)
 				{
 					CachedEntryInfos[Index].VolumeMessage = *FoundVolumeCachedMessage;
 				}
 
 				// Pitch
-				const FSoundPitchMessage* FoundPitchCachedMessage = FindClosestMessageToTimestamp(SessionCachedMessages->PitchCachedMessages, TimeMarker, PlayOrder);
+				const FSoundPitchMessage* FoundPitchCachedMessage = FindClosestMessageToTimestamp(SessionCachedMessages->PitchCachedMessages, TimeMarker, ActiveSoundPlayOrder);
 				if (FoundPitchCachedMessage)
 				{
 					CachedEntryInfos[Index].PitchMessage = *FoundPitchCachedMessage;
 				}
 
 				// Relative Render Cost
-				const FSoundRelativeRenderCostMessage* FoundRelativeRenderCostCachedMessage = FindClosestMessageToTimestamp(SessionCachedMessages->RelativeRenderCostCachedMessages, TimeMarker, PlayOrder);
+				const FSoundRelativeRenderCostMessage* FoundRelativeRenderCostCachedMessage = FindClosestMessageToTimestamp(SessionCachedMessages->RelativeRenderCostCachedMessages, TimeMarker, ActiveSoundPlayOrder);
 				if (FoundRelativeRenderCostCachedMessage)
 				{
 					CachedEntryInfos[Index].RelativeRenderCostMessage = *FoundRelativeRenderCostCachedMessage;
@@ -1385,14 +1698,9 @@
 				SoundEntry->Timestamp = CachedEntryInfo.PriorityMessage.Timestamp;
 			}
 		}
-
-#endif //!WITH_EDITOR
-
-		// Call parent method to update LastMessageId
-		FTraceProviderBase::OnTimingViewTimeMarkerChanged(TimeMarker);
 	}
+#endif // WITH_EDITOR
 
-#if !WITH_EDITOR
 	void FSoundTraceProvider::CacheIsPlottingFlagRecursive(const TSharedPtr<IDashboardDataTreeViewEntry>& Entry)
 	{
 		using namespace FSoundTraceProviderPrivate;
@@ -1419,7 +1727,6 @@
 			CacheIsPlottingFlagRecursive(Child);
 		}
 	}
-#endif // !WITH_EDITOR
 
 } // namespace UE::Audio::Insights
 #undef LOCTEXT_NAMESPACE
diff --git a/Engine/Plugins/AudioInsights/Source/AudioInsights/Private/Providers/SoundTraceProvider.h b/Engine/Plugins/AudioInsights/Source/AudioInsights/Private/Providers/SoundTraceProvider.h
--- a/Engine/Plugins/AudioInsights/Source/AudioInsights/Private/Providers/SoundTraceProvider.h
+++ b/Engine/Plugins/AudioInsights/Source/AudioInsights/Private/Providers/SoundTraceProvider.h
@@ -37,6 +37,9 @@ namespace UE::Audio::Insights
 
 		void SetDashboardTimeoutTime(const double InTimeOutTime) { DashboardTimeoutTime = InTimeOutTime; }
 
+		DECLARE_MULTICAST_DELEGATE(FOnProcessPlotData);
+		FOnProcessPlotData OnProcessPlotData;
+
 	private:
 		void GetOrCreateActiveSoundEntry(const FSoundStartMessage& Msg, TSharedPtr<FSoundDashboardEntry>& OutReturnedSoundEntry);
 		void CreateAndGetActiveSoundEntryIfNotFound(const FSoundStartMessage& Msg, TSharedPtr<FSoundDashboardEntry>& OutReturnedSoundEntry);
@@ -47,10 +50,17 @@ namespace UE::Audio::Insights
 		void CollectAggregateData(FSoundDashboardEntry& ActiveSoundEntry);
 
 		virtual void OnTimingViewTimeMarkerChanged(double TimeMarker) override;
+		virtual void OnTimeControlMethodReset() override;
+
+#if WITH_EDITOR
+		void CollectParamsForTimestamp(const FAudioInsightsCacheManager& CacheManager, const TArray<uint32>& ActiveSoundPlayOrderArray, const double TimeMarker);
+#else
+		void CollectParamsForTimestampStandalone(const TArray<uint32>& ActiveSoundPlayOrderArray, const double TimeMarker);
+#endif // WITH_EDITOR
 
-#if !WITH_EDITOR
 		void CacheIsPlottingFlagRecursive(const TSharedPtr<IDashboardDataTreeViewEntry>& Entry);
 
+#if !WITH_EDITOR
 		TUniquePtr<FSoundSessionCachedMessages> SessionCachedMessages;
 #endif // !WITH_EDITOR
 
diff --git a/Engine/Plugins/AudioInsights/Source/AudioInsights/Public/Messages/SoundTraceMessages.h b/Engine/Plugins/AudioInsights/Source/AudioInsights/Public/Messages/SoundTraceMessages.h
--- a/Engine/Plugins/AudioInsights/Source/AudioInsights/Public/Messages/SoundTraceMessages.h
+++ b/Engine/Plugins/AudioInsights/Source/AudioInsights/Public/Messages/SoundTraceMessages.h
@@ -2,6 +2,7 @@
 #pragma once
 
 #include "AudioDefines.h"
+#include "Cache/IAudioCachedMessage.h"
 #include "DSP/Dsp.h"
 #include "Messages/AnalyzerMessageQueue.h"
 #include "Trace/Analyzer.h"
@@ -37,15 +38,35 @@ namespace UE::Audio::Insights
 		const FString SoundCueTemplate = TEXT("SoundCueTemplate");
 	}
 
+	namespace SoundMessageNames
+	{
+		extern const FName SoundStart;
+		extern const FName SoundIsAlivePing;
+		extern const FName SoundWaveStart;
+		extern const FName SoundWaveIsAlivePing;
+		extern const FName SoundStop;
+		
+		extern const FName PriorityParam;
+		extern const FName DistanceParam;
+		extern const FName DistanceAttenuationParam;
+		extern const FName HPFFreqParam;
+		extern const FName LPFFreqParam;
+		extern const FName EnvelopeParam;
+		extern const FName PitchParam;
+		extern const FName VolumeParam;
+		extern const FName RelativeRenderCostParam;
+	};
+
 	// Trace messages
-	struct FSoundMessageBase
+	struct FSoundMessageBase : public IAudioCachedMessage
 	{
 		FSoundMessageBase() = default;
 		FSoundMessageBase(const Trace::IAnalyzer::FOnEventContext& InContext);
 
+		virtual uint32 GetID() const override { return PlayOrder; }
+
 		::Audio::FDeviceId DeviceId = INDEX_NONE;
 		uint32 PlayOrder = INDEX_NONE;
-		double Timestamp = 0.0;
 	};
 
 	struct FSoundParameterMessage : public FSoundMessageBase
@@ -58,6 +79,8 @@ namespace UE::Audio::Insights
 			ActiveSoundPlayOrder = InContext.EventData.GetValue<uint32>("ActiveSoundPlayOrder");
 		}
 
+		virtual uint32 GetID() const override { return ActiveSoundPlayOrder; }
+
 		uint32 ActiveSoundPlayOrder = INDEX_NONE;
 	};
 
@@ -66,11 +89,25 @@ namespace UE::Audio::Insights
 		FSoundStartMessage() = default;
 		FSoundStartMessage(const Trace::IAnalyzer::FOnEventContext& InContext);
 
+		virtual uint32 GetSizeOf() const override;
+		virtual const FName GetMessageName() const override { return SoundMessageNames::SoundStart; }
+
 		FString Name;
 		ESoundDashboardEntryType EntryType;
 		FString ActorLabel;
 	};
 
+	struct FSoundIsAlivePingMessage : public FSoundStartMessage
+	{
+		FSoundIsAlivePingMessage() = default;
+		FSoundIsAlivePingMessage(const Trace::IAnalyzer::FOnEventContext& InContext)
+			: FSoundStartMessage(InContext)
+		{
+		}
+
+		virtual const FName GetMessageName() const override { return SoundMessageNames::SoundIsAlivePing; }
+	};
+
 	struct FSoundWaveStartMessage : public FSoundStartMessage
 	{
 		FSoundWaveStartMessage() = default;
@@ -81,35 +118,61 @@ namespace UE::Audio::Insights
 			ActiveSoundPlayOrder = InContext.EventData.GetValue<uint32>("ActiveSoundPlayOrder");
 		}
 
+		virtual uint32 GetSizeOf() const override;
+		virtual const FName GetMessageName() const override { return SoundMessageNames::SoundWaveStart; }
+
 		uint32 ActiveSoundPlayOrder;
 	};
 
-	using FSoundIsAlivePingMessage = FSoundStartMessage;
-	using FSoundWaveIsAlivePingMessage = FSoundWaveStartMessage;
-	using FSoundStopMessage = FSoundMessageBase;
+	struct FSoundWaveIsAlivePingMessage : public FSoundWaveStartMessage
+	{
+		FSoundWaveIsAlivePingMessage() = default;
+		FSoundWaveIsAlivePingMessage(const Trace::IAnalyzer::FOnEventContext& InContext)
+			: FSoundWaveStartMessage(InContext)
+		{
+		}
+
+		virtual const FName GetMessageName() const override { return SoundMessageNames::SoundWaveIsAlivePing; }
+	};
 
-#define DEFINE_SOUND_PARAM_MESSAGE(ClassName, SuperClassName, ParamName, Type, Default)	\
-	struct ClassName : public SuperClassName													\
+	struct FSoundStopMessage : public FSoundMessageBase
+	{
+		FSoundStopMessage() = default;
+		FSoundStopMessage(const Trace::IAnalyzer::FOnEventContext& InContext)
+			: FSoundMessageBase(InContext)
+		{
+		}
+
+		virtual uint32 GetSizeOf() const override;
+		virtual const FName GetMessageName() const override { return SoundMessageNames::SoundStop; }
+	};
+
+#define DEFINE_SOUND_PARAM_MESSAGE(ClassName, MessageName, ParamName, Type, Default)			\
+	struct ClassName : public FSoundParameterMessage											\
 	{																							\
 		ClassName() = default;																	\
 		ClassName(const Trace::IAnalyzer::FOnEventContext& InContext)							\
-			: SuperClassName(InContext)															\
+			: FSoundParameterMessage(InContext)													\
 		{																						\
 			const Trace::IAnalyzer::FEventData& EventData = InContext.EventData;				\
 			ParamName = EventData.GetValue<Type>(#ParamName);									\
 		}																						\
+																								\
+		virtual uint32 GetSizeOf() const override { return sizeof(ClassName); }					\
+		virtual const FName GetMessageName() const override { return MessageName; }				\
+																								\
 		Type ParamName = Default;																\
 	};
 
-	DEFINE_SOUND_PARAM_MESSAGE(FSoundPriorityMessage, FSoundParameterMessage, Priority, float, 0.0f)
-	DEFINE_SOUND_PARAM_MESSAGE(FSoundDistanceMessage, FSoundParameterMessage, Distance, float, 0.0f)
-	DEFINE_SOUND_PARAM_MESSAGE(FSoundDistanceAttenuationMessage, FSoundParameterMessage, DistanceAttenuation, float, 1.0f)
-	DEFINE_SOUND_PARAM_MESSAGE(FSoundHPFFreqMessage, FSoundParameterMessage, HPFFrequency, float, MIN_FILTER_FREQUENCY)
-	DEFINE_SOUND_PARAM_MESSAGE(FSoundLPFFreqMessage, FSoundParameterMessage, LPFFrequency, float, MAX_FILTER_FREQUENCY)
-	DEFINE_SOUND_PARAM_MESSAGE(FSoundEnvelopeMessage, FSoundParameterMessage, Envelope, float, 0.0f)
-	DEFINE_SOUND_PARAM_MESSAGE(FSoundPitchMessage, FSoundParameterMessage, Pitch, float, 1.0f)
-	DEFINE_SOUND_PARAM_MESSAGE(FSoundVolumeMessage, FSoundParameterMessage, Volume, float, 1.0f)
-	DEFINE_SOUND_PARAM_MESSAGE(FSoundRelativeRenderCostMessage, FSoundParameterMessage, RelativeRenderCost, float, 0.0f)
+	DEFINE_SOUND_PARAM_MESSAGE(FSoundPriorityMessage, SoundMessageNames::PriorityParam, Priority, float, 0.0f)
+	DEFINE_SOUND_PARAM_MESSAGE(FSoundDistanceMessage, SoundMessageNames::DistanceParam, Distance, float, 0.0f)
+	DEFINE_SOUND_PARAM_MESSAGE(FSoundDistanceAttenuationMessage, SoundMessageNames::DistanceAttenuationParam, DistanceAttenuation, float, 1.0f)
+	DEFINE_SOUND_PARAM_MESSAGE(FSoundHPFFreqMessage, SoundMessageNames::HPFFreqParam, HPFFrequency, float, MIN_FILTER_FREQUENCY)
+	DEFINE_SOUND_PARAM_MESSAGE(FSoundLPFFreqMessage, SoundMessageNames::LPFFreqParam, LPFFrequency, float, MAX_FILTER_FREQUENCY)
+	DEFINE_SOUND_PARAM_MESSAGE(FSoundEnvelopeMessage, SoundMessageNames::EnvelopeParam, Envelope, float, 0.0f)
+	DEFINE_SOUND_PARAM_MESSAGE(FSoundPitchMessage, SoundMessageNames::PitchParam, Pitch, float, 1.0f)
+	DEFINE_SOUND_PARAM_MESSAGE(FSoundVolumeMessage, SoundMessageNames::VolumeParam, Volume, float, 1.0f)
+	DEFINE_SOUND_PARAM_MESSAGE(FSoundRelativeRenderCostMessage, SoundMessageNames::RelativeRenderCostParam, RelativeRenderCost, float, 0.0f)
 #undef DEFINE_SOUND_PARAM_MESSAGE
 
 	// Trace message queues
@@ -221,9 +284,8 @@ namespace UE::Audio::Insights
 		virtual bool HasSetInitExpansion() const override { return bHasSetInitExpansion; }
 		virtual void ResetHasSetInitExpansion() override { bHasSetInitExpansion = true; }
 
-#if !WITH_EDITOR
+		void ResetDataBuffers(const uint32 DataPointsCapacity);
 		void SwitchBuffersToSingleValues();
-#endif // !WITH_EDITOR
 
 		enum class EPinnedEntryType : uint8
 		{
@@ -244,9 +306,7 @@ namespace UE::Audio::Insights
 		bool bHasSetInitExpansion = false;
 		bool bIsVisible = false;
 		bool bIsPlotActive = false;
-#if !WITH_EDITOR
 		bool bForceKeepEntryAlive = false;
-#endif // !WITH_EDITOR
 
 		::Audio::TCircularAudioBuffer<FDataPoint> PriorityDataPoints;
 		::Audio::TCircularAudioBuffer<FDataPoint> DistanceDataPoints;
