[Audio Insights] Cache Management part 06: Integrate the Plots widget with the new cache manager

diff --git a/Engine/Plugins/AudioInsights/Source/AudioInsights/Private/Views/SoundDashboardViewFactory.cpp b/Engine/Plugins/AudioInsights/Source/AudioInsights/Private/Views/SoundDashboardViewFactory.cpp
--- a/Engine/Plugins/AudioInsights/Source/AudioInsights/Private/Views/SoundDashboardViewFactory.cpp
+++ b/Engine/Plugins/AudioInsights/Source/AudioInsights/Private/Views/SoundDashboardViewFactory.cpp
@@ -114,11 +114,7 @@ namespace UE::Audio::Insights
 		{
 			const FSoundDashboardEntry& SoundEntry = CastEntry(InEntry);
 
-#if WITH_EDITOR
-			return SoundEntry.bIsVisible && (bShowRecentlyStoppedSounds || SoundEntry.TimeoutTimestamp == INVALID_TIMEOUT);
-#else
 			return SoundEntry.bIsVisible && ((bShowRecentlyStoppedSounds || SoundEntry.TimeoutTimestamp == INVALID_TIMEOUT) || SoundEntry.bForceKeepEntryAlive);
-#endif // WITH_EDITOR
 		}
 
 		FSlateColor GetIconColor(const TSharedPtr<IDashboardDataTreeViewEntry>& InEntry)
@@ -786,6 +782,7 @@ namespace UE::Audio::Insights
 		FTraceModule& AudioInsightsTraceModule = static_cast<FTraceModule&>(FAudioInsightsModule::GetChecked().GetTraceModule());
 
 		const TSharedPtr<FSoundTraceProvider> SoundsTraceProvider = MakeShared<FSoundTraceProvider>();
+		SoundsTraceProvider->OnProcessPlotData.AddRaw(this, &FSoundDashboardViewFactory::ProcessPlotData);
 
 		AudioInsightsTraceModule.AddTraceProvider(SoundsTraceProvider);
 
@@ -806,6 +803,12 @@ namespace UE::Audio::Insights
 
 	FSoundDashboardViewFactory::~FSoundDashboardViewFactory()
 	{
+		const TSharedPtr<FSoundTraceProvider> Provider = FindProvider<FSoundTraceProvider>();
+		if (Provider.IsValid())
+		{
+			Provider->OnProcessPlotData.RemoveAll(this);
+		}
+
 		FSoundDashboardCommands::Unregister();
 
 #if WITH_EDITOR
@@ -1427,11 +1430,7 @@ namespace UE::Audio::Insights
 			.RenderOpacity_Lambda([Item]
 			{
 				const FSoundDashboardEntry& SoundDashboardEntry = CastEntry(*Item);
-#if WITH_EDITOR
-				return SoundDashboardEntry.TimeoutTimestamp == INVALID_TIMEOUT ? 1.0f : 0.25f;
-#else
 				return SoundDashboardEntry.TimeoutTimestamp == INVALID_TIMEOUT && !SoundDashboardEntry.bForceKeepEntryAlive ? 1.0f : 0.25f;
-#endif // WITH_EDITOR
 			})
 			.Visibility_Lambda([this, Item]()
 			{
@@ -1453,11 +1452,7 @@ namespace UE::Audio::Insights
 				const int32 NumUnpinnedChildren = GetNumChildrenWithoutPinEntryType(*Item, FSoundDashboardEntry::EPinnedEntryType::HiddenOriginalEntry, bShowRecentlyStoppedSounds);
 
 				const bool bRowShouldBeVisible = SoundEntry.bIsVisible 
-#if WITH_EDITOR
-												&& (bShowRecentlyStoppedSounds || SoundEntry.TimeoutTimestamp == INVALID_TIMEOUT)
-#else
 												&& ((bShowRecentlyStoppedSounds || SoundEntry.TimeoutTimestamp == INVALID_TIMEOUT) || SoundEntry.bForceKeepEntryAlive)
-#endif // WITH_EDITOR
 												&& SoundEntry.PinnedEntryType != FSoundDashboardEntry::EPinnedEntryType::HiddenOriginalEntry 
 												&& (!IsCategoryItem(*Item) || (NumUnpinnedChildren > 0));
 
@@ -1518,9 +1513,17 @@ namespace UE::Audio::Insights
 			return bEntryTypePassesFilter;
 		});
 
-		if (InReason == FTraceTreeDashboardViewFactory::EProcessReason::EntriesUpdated && FilteredEntriesListView.IsValid())
+		if (InReason == FTraceTreeDashboardViewFactory::EProcessReason::EntriesUpdated)
 		{
-			OnProcessPlotData.Broadcast(DataViewEntries, FilteredEntriesListView->GetSelectedItems());
+#if WITH_EDITOR
+			const FAudioInsightsModule& AudioInsightsModule = FAudioInsightsModule::GetChecked();
+			if (AudioInsightsModule.GetTimingViewExtender().GetTimeControlMethod() != ETimeControlMethod::LatestMessageTime)
+			{
+				return;
+			}
+#endif // WITH_EDITOR
+
+			ProcessPlotData();
 		}
 	}
 
@@ -1706,12 +1709,10 @@ namespace UE::Audio::Insights
 
 	void FSoundDashboardViewFactory::OnSelectionChanged(TSharedPtr<IDashboardDataTreeViewEntry> SelectedItem, ESelectInfo::Type SelectInfo)
 	{
-#if !WITH_EDITOR
 		if (FilteredEntriesListView.IsValid())
 		{
 			OnUpdatePlotSelection.Broadcast(FilteredEntriesListView->GetSelectedItems());
 		}
-#endif // !WITH_EDITOR
 	}
 
 	FReply FSoundDashboardViewFactory::OnDataRowKeyInput(const FGeometry& InGeometry, const FKeyEvent& InKeyEvent) const
@@ -3186,6 +3187,14 @@ namespace UE::Audio::Insights
 		}
 	}
 
+	void FSoundDashboardViewFactory::ProcessPlotData()
+	{
+		if (FilteredEntriesListView.IsValid())
+		{
+			OnProcessPlotData.Broadcast(DataViewEntries, FilteredEntriesListView->GetSelectedItems());
+		}
+	}
+
 	TMap<FName, FSoundPlotsWidgetView::FPlotColumnInfo> FSoundDashboardViewFactory::GetPlotColumnInfo()
 	{
 		using namespace FSoundDashboardViewFactoryPrivate;
diff --git a/Engine/Plugins/AudioInsights/Source/AudioInsights/Private/Views/SoundPlotsDashboardViewFactory.cpp b/Engine/Plugins/AudioInsights/Source/AudioInsights/Private/Views/SoundPlotsDashboardViewFactory.cpp
--- a/Engine/Plugins/AudioInsights/Source/AudioInsights/Private/Views/SoundPlotsDashboardViewFactory.cpp
+++ b/Engine/Plugins/AudioInsights/Source/AudioInsights/Private/Views/SoundPlotsDashboardViewFactory.cpp
@@ -54,11 +54,8 @@ namespace UE::Audio::Insights
 		PlotsWidget = MakeShared<FSoundPlotsWidgetView>(MoveTemp(ColumnInfo), [](const IDashboardDataTreeViewEntry& InEntry) -> bool { return CastEntry(InEntry).bIsPlotActive; });
 
 		OnProcessPlotData = SoundsDashboard->OnProcessPlotData.AddSP(this, &FSoundPlotsDashboardViewFactory::ProcessPlotData);
-
-#if !WITH_EDITOR
 		OnUpdatePlotVisibility = SoundsDashboard->OnUpdatePlotVisibility.AddSP(this, &FSoundPlotsDashboardViewFactory::UpdatePlotVisibility);
 		OnUpdatePlotSelection = SoundsDashboard->OnUpdatePlotSelection.AddSP(this, &FSoundPlotsDashboardViewFactory::UpdatePlotSelection);
-#endif // !WITH_EDITOR
 	}
 
 	void FSoundPlotsDashboardViewFactory::ProcessPlotData(const TArray<TSharedPtr<IDashboardDataTreeViewEntry>>& DataViewEntries, const TArray<TSharedPtr<IDashboardDataTreeViewEntry>>& SelectedEntries)
@@ -69,7 +66,6 @@ namespace UE::Audio::Insights
 		}
 	}
 
-#if !WITH_EDITOR
 	void FSoundPlotsDashboardViewFactory::UpdatePlotVisibility(const TArray<TSharedPtr<IDashboardDataTreeViewEntry>>& DataViewEntries)
 	{
 		if (PlotsWidget.IsValid())
@@ -85,6 +81,5 @@ namespace UE::Audio::Insights
 			PlotsWidget->UpdatePlotSelection(SelectedEntries);
 		}
 	}
-#endif // !WITH_EDITOR
 } // namespace UE::Audio::Insights
 #undef LOCTEXT_NAMESPACE
diff --git a/Engine/Plugins/AudioInsights/Source/AudioInsights/Private/Views/SoundPlotsWidgetView.cpp b/Engine/Plugins/AudioInsights/Source/AudioInsights/Private/Views/SoundPlotsWidgetView.cpp
--- a/Engine/Plugins/AudioInsights/Source/AudioInsights/Private/Views/SoundPlotsWidgetView.cpp
+++ b/Engine/Plugins/AudioInsights/Source/AudioInsights/Private/Views/SoundPlotsWidgetView.cpp
@@ -69,6 +69,7 @@ namespace UE::Audio::Insights
 	{
 		FAudioInsightsModule& AudioInsightsModule = FAudioInsightsModule::GetChecked();
 		AudioInsightsModule.GetTimingViewExtender().OnTimingViewTimeMarkerChanged.AddSP(this, &FSoundPlotsWidgetView::OnTimingViewTimeMarkerChanged);
+		AudioInsightsModule.GetTimingViewExtender().OnTimeControlMethodReset.AddSP(this, &FSoundPlotsWidgetView::OnTimeControlMethodReset);
 
 #if WITH_EDITOR
 		FEditorDelegates::PostPIEStarted.AddSP(this, &FSoundPlotsWidgetView::OnPIEStarted);
@@ -90,12 +91,10 @@ namespace UE::Audio::Insights
 
 	void FSoundPlotsWidgetView::ProcessPlotData(const TArray<TSharedPtr<IDashboardDataTreeViewEntry>>& DataViewEntries, const TArray<TSharedPtr<IDashboardDataTreeViewEntry>>& SelectedEntries)
 	{
-#if !WITH_EDITOR
 		if (GameState == EGameState::ControlledByTimemarker && !bTimeMarkerMoved)
 		{
 			return;
 		}
-#endif // !WITH_EDITOR
 
 		if (DataViewEntries.IsEmpty() || GameState == EGameState::Paused)
 		{
@@ -200,7 +199,6 @@ namespace UE::Audio::Insights
 		}
 	}
 
-#if !WITH_EDITOR
 	void FSoundPlotsWidgetView::UpdatePlotSelection(const TArray<TSharedPtr<IDashboardDataTreeViewEntry>>& SelectedEntries)
 	{
 		using namespace FSoundPlotsWidgetViewPrivate;
@@ -228,7 +226,6 @@ namespace UE::Audio::Insights
 			}
 		}
 	}
-#endif // !WITH_EDITOR
 
 	void FSoundPlotsWidgetView::OnAnalysisStarting(const double Timestamp)
 	{
@@ -304,6 +301,12 @@ namespace UE::Audio::Insights
 		bTimeMarkerMoved = true;
 	}
 
+	void FSoundPlotsWidgetView::OnTimeControlMethodReset()
+	{
+		GameState = EGameState::Running;
+		bTimeMarkerMoved = false;
+	}
+
 	TSharedRef<SWidget> FSoundPlotsWidgetView::MakePlotsWidget()
 	{
 		if (!PlotWidget.IsValid() || !ComboBox.IsValid())
@@ -524,6 +527,10 @@ namespace UE::Audio::Insights
 					(*LatestTimestamp)->Add(SourceId);
 				}
 			}
+			else if (GameState == EGameState::ControlledByTimemarker)
+			{
+				DataPoints->Empty(DataPoints->Num());
+			}
 
 			auto DataFunc = GetPlotColumnDataFunc(ColumnName);
 			if (!DataFunc.IsSet())
@@ -545,7 +552,11 @@ namespace UE::Audio::Insights
 					double* LatestColumnTimestep = (*LatestTimestampMap)->Find(SourceId);
 					if (LatestColumnTimestep)
 					{
-						if (Timestamp <= *LatestColumnTimestep)
+						if (DataPoints->IsEmpty())
+						{
+							*LatestColumnTimestep = Timestamp;
+						}
+						else if (Timestamp <= *LatestColumnTimestep)
 						{
 							continue;
 						}
@@ -556,18 +567,21 @@ namespace UE::Audio::Insights
 					}
 				}
 
-#if WITH_EDITOR
-				CurrentTimestamp = FMath::Max(CurrentTimestamp, Timestamp);
-				const double DataPointTime = Timestamp - BeginTimestamp;
-#else
 				if (GameState != EGameState::ControlledByTimemarker)
 				{
 					CurrentTimestamp = FMath::Max(CurrentTimestamp, Timestamp);
 				}
 
+#if WITH_EDITOR
+				const double DataPointTime = Timestamp - BeginTimestamp;
+#else
 				const double DataPointTime = Timestamp;
 #endif // WITH_EDITOR
-				DataPoints->Emplace(DataPointTime, UnitConversionFunc.IsSet() ? UnitConversionFunc(Value) : Value);
+
+				if (DataPointTime >= 0.0)
+				{
+					DataPoints->Emplace(DataPointTime, UnitConversionFunc.IsSet() ? UnitConversionFunc(Value) : Value);
+				}
 			}
 		}
 
@@ -768,7 +782,7 @@ namespace UE::Audio::Insights
 
 		const bool bAnyMessageReceived = CurrentTimestamp != TNumericLimits<double>::Lowest();
 
-		if (GameState != EGameState::Paused && bAnyMessageReceived)
+		if (GameState != EGameState::Paused && GameState != EGameState ::ControlledByTimemarker && bAnyMessageReceived)
 		{
 			const double RelativeCurrentTime = CurrentTime - BeginTimestamp;
 			TimestampsDiff = RelativeCurrentTime - (CurrentTimestamp - BeginTimestamp);
diff --git a/Engine/Plugins/AudioInsights/Source/AudioInsights/Public/Views/SoundDashboardViewFactory.h b/Engine/Plugins/AudioInsights/Source/AudioInsights/Public/Views/SoundDashboardViewFactory.h
--- a/Engine/Plugins/AudioInsights/Source/AudioInsights/Public/Views/SoundDashboardViewFactory.h
+++ b/Engine/Plugins/AudioInsights/Source/AudioInsights/Public/Views/SoundDashboardViewFactory.h
@@ -202,6 +202,8 @@ namespace UE::Audio::Insights
 		UE_API void SetShowPlotRecursive(const TSharedPtr<IDashboardDataTreeViewEntry>& Entry, const bool bShowPlot);
 		UE_API void SetShowPlot(const TSharedPtr<IDashboardDataTreeViewEntry>& Entry, const bool bShowPlot);
 
+		void ProcessPlotData();
+
 		UE_API bool SelectedItemsIncludesAnAsset() const;
 		UE_API bool SelectionIncludesUnpinnedItem() const;
 		UE_API void ClearSelection();
diff --git a/Engine/Plugins/AudioInsights/Source/AudioInsights/Public/Views/SoundPlotsDashboardViewFactory.h b/Engine/Plugins/AudioInsights/Source/AudioInsights/Public/Views/SoundPlotsDashboardViewFactory.h
--- a/Engine/Plugins/AudioInsights/Source/AudioInsights/Public/Views/SoundPlotsDashboardViewFactory.h
+++ b/Engine/Plugins/AudioInsights/Source/AudioInsights/Public/Views/SoundPlotsDashboardViewFactory.h
@@ -30,13 +30,11 @@ namespace UE::Audio::Insights
 
 		void ProcessPlotData(const TArray<TSharedPtr<IDashboardDataTreeViewEntry>>& DataViewEntries, const TArray<TSharedPtr<IDashboardDataTreeViewEntry>>& SelectedEntries);
 
-#if !WITH_EDITOR
 		void UpdatePlotVisibility(const TArray<TSharedPtr<IDashboardDataTreeViewEntry>>& DataViewEntries);
 		void UpdatePlotSelection(const TArray<TSharedPtr<IDashboardDataTreeViewEntry>>& SelectedEntries);
 
 		FDelegateHandle OnUpdatePlotVisibility;
 		FDelegateHandle OnUpdatePlotSelection;
-#endif // !WITH_EDITOR
 
 		FDelegateHandle OnProcessPlotData;
 	};
diff --git a/Engine/Plugins/AudioInsights/Source/AudioInsights/Public/Views/SoundPlotsWidgetView.h b/Engine/Plugins/AudioInsights/Source/AudioInsights/Public/Views/SoundPlotsWidgetView.h
--- a/Engine/Plugins/AudioInsights/Source/AudioInsights/Public/Views/SoundPlotsWidgetView.h
+++ b/Engine/Plugins/AudioInsights/Source/AudioInsights/Public/Views/SoundPlotsWidgetView.h
@@ -37,9 +37,7 @@ namespace UE::Audio::Insights
 		UE_API void ProcessPlotData(const TArray<TSharedPtr<IDashboardDataTreeViewEntry>>& DataViewEntries, const TArray<TSharedPtr<IDashboardDataTreeViewEntry>>& SelectedEntries);
 		UE_API void UpdatePlotVisibility(const TArray<TSharedPtr<IDashboardDataTreeViewEntry>>& DataViewEntries);
 
-#if !WITH_EDITOR
 		UE_API void UpdatePlotSelection(const TArray<TSharedPtr<IDashboardDataTreeViewEntry>>& SelectedEntries);
-#endif // !WITH_EDITOR
 
 	private:
 		using FPlotCurvePoint = SAudioCurveView::FCurvePoint;
@@ -61,6 +59,7 @@ namespace UE::Audio::Insights
 #endif // WITH_EDITOR
 
 		void OnTimingViewTimeMarkerChanged(double InTimeMarker);
+		void OnTimeControlMethodReset();
 
 		TSharedRef<SWidget> MakePlotsWidget();
 		void ConstructPlotsWidgetContents();
