[Audio Insights] Enable meter in the Analyzers tab for main submix in standalone

diff --git a/Engine/Plugins/AudioInsights/Source/AudioInsights/Private/AudioInsightsDashboardFactory.cpp b/Engine/Plugins/AudioInsights/Source/AudioInsights/Private/AudioInsightsDashboardFactory.cpp
--- a/Engine/Plugins/AudioInsights/Source/AudioInsights/Private/AudioInsightsDashboardFactory.cpp
+++ b/Engine/Plugins/AudioInsights/Source/AudioInsights/Private/AudioInsightsDashboardFactory.cpp
@@ -192,6 +192,7 @@ namespace UE::Audio::Insights
 		TSharedRef<FTabManager::FStack> LogTabStack = FTabManager::NewStack();
 		TSharedRef<FTabManager::FStack> AnalysisTabStack = FTabManager::NewStack();
 		TSharedRef<FTabManager::FStack> AudioMetersTabStack = FTabManager::NewStack();
+		TSharedRef<FTabManager::FStack> AudioAnalyzerRackTabStack = FTabManager::NewStack();
 		TSharedRef<FTabManager::FStack> PlotsTabStack = FTabManager::NewStack();
 
 		for (const auto& [FactoryName, Factory] : DashboardViewFactories)
@@ -218,6 +219,12 @@ namespace UE::Audio::Insights
 				}
 				break;
 
+				case EDefaultDashboardTabStack::AudioAnalyzerRack:
+				{
+					AudioAnalyzerRackTabStack->AddTab(FactoryName, ETabState::OpenedTab);
+				}
+				break;
+
 				case EDefaultDashboardTabStack::Plots:
 				{
 					PlotsTabStack->AddTab(FactoryName, ETabState::OpenedTab);
@@ -231,13 +238,14 @@ namespace UE::Audio::Insights
 
 		AnalysisTabStack->SetForegroundTab(FName("Sounds"));
 
-		return FTabManager::NewLayout("AudioDashboard_Layout_v1")
+		return FTabManager::NewLayout("AudioDashboard_Layout_v2")
 		->AddArea
 		(
 			FTabManager::NewPrimaryArea()
 			->SetOrientation(Orient_Vertical)
 			->Split
 			(
+				// Left column
 				FTabManager::NewSplitter()
 				->SetOrientation(Orient_Horizontal)
 				->Split
@@ -245,6 +253,8 @@ namespace UE::Audio::Insights
 					LogTabStack
 					->SetSizeCoefficient(0.25f)
 				)
+
+				// Middle column
 				->Split
 				(
 					// Top
@@ -282,6 +292,13 @@ namespace UE::Audio::Insights
 						)
 					)
 				)
+
+				// Right column
+				->Split
+				(
+					AudioAnalyzerRackTabStack
+					->SetSizeCoefficient(0.25f)
+				)
 			)
 		);
 	}
diff --git a/Engine/Plugins/AudioInsights/Source/AudioInsights/Private/AudioInsightsModule.cpp b/Engine/Plugins/AudioInsights/Source/AudioInsights/Private/AudioInsightsModule.cpp
--- a/Engine/Plugins/AudioInsights/Source/AudioInsights/Private/AudioInsightsModule.cpp
+++ b/Engine/Plugins/AudioInsights/Source/AudioInsights/Private/AudioInsightsModule.cpp
@@ -18,6 +18,7 @@
 #include "Views/AudioBusDashboardViewFactory.h"
 #include "Views/AudioEventLogDashboardViewFactory.h"
 #include "Views/AudioMetersPanelDashboardViewFactory.h"
+#include "Views/OutputMeterDashboardViewFactory.h"
 #include "Views/SoundDashboardViewFactory.h"
 #include "Views/SoundPlotsDashboardViewFactory.h"
 #include "Views/SubmixDashboardViewFactory.h"
@@ -59,14 +60,16 @@ namespace UE::Audio::Insights
 
 			TSharedRef<FSoundDashboardViewFactory> SoundsDashboard = MakeShared<FSoundDashboardViewFactory>();
 			TSharedRef<FSoundPlotsDashboardViewFactory> PlotsDashboard = MakeShared<FSoundPlotsDashboardViewFactory>();
+			TSharedRef<FSubmixDashboardViewFactory> SubmixDashboard = MakeShared<FSubmixDashboardViewFactory>();
 
 			DashboardFactory->RegisterViewFactory(MakeShared<FAudioEventLogDashboardViewFactory>());
 			DashboardFactory->RegisterViewFactory(SoundsDashboard);
 			DashboardFactory->RegisterViewFactory(MakeShared<FVirtualLoopDashboardViewFactory>());
 			DashboardFactory->RegisterViewFactory(PlotsDashboard);
-			DashboardFactory->RegisterViewFactory(MakeShared<FSubmixDashboardViewFactory>());
+			DashboardFactory->RegisterViewFactory(SubmixDashboard);
 			DashboardFactory->RegisterViewFactory(MakeShared<FAudioBusDashboardViewFactory>());
 			DashboardFactory->RegisterViewFactory(MakeShared<FAudioMetersPanelDashboardViewFactory>());
+			DashboardFactory->RegisterViewFactory(MakeShared<FOutputMeterDashboardViewFactory>(SubmixDashboard));
 
 			PlotsDashboard->InitPlots(SoundsDashboard);
 #endif // !WITH_EDITOR
diff --git a/Engine/Plugins/AudioInsights/Source/AudioInsights/Private/Providers/SubmixTraceProvider.cpp b/Engine/Plugins/AudioInsights/Source/AudioInsights/Private/Providers/SubmixTraceProvider.cpp
--- a/Engine/Plugins/AudioInsights/Source/AudioInsights/Private/Providers/SubmixTraceProvider.cpp
+++ b/Engine/Plugins/AudioInsights/Source/AudioInsights/Private/Providers/SubmixTraceProvider.cpp
@@ -191,6 +191,7 @@ namespace UE::Audio::Insights
 				FSubmixDashboardEntry& EntryRef = *OutEntry->Get();
 				EntryRef.Name = Msg.Name;
 				EntryRef.AudioMeterInfo->NumChannels = Msg.NumChannels;
+				EntryRef.bIsMainSubmix = Msg.bIsMainSubmix;
 			}
 		});
 #endif // !WITH_EDITOR
@@ -373,6 +374,7 @@ namespace UE::Audio::Insights
 
 				Entry->Name = *SubmixLoadedCachedMessage.Name;
 				Entry->AudioMeterInfo->NumChannels = SubmixLoadedCachedMessage.NumChannels;
+				Entry->bIsMainSubmix = SubmixLoadedCachedMessage.bIsMainSubmix;
 			});
 		}
 
diff --git a/Engine/Plugins/AudioInsights/Source/AudioInsights/Private/Views/OutputMeterDashboardViewFactory.cpp b/Engine/Plugins/AudioInsights/Source/AudioInsights/Private/Views/OutputMeterDashboardViewFactory.cpp
new file mode 100644
--- /dev/null
+++ b/Engine/Plugins/AudioInsights/Source/AudioInsights/Private/Views/OutputMeterDashboardViewFactory.cpp
@@ -0,0 +1,175 @@
+// Copyright Epic Games, Inc. All Rights Reserved.
+#include "Views/OutputMeterDashboardViewFactory.h"
+
+#include "AudioInsightsStyle.h"
+#include "DSP/Dsp.h"
+#include "ISessionServicesModule.h"
+#include "SAudioMeterWidget.h"
+#include "Widgets/Colors/SColorBlock.h"
+#include "Widgets/SOverlay.h"
+
+#define LOCTEXT_NAMESPACE "AudioInsights"
+
+namespace UE::Audio::Insights
+{
+	FOutputMeterDashboardViewFactory::FOutputMeterDashboardViewFactory(TSharedRef<FSubmixDashboardViewFactory> SubmixDashboard)
+		: SubmixProvider(SubmixDashboard->FindProvider<FSubmixTraceProvider>())
+	{
+	}
+
+	FOutputMeterDashboardViewFactory::~FOutputMeterDashboardViewFactory()
+	{
+		if (TickerHandle.IsValid())
+		{
+			FTSTicker::GetCoreTicker().RemoveTicker(TickerHandle);
+		}
+	}
+
+	EDefaultDashboardTabStack FOutputMeterDashboardViewFactory::GetDefaultTabStack() const
+	{
+		return EDefaultDashboardTabStack::AudioAnalyzerRack;
+	}
+
+	FText FOutputMeterDashboardViewFactory::GetDisplayName() const
+	{
+		return LOCTEXT("AudioDashboard_DashboardsAnalyzerRackTab_DisplayName", "Analyzers");
+	}
+
+	FName FOutputMeterDashboardViewFactory::GetName() const
+	{
+		return "OutputMeter";
+	}
+
+	FSlateIcon FOutputMeterDashboardViewFactory::GetIcon() const
+	{
+		return FSlateStyle::Get().CreateIcon("AudioInsights.Icon");
+	}
+
+	TSharedRef<SWidget> FOutputMeterDashboardViewFactory::MakeWidget(TSharedRef<SDockTab> OwnerTab, const FSpawnTabArgs& SpawnTabArgs)
+	{
+		if (!TickerHandle.IsValid())
+		{
+			TickerHandle = FTSTicker::GetCoreTicker().AddTicker(TEXT("FOutputMeterDashboardViewFactory"), 0.0f, [this](float DeltaTime)
+				{
+					Tick(DeltaTime);
+					return true;
+				});
+		}
+
+		return SNew(SOverlay)
+			+ SOverlay::Slot()
+			[
+				SNew(SColorBlock)
+					.Color(FSlateStyle::Get().GetColor("AudioInsights.Analyzers.BackgroundColor"))
+			]
+			+ SOverlay::Slot()
+			[
+				SNew(SVerticalBox)
+					+ SVerticalBox::Slot()
+					.FillHeight(1.0f)
+					.HAlign(HAlign_Center)
+					.VAlign(VAlign_Fill)
+					[
+						SNew(SAudioMeterWidget)
+							.Orientation(EOrientation::Orient_Vertical)
+							.BackgroundColor(FLinearColor::Transparent)
+							.MeterBackgroundColor(FAudioMeterDefaultColorWidgetStyle::GetDefault().MeterBackgroundColor)
+							.MeterValueColor(FAudioMeterDefaultColorWidgetStyle::GetDefault().MeterValueColor)
+							.MeterPeakColor(FAudioMeterDefaultColorWidgetStyle::GetDefault().MeterPeakColor)
+							.MeterClippingColor(FAudioMeterDefaultColorWidgetStyle::GetDefault().MeterClippingColor)
+							.MeterScaleColor(FAudioMeterDefaultColorWidgetStyle::GetDefault().MeterScaleColor)
+							.MeterScaleLabelColor(FAudioMeterDefaultColorWidgetStyle::GetDefault().MeterScaleLabelColor)
+							.MeterChannelInfo(TAttribute<TArray<FAudioMeterChannelInfo>>::CreateSP(this, &FOutputMeterDashboardViewFactory::GetAudioMeterChannelInfo))
+					]
+			];
+	}
+
+	void FOutputMeterDashboardViewFactory::Tick(float InElapsed)
+	{
+		// Try and subscribe to OnSessionInstanceUpdated events if not already:
+		if (!SessionInstanceUpdatedDelegateHandle.IsValid())
+		{
+			ISessionServicesModule& SessionServicesModule = FModuleManager::LoadModuleChecked<ISessionServicesModule>("SessionServices");
+			if (TSharedPtr<ISessionManager> SessionManager = SessionServicesModule.GetSessionManager())
+			{
+				SessionInstanceUpdatedDelegateHandle = SessionManager->OnSessionInstanceUpdated().AddSP(this, &FOutputMeterDashboardViewFactory::HandleSessionInstanceUpdated);
+			}
+		}
+
+		if (SubmixProvider.IsValid())
+		{
+			const uint64 LastUpdateId = SubmixProvider->GetLastUpdateId();
+			if (ProcessedUpdateId != LastUpdateId)
+			{
+				if (const FSubmixTraceProvider::FDeviceData* DeviceData = SubmixProvider->FindFilteredDeviceData())
+				{
+					ProcessDeviceData(*DeviceData);
+				}
+
+				ProcessedUpdateId = LastUpdateId;
+			}
+		}
+	}
+
+	void FOutputMeterDashboardViewFactory::HandleSessionInstanceUpdated()
+	{
+		if (SubmixDashboardEntry.IsValid())
+		{
+			// As we may have just gotten a valid session, we request starting of the envelope follower:
+			FSubmixDashboardViewFactory::SendSubmixEnvelopeFollowerCVar(true, *SubmixDashboardEntry);
+		}
+	}
+
+	void FOutputMeterDashboardViewFactory::ProcessDeviceData(const FSubmixTraceProvider::FDeviceData& DeviceData)
+	{
+		// Try and find the FSubmixDashboardEntry for the main submix:
+		const auto IsMainSubmixEntryPair = [](const FSubmixTraceProvider::FEntryPair& EntryPair) { return EntryPair.Value->IsMainSubmix(); };
+		const FSubmixTraceProvider::FEntryPair* MainSubmixEntryPair = Algo::FindByPredicate(DeviceData, IsMainSubmixEntryPair);
+		TSharedPtr<FSubmixDashboardEntry> MainSubmixDashboardEntry = (MainSubmixEntryPair) ? MainSubmixEntryPair->Value : nullptr;
+
+		if (SubmixDashboardEntry != MainSubmixDashboardEntry)
+		{
+			SubmixDashboardEntry = MainSubmixDashboardEntry;
+
+			if (SubmixDashboardEntry.IsValid())
+			{
+				// Request starting of envelope follower:
+				FSubmixDashboardViewFactory::SendSubmixEnvelopeFollowerCVar(true, *SubmixDashboardEntry);
+			}
+		}
+	}
+
+	TArray<FAudioMeterChannelInfo> FOutputMeterDashboardViewFactory::GetAudioMeterChannelInfo() const
+	{
+		TArray<FAudioMeterChannelInfo> MeterChannelInfo;
+
+		if (SubmixDashboardEntry.IsValid())
+		{
+			constexpr float MinDecibelValue = -160.0f;
+
+			const FAudioMeterInfo& AudioMeterInfo = *SubmixDashboardEntry->AudioMeterInfo;
+			const TArray<float>& EnvelopeValues = AudioMeterInfo.EnvelopeValues;
+
+			if (!EnvelopeValues.IsEmpty())
+			{
+				MeterChannelInfo.Reserve(EnvelopeValues.Num());
+				for (const float EnvelopeValue : EnvelopeValues)
+				{
+					MeterChannelInfo.Emplace(::Audio::ConvertToDecibels(EnvelopeValue), MinDecibelValue, MinDecibelValue);
+				}
+			}
+			else
+			{
+				MeterChannelInfo.Reserve(AudioMeterInfo.NumChannels);
+				while (MeterChannelInfo.Num() < AudioMeterInfo.NumChannels)
+				{
+					MeterChannelInfo.Emplace(MinDecibelValue, MinDecibelValue, MinDecibelValue);
+				}
+			}
+		}
+
+		return MeterChannelInfo;
+	}
+} // namespace UE::Audio::Insights
+
+#undef LOCTEXT_NAMESPACE
diff --git a/Engine/Plugins/AudioInsights/Source/AudioInsights/Private/Views/OutputMeterDashboardViewFactory.h b/Engine/Plugins/AudioInsights/Source/AudioInsights/Private/Views/OutputMeterDashboardViewFactory.h
new file mode 100644
--- /dev/null
+++ b/Engine/Plugins/AudioInsights/Source/AudioInsights/Private/Views/OutputMeterDashboardViewFactory.h
@@ -0,0 +1,36 @@
+// Copyright Epic Games, Inc. All Rights Reserved.
+#pragma once
+
+#include "AudioMeterChannelInfo.h"
+#include "Containers/Ticker.h"
+#include "Providers/SubmixTraceProvider.h"
+#include "Views/SubmixDashboardViewFactory.h"
+
+namespace UE::Audio::Insights
+{
+	class FOutputMeterDashboardViewFactory : public IDashboardViewFactory, public TSharedFromThis<FOutputMeterDashboardViewFactory>
+	{
+	public:
+		FOutputMeterDashboardViewFactory(TSharedRef<FSubmixDashboardViewFactory> SubmixDashboard);
+		virtual ~FOutputMeterDashboardViewFactory();
+
+		virtual EDefaultDashboardTabStack GetDefaultTabStack() const override;
+		virtual FText GetDisplayName() const override;
+		virtual FName GetName() const override;
+		virtual FSlateIcon GetIcon() const override;
+		virtual TSharedRef<SWidget> MakeWidget(TSharedRef<SDockTab> OwnerTab, const FSpawnTabArgs& SpawnTabArgs) override;
+
+	private:
+		void Tick(float InElapsed);
+		void HandleSessionInstanceUpdated();
+		void ProcessDeviceData(const FSubmixTraceProvider::FDeviceData& DeviceData);
+		TArray<FAudioMeterChannelInfo> GetAudioMeterChannelInfo() const;
+
+		TSharedPtr<const FSubmixTraceProvider> SubmixProvider;
+		FTSTicker::FDelegateHandle TickerHandle;
+		FDelegateHandle SessionInstanceUpdatedDelegateHandle;
+		TOptional<uint64> ProcessedUpdateId;
+		TSharedPtr<FSubmixDashboardEntry> SubmixDashboardEntry;
+	};
+} // namespace UE::Audio::Insights
+
diff --git a/Engine/Plugins/AudioInsights/Source/AudioInsights/Private/Views/SubmixDashboardViewFactory.cpp b/Engine/Plugins/AudioInsights/Source/AudioInsights/Private/Views/SubmixDashboardViewFactory.cpp
--- a/Engine/Plugins/AudioInsights/Source/AudioInsights/Private/Views/SubmixDashboardViewFactory.cpp
+++ b/Engine/Plugins/AudioInsights/Source/AudioInsights/Private/Views/SubmixDashboardViewFactory.cpp
@@ -23,22 +23,6 @@ namespace UE::Audio::Insights
 		{
 			return static_cast<const FSubmixDashboardEntry&>(InData);
 		};
-
-		void SendSubmixEnvelopeFollowerCVar(const bool bInIsSubmixChecked, const FString& InSubmixName)
-		{
-#if WITH_EDITOR
-			const IAudioInsightsModule& AudioInsightsModule = IAudioInsightsModule::GetEditorChecked();
-#else
-			const IAudioInsightsModule& AudioInsightsModule = IAudioInsightsModule::GetChecked();
-#endif // WITH_EDITOR
-
-			const ::Audio::FDeviceId AudioDeviceId = AudioInsightsModule.GetDeviceId();
-
-			const FString SubmixCommandStr = FString::Printf(TEXT("au.%sSubmixEnvelopeFollower_AD%d %s"), bInIsSubmixChecked ? TEXT("Start") : TEXT("Stop"), AudioDeviceId, *InSubmixName);
-			
-			IAudioInsightsTraceModule& AudioInsightsTraceModule = IAudioInsightsModule::GetChecked().GetTraceModule();
-			AudioInsightsTraceModule.ExecuteConsoleCommand(SubmixCommandStr);
-		}
 	}
 
 	FSubmixDashboardViewFactory::FSubmixDashboardViewFactory()
@@ -134,7 +118,7 @@ namespace UE::Audio::Insights
 
 			if (bIsSubmixChecked)
 			{
-				SendSubmixEnvelopeFollowerCVar(bIsSubmixChecked, SubmixDashboardEntry.GetDisplayName().ToString());
+				SendSubmixEnvelopeFollowerCVar(bIsSubmixChecked, SubmixDashboardEntry);
 				FAudioMeterProvider::OnUpdateAudioMeterInfo.ExecuteIfBound(SubmixId, SubmixDashboardEntry.AudioMeterInfo, SubmixDashboardEntry.GetDisplayName());
 			}
 
@@ -162,7 +146,7 @@ namespace UE::Audio::Insights
 							*FoundSubmixCheckedState = bIsChecked;
 						}
 
-						SendSubmixEnvelopeFollowerCVar(bIsChecked, SubmixDashboardEntry.GetDisplayName().ToString());
+						SendSubmixEnvelopeFollowerCVar(bIsChecked, SubmixDashboardEntry);
 
 						if (bIsChecked)
 						{
@@ -263,6 +247,31 @@ namespace UE::Audio::Insights
 		return DashboardWidget->AsShared();
 	}
 
+	void FSubmixDashboardViewFactory::SendSubmixEnvelopeFollowerCVar(const bool bEnableEnvelopeFollower, const FSubmixDashboardEntry& SubmixDashboardEntry)
+	{
+#if !WITH_EDITOR
+		if (SubmixDashboardEntry.IsMainSubmix() && !bEnableEnvelopeFollower)
+		{
+			// Don't disable envelope follower on main submix.
+			return;
+		}
+#endif // !WITH_EDITOR
+
+		const FString InSubmixName = SubmixDashboardEntry.GetDisplayName().ToString();
+#if WITH_EDITOR
+		const IAudioInsightsModule& AudioInsightsModule = IAudioInsightsModule::GetEditorChecked();
+#else
+		const IAudioInsightsModule& AudioInsightsModule = IAudioInsightsModule::GetChecked();
+#endif // WITH_EDITOR
+
+		const ::Audio::FDeviceId AudioDeviceId = AudioInsightsModule.GetDeviceId();
+
+		const FString SubmixCommandStr = FString::Printf(TEXT("au.%sSubmixEnvelopeFollower_AD%d %s"), bEnableEnvelopeFollower ? TEXT("Start") : TEXT("Stop"), AudioDeviceId, *InSubmixName);
+
+		IAudioInsightsTraceModule& AudioInsightsTraceModule = IAudioInsightsModule::GetChecked().GetTraceModule();
+		AudioInsightsTraceModule.ExecuteConsoleCommand(SubmixCommandStr);
+	}
+
 	const TMap<FName, FTraceTableDashboardViewFactory::FColumnData>& FSubmixDashboardViewFactory::GetColumns() const
 	{
 		using namespace FSubmixDashboardViewFactoryPrivate;
diff --git a/Engine/Plugins/AudioInsights/Source/AudioInsights/Public/Messages/SubmixTraceMessages.h b/Engine/Plugins/AudioInsights/Source/AudioInsights/Public/Messages/SubmixTraceMessages.h
--- a/Engine/Plugins/AudioInsights/Source/AudioInsights/Public/Messages/SubmixTraceMessages.h
+++ b/Engine/Plugins/AudioInsights/Source/AudioInsights/Public/Messages/SubmixTraceMessages.h
@@ -3,6 +3,9 @@
 
 #include "Providers/AudioMeterProvider.h"
 #include "Messages/AnalyzerMessageQueue.h"
+#if WITH_EDITOR
+#include "Sound/AudioSettings.h"
+#endif // WITH_EDITOR
 #include "Trace/Analyzer.h"
 #include "Views/TableDashboardViewFactory.h"
 
@@ -42,10 +45,12 @@ namespace UE::Audio::Insights
 
 			EventData.GetString("Name", Name);
 			NumChannels = EventData.GetValue<int32>("NumChannels");
+			bIsMainSubmix = EventData.GetValue<bool>("IsMainSubmix");
 		}
 
 		FString Name;
 		int32 NumChannels = 0;
+		bool bIsMainSubmix = false;
 	};
 
 	struct FSubmixHasActivityMessage : public FSubmixMessageBase
@@ -127,6 +132,15 @@ namespace UE::Audio::Insights
 		virtual UObject* GetObject() override { return FSoftObjectPath(Name).ResolveObject(); }
 		virtual bool IsValid() const override { return SubmixId != static_cast<uint32>(INDEX_NONE); }
 
+		bool IsMainSubmix() const
+		{
+#if !WITH_EDITOR
+			return bIsMainSubmix;
+#else
+			return GetDefault<UAudioSettings>() && (FSoftObjectPath(Name) == GetDefault<UAudioSettings>()->MasterSubmix);
+#endif
+		}
+
 		uint32 SubmixId = INDEX_NONE;
 		double Timestamp = 0.0;
 
@@ -134,5 +148,8 @@ namespace UE::Audio::Insights
 
 		FString Name;
 		bool bHasActivity = false;
+#if !WITH_EDITOR
+		bool bIsMainSubmix = false;
+#endif // !WITH_EDITOR
 	};
 } // namespace UE::Audio::Insights
diff --git a/Engine/Plugins/AudioInsights/Source/AudioInsights/Public/Views/SubmixDashboardViewFactory.h b/Engine/Plugins/AudioInsights/Source/AudioInsights/Public/Views/SubmixDashboardViewFactory.h
--- a/Engine/Plugins/AudioInsights/Source/AudioInsights/Public/Views/SubmixDashboardViewFactory.h
+++ b/Engine/Plugins/AudioInsights/Source/AudioInsights/Public/Views/SubmixDashboardViewFactory.h
@@ -8,6 +8,7 @@
 
 namespace UE::Audio::Insights
 {
+	struct FSubmixDashboardEntry;
 	class FSubmixTraceProvider;
 
 	class FSubmixDashboardViewFactory : public FTraceObjectTableDashboardViewFactory
@@ -22,6 +23,8 @@ namespace UE::Audio::Insights
 		UE_API virtual EDefaultDashboardTabStack GetDefaultTabStack() const override;
 		UE_API virtual TSharedRef<SWidget> MakeWidget(TSharedRef<SDockTab> OwnerTab, const FSpawnTabArgs& SpawnTabArgs) override;
 
+		UE_API static void SendSubmixEnvelopeFollowerCVar(const bool bEnableEnvelopeFollower, const FSubmixDashboardEntry& SubmixDashboardEntry);
+
 	protected:
 		UE_API virtual TSharedRef<SWidget> GenerateWidgetForColumn(TSharedRef<IDashboardDataViewEntry> InRowData, const FName& InColumnName) override;
 		UE_API virtual void ProcessEntries(FTraceTableDashboardViewFactory::EProcessReason Reason) override;
diff --git a/Engine/Source/Runtime/AudioMixer/Private/AudioMixerDevice.cpp b/Engine/Source/Runtime/AudioMixer/Private/AudioMixerDevice.cpp
--- a/Engine/Source/Runtime/AudioMixer/Private/AudioMixerDevice.cpp
+++ b/Engine/Source/Runtime/AudioMixer/Private/AudioMixerDevice.cpp
@@ -134,6 +134,7 @@ UE_TRACE_EVENT_BEGIN(Audio, SubmixAlivePing)
 	UE_TRACE_EVENT_FIELD(uint64, Timestamp)
 	UE_TRACE_EVENT_FIELD(UE::Trace::WideString, Name)
 	UE_TRACE_EVENT_FIELD(int32, NumChannels)
+	UE_TRACE_EVENT_FIELD(bool, IsMainSubmix)
 UE_TRACE_EVENT_END()
 
 #endif // UE_AUDIO_PROFILERTRACE_ENABLED
@@ -397,7 +398,8 @@ namespace Audio
 					<< SubmixAlivePing.SubmixId(GetTypeHash(InSubmix->GetPathName()))
 					<< SubmixAlivePing.Timestamp(FPlatformTime::Cycles64())
 					<< SubmixAlivePing.Name(*InSubmix->GetPathName())
-					<< SubmixAlivePing.NumChannels(InSubmix->NumChannels);
+					<< SubmixAlivePing.NumChannels(InSubmix->NumChannels)
+					<< SubmixAlivePing.IsMainSubmix(InSubmix == GetMainSubmix());
 			};
 
 			Submixes.Iterate([this, &TraceSubmixLoadedMessage](const FSubmixMap::FPair& SubmixPair)
diff --git a/Engine/Source/Runtime/AudioMixer/Private/AudioMixerSubmix.cpp b/Engine/Source/Runtime/AudioMixer/Private/AudioMixerSubmix.cpp
--- a/Engine/Source/Runtime/AudioMixer/Private/AudioMixerSubmix.cpp
+++ b/Engine/Source/Runtime/AudioMixer/Private/AudioMixerSubmix.cpp
@@ -65,6 +65,7 @@ UE_TRACE_EVENT_BEGIN(Audio, SubmixLoaded)
 	UE_TRACE_EVENT_FIELD(uint64, Timestamp)
 	UE_TRACE_EVENT_FIELD(UE::Trace::WideString, Name)
 	UE_TRACE_EVENT_FIELD(int32, NumChannels)
+	UE_TRACE_EVENT_FIELD(bool, IsMainSubmix)
 UE_TRACE_EVENT_END()
 
 UE_TRACE_EVENT_BEGIN(Audio, SubmixHasActivity)
@@ -285,7 +286,8 @@ namespace Audio
 				<< SubmixLoaded.SubmixId(GetTypeHash(OwningSubmixObject->GetPathName()))
 				<< SubmixLoaded.Timestamp(FPlatformTime::Cycles64())
 				<< SubmixLoaded.Name(*OwningSubmixObject->GetPathName())
-				<< SubmixLoaded.NumChannels(NumChannels);
+				<< SubmixLoaded.NumChannels(NumChannels)
+				<< SubmixLoaded.IsMainSubmix(InSoundSubmix == &MixerDevice->GetMainSubmixObject());
 		}
 #endif // UE_AUDIO_PROFILERTRACE_ENABLED
 	}
