[Audio Insights] Disable horizontal scroll box on Event Log

diff --git a/Engine/Plugins/AudioInsights/Source/AudioInsights/Private/Views/TableDashboardViewFactory.cpp b/Engine/Plugins/AudioInsights/Source/AudioInsights/Private/Views/TableDashboardViewFactory.cpp
--- a/Engine/Plugins/AudioInsights/Source/AudioInsights/Private/Views/TableDashboardViewFactory.cpp
+++ b/Engine/Plugins/AudioInsights/Source/AudioInsights/Private/Views/TableDashboardViewFactory.cpp
@@ -116,6 +116,25 @@ namespace UE::Audio::Insights
 		RequestSort();
 	}
 
+	TSharedRef<SWidget> FTraceTableDashboardViewFactory::CreateHorizontalScrollBox()
+	{
+		if (!FilteredEntriesListView.IsValid())
+		{
+			return SNullWidget::NullWidget;
+		}
+
+		return SNew(SScrollBox)
+				.Orientation(EOrientation::Orient_Horizontal)
+			
+				+ SScrollBox::Slot()
+				.Padding(0)
+				.FillSize(1.0f)
+				.HAlign(HAlign_Fill)
+				[
+					FilteredEntriesListView->AsShared()
+				];
+	}
+
 	TSharedRef<SWidget> FTraceTableDashboardViewFactory::CreateSettingsButtonWidget()
 	{
 		// Override OnGetSettingsMenuContent in derived class to enable the settings menu widget
@@ -278,15 +297,7 @@ namespace UE::Audio::Insights
 			.FillHeight(1.0f)
 			.Padding(0, 2)
 			[
-				SNew(SScrollBox)
-				.Orientation(EOrientation::Orient_Horizontal)
-				+ SScrollBox::Slot()
-				.Padding(0)
-				.FillSize(1.0f)
-				.HAlign(HAlign_Fill)
-				[
-					FilteredEntriesListView->AsShared()
-				]
+				EnableHorizontalScrollBox() ? CreateHorizontalScrollBox() : FilteredEntriesListView->AsShared()
 			];
 		}
 
diff --git a/Engine/Plugins/AudioInsights/Source/AudioInsights/Public/Views/AudioEventLogDashboardViewFactory.h b/Engine/Plugins/AudioInsights/Source/AudioInsights/Public/Views/AudioEventLogDashboardViewFactory.h
--- a/Engine/Plugins/AudioInsights/Source/AudioInsights/Public/Views/AudioEventLogDashboardViewFactory.h
+++ b/Engine/Plugins/AudioInsights/Source/AudioInsights/Public/Views/AudioEventLogDashboardViewFactory.h
@@ -82,6 +82,7 @@ namespace UE::Audio::Insights
 
 		UE_API virtual void OnSelectionChanged(TSharedPtr<IDashboardDataViewEntry> SelectedItem, ESelectInfo::Type SelectInfo) override;
 		virtual bool ClearSelectionOnClick() const override { return true; }
+		virtual bool EnableHorizontalScrollBox() const override { return false; }
 		UE_API virtual const FTableRowStyle* GetRowStyle() const override;
 
 #if WITH_EDITOR
diff --git a/Engine/Plugins/AudioInsights/Source/AudioInsights/Public/Views/TableDashboardViewFactory.h b/Engine/Plugins/AudioInsights/Source/AudioInsights/Public/Views/TableDashboardViewFactory.h
--- a/Engine/Plugins/AudioInsights/Source/AudioInsights/Public/Views/TableDashboardViewFactory.h
+++ b/Engine/Plugins/AudioInsights/Source/AudioInsights/Public/Views/TableDashboardViewFactory.h
@@ -159,6 +159,8 @@ namespace UE::Audio::Insights
 		UE_API virtual void OnSelectionChanged(TSharedPtr<IDashboardDataViewEntry> SelectedItem, ESelectInfo::Type SelectInfo);
 		virtual bool ClearSelectionOnClick() const { return false; }
 
+		virtual bool EnableHorizontalScrollBox() const { return true; }
+
 		virtual const FTableRowStyle* GetRowStyle() const { return nullptr; }
 		UE_API virtual FSlateColor GetRowColor(const TSharedPtr<IDashboardDataViewEntry>& InRowDataPtr);
 
@@ -186,6 +188,8 @@ namespace UE::Audio::Insights
 		UE_API void RequestSort();
 		UE_API void OnColumnSortModeChanged(const EColumnSortPriority::Type InSortPriority, const FName& InColumnId, const EColumnSortMode::Type InSortMode);
 
+		TSharedRef<SWidget> CreateHorizontalScrollBox();
+
 		TSharedPtr<SSearchBox> SearchBoxWidget;
 		FText SearchBoxFilterText;
 	};
