[Audio Insights] Display the Cache size and duration in the event log

diff --git a/Engine/Plugins/AudioInsights/Source/AudioInsights/Private/Cache/AudioInsightsCacheManager.cpp b/Engine/Plugins/AudioInsights/Source/AudioInsights/Private/Cache/AudioInsightsCacheManager.cpp
--- a/Engine/Plugins/AudioInsights/Source/AudioInsights/Private/Cache/AudioInsightsCacheManager.cpp
+++ b/Engine/Plugins/AudioInsights/Source/AudioInsights/Private/Cache/AudioInsightsCacheManager.cpp
@@ -74,6 +74,22 @@ namespace UE::Audio::Insights
 		return Cache[WriteIndex]->GetChunkTimeRangeEnd();
 	}
 
+	uint32 FAudioInsightsCacheManager::GetUsedCacheSize() const
+	{
+		uint32 UsedSpace = 0u;
+		for (uint32 Index = 0u; Index < Cache.Capacity(); ++Index)
+		{
+			UsedSpace += Cache[Index]->GetCurrentChunkSize();
+		}
+
+		return UsedSpace;
+	}
+
+	uint32 FAudioInsightsCacheManager::GetMaxCacheSize() const
+	{
+		return CacheManagerPrivate::MaxChunkSize * CacheManagerPrivate::NumChunks;
+	}
+
 	void FAudioInsightsCacheManager::CreateCache()
 	{
 		for (uint32 Index = 0u; Index < Cache.Capacity(); ++Index)
diff --git a/Engine/Plugins/AudioInsights/Source/AudioInsights/Private/Views/AudioEventLogDashboardViewFactory.cpp b/Engine/Plugins/AudioInsights/Source/AudioInsights/Private/Views/AudioEventLogDashboardViewFactory.cpp
--- a/Engine/Plugins/AudioInsights/Source/AudioInsights/Private/Views/AudioEventLogDashboardViewFactory.cpp
+++ b/Engine/Plugins/AudioInsights/Source/AudioInsights/Private/Views/AudioEventLogDashboardViewFactory.cpp
@@ -477,6 +477,86 @@ namespace UE::Audio::Insights
 		return FilterBar.ToSharedRef();
 	}
 
+	TSharedRef<SWidget> FAudioEventLogDashboardViewFactory::CreateSettingsButtonWidget()
+	{
+		if (!SettingsAreaWidget.IsValid())
+		{
+#if WITH_EDITOR
+			constexpr FLinearColor StatTitleColor(1.0f, 1.0f, 1.0f, 0.5f);
+#endif // WITH_EDITOR
+
+			SAssignNew(SettingsAreaWidget, SHorizontalBox)
+			
+#if WITH_EDITOR
+			// Cache stats
+			+ SHorizontalBox::Slot()
+			.VAlign(VAlign_Center)
+			.AutoWidth()
+			[
+				SNew(STextBlock)
+				.Justification(ETextJustify::Center)
+				.ColorAndOpacity(StatTitleColor)
+				.Text(LOCTEXT("AudioEventLog_CacheStats_Title", "Cache Size: "))
+			]
+
+			+ SHorizontalBox::Slot()
+			.VAlign(VAlign_Center)
+			.AutoWidth()
+			[
+				SNew(STextBlock)
+				.Justification(ETextJustify::Center)
+				.Text_Lambda([]()
+				{
+					FAudioInsightsCacheManager& CacheManager = FAudioInsightsModule::GetChecked().GetCacheManager();
+					const float CacheSizeMb = static_cast<float>(CacheManager.GetUsedCacheSize()) / 1000000.0f;
+					const float MaxCacheSizeMb = static_cast<float>(CacheManager.GetMaxCacheSize()) / 1000000.0f;
+
+					return FText::Format(LOCTEXT("AudioEventLog_CacheStats_Memory", "{0} / {1}")
+										, FText::AsNumber(CacheSizeMb, FSlateStyle::Get().GetMemoryFormat())
+										, FSlateStyle::Get().FormatMemoryAsMegabytes(MaxCacheSizeMb));
+				})
+			]
+
+			+ SHorizontalBox::Slot()
+			.VAlign(VAlign_Center)
+			.AutoWidth()
+			[
+				SNew(STextBlock)
+				.Justification(ETextJustify::Center)
+				.ColorAndOpacity(StatTitleColor)
+				.Text(LOCTEXT("AudioEventLog_CacheStats_Duration", " : Duration: "))
+			]
+
+			+ SHorizontalBox::Slot()
+			.VAlign(VAlign_Center)
+			.AutoWidth()
+			[
+				SNew(STextBlock)
+				.Justification(ETextJustify::Center)
+				.Text_Lambda([]()
+				{
+					FAudioInsightsCacheManager& CacheManager = FAudioInsightsModule::GetChecked().GetCacheManager();
+					const float CacheDuration = CacheManager.GetCacheDuration();
+
+					return FText::Format(LOCTEXT("AudioEventLog_CacheStats_DurationValue", "{0}s")
+										, FText::AsNumber(CacheDuration, FSlateStyle::Get().GetShortTimeFormat()));
+				})
+			]
+#endif // WITH_EDITOR
+
+			// Settings button
+			+ SHorizontalBox::Slot()
+			.AutoWidth()
+			.Padding(2.0f, 0.0f, 0.0f, 0.0f)
+			.HAlign(HAlign_Right)
+			[
+				FTraceObjectTableDashboardViewFactory::CreateSettingsButtonWidget()
+			];
+		}
+
+		return SettingsAreaWidget.ToSharedRef();
+	}
+
 	TSharedRef<SWidget> FAudioEventLogDashboardViewFactory::OnGetSettingsMenuContent()
 	{
 		FMenuBuilder MenuBuilder(false /*bShouldCloseWindowAfterMenuSelection*/, CommandList);
diff --git a/Engine/Plugins/AudioInsights/Source/AudioInsights/Public/AudioInsightsStyle.h b/Engine/Plugins/AudioInsights/Source/AudioInsights/Public/AudioInsightsStyle.h
--- a/Engine/Plugins/AudioInsights/Source/AudioInsights/Public/AudioInsightsStyle.h
+++ b/Engine/Plugins/AudioInsights/Source/AudioInsights/Public/AudioInsightsStyle.h
@@ -83,6 +83,15 @@ namespace UE::Audio::Insights
 			return &FloatFormat;
 		};
 
+		const FNumberFormattingOptions* GetShortTimeFormat() const
+		{
+			static FNumberFormattingOptions FloatFormat;
+			FloatFormat.MinimumIntegralDigits = 1;
+			FloatFormat.MinimumFractionalDigits = 1;
+			FloatFormat.MaximumFractionalDigits = 1;
+			return &FloatFormat;
+		};
+
 		const FNumberFormattingOptions* GetRelativeRenderCostFormat() const
 		{
 			static FNumberFormattingOptions FloatFormat;
@@ -93,6 +102,15 @@ namespace UE::Audio::Insights
 			return &FloatFormat;
 		};
 
+		const FNumberFormattingOptions* GetMemoryFormat() const
+		{
+			static FNumberFormattingOptions FloatFormat;
+			FloatFormat.MinimumIntegralDigits = 1;
+			FloatFormat.MinimumFractionalDigits = 1;
+			FloatFormat.MaximumFractionalDigits = 1;
+			return &FloatFormat;
+		};
+
 		FText FormatSecondsAsTime(float InTimeSec) const
 		{
 			return FText::Format(LOCTEXT("TimeInSecondsFormat", "{0}s"), FText::AsNumber(InTimeSec, GetTimeFormat()));
@@ -103,6 +121,11 @@ namespace UE::Audio::Insights
 			return FText::Format(LOCTEXT("TimeInMillisecondsFormat", "{0}ms"), FText::AsNumber(InTimeMS, GetTimeFormat()));
 		}
 
+		FText FormatMemoryAsMegabytes(const float InMemoryMB) const
+		{
+			return FText::Format(LOCTEXT("MemoryInMegabytesFormat", "{0}mb"), FText::AsNumber(InMemoryMB, GetMemoryFormat()));
+		}
+
 		FSlateIcon CreateIcon(FName InName)  const
 		{
 			return { GetStyleName(), InName };
diff --git a/Engine/Plugins/AudioInsights/Source/AudioInsights/Public/Cache/AudioInsightsCacheManager.h b/Engine/Plugins/AudioInsights/Source/AudioInsights/Public/Cache/AudioInsightsCacheManager.h
--- a/Engine/Plugins/AudioInsights/Source/AudioInsights/Public/Cache/AudioInsightsCacheManager.h
+++ b/Engine/Plugins/AudioInsights/Source/AudioInsights/Public/Cache/AudioInsightsCacheManager.h
@@ -23,6 +23,9 @@ namespace UE::Audio::Insights
 		UE_API double GetCacheStartTimeStamp() const;
 		UE_API double GetCacheEndTimeStamp() const;
 
+		UE_API uint32 GetUsedCacheSize() const;
+		UE_API uint32 GetMaxCacheSize() const;
+
 		// Find the closest message to a timestamp in the cache for a specific message type
 		// Will look backwards from the chosen timestamp until a message is found
 		// @param ID: When overloaded, the search algorithm will only return messages that have a matching ID.
diff --git a/Engine/Plugins/AudioInsights/Source/AudioInsights/Public/Views/AudioEventLogDashboardViewFactory.h b/Engine/Plugins/AudioInsights/Source/AudioInsights/Public/Views/AudioEventLogDashboardViewFactory.h
--- a/Engine/Plugins/AudioInsights/Source/AudioInsights/Public/Views/AudioEventLogDashboardViewFactory.h
+++ b/Engine/Plugins/AudioInsights/Source/AudioInsights/Public/Views/AudioEventLogDashboardViewFactory.h
@@ -73,6 +73,7 @@ namespace UE::Audio::Insights
 		UE_API virtual TSharedPtr<SWidget> CreateFilterBarButtonWidget() override;
 		UE_API virtual TSharedRef<SWidget> CreateFilterBarWidget() override;
 
+		UE_API virtual TSharedRef<SWidget> CreateSettingsButtonWidget() override;
 		UE_API virtual TSharedRef<SWidget> OnGetSettingsMenuContent() override;
 		UE_API virtual FText OnGetSettingsMenuToolTip() override;
 
@@ -113,6 +114,8 @@ namespace UE::Audio::Insights
 		TSharedPtr<SWidget> DashboardWidget;
 		TSharedPtr<SAudioFilterBar<FEventLogFilterID>> FilterBar;
 
+		TSharedPtr<SWidget> SettingsAreaWidget;
+
 		TMap<FString, TSharedPtr<FFilterCategory>> FilterCategories;
 		TMap<FString, FEventLogFilterID> EventFilterTypes;
 
