[Audio Insights] Limit what displays in the Event Log by what is inside the cache

diff --git a/Engine/Plugins/AudioInsights/Source/AudioInsights/Private/Cache/AudioInsightsCacheManager.cpp b/Engine/Plugins/AudioInsights/Source/AudioInsights/Private/Cache/AudioInsightsCacheManager.cpp
--- a/Engine/Plugins/AudioInsights/Source/AudioInsights/Private/Cache/AudioInsightsCacheManager.cpp
+++ b/Engine/Plugins/AudioInsights/Source/AudioInsights/Private/Cache/AudioInsightsCacheManager.cpp
@@ -41,6 +41,8 @@ namespace UE::Audio::Insights
 			}
 
 			NextChunk.AddMessageToChunk(Message);
+
+			OnChunkOverwritten.Broadcast(GetCacheStartTimeStamp());
 		}
 		else
 		{
@@ -90,6 +92,11 @@ namespace UE::Audio::Insights
 		return CacheManagerPrivate::MaxChunkSize * CacheManagerPrivate::NumChunks;
 	}
 
+	uint32 FAudioInsightsCacheManager::GetNumChunks() const
+	{
+		return CacheManagerPrivate::NumChunks;
+	}
+
 	void FAudioInsightsCacheManager::CreateCache()
 	{
 		for (uint32 Index = 0u; Index < Cache.Capacity(); ++Index)
diff --git a/Engine/Plugins/AudioInsights/Source/AudioInsights/Private/Providers/AudioEventLogTraceProvider.cpp b/Engine/Plugins/AudioInsights/Source/AudioInsights/Private/Providers/AudioEventLogTraceProvider.cpp
--- a/Engine/Plugins/AudioInsights/Source/AudioInsights/Private/Providers/AudioEventLogTraceProvider.cpp
+++ b/Engine/Plugins/AudioInsights/Source/AudioInsights/Private/Providers/AudioEventLogTraceProvider.cpp
@@ -17,6 +17,26 @@
 
 namespace UE::Audio::Insights
 {
+	FAudioEventLogTraceProvider::FAudioEventLogTraceProvider()
+		: TDeviceDataMapTraceProvider<uint32, TSharedPtr<FAudioEventLogDashboardEntry>>(GetName_Static())
+	{
+#if WITH_EDITOR
+		FAudioInsightsCacheManager& CacheManager = IAudioInsightsModule::GetChecked().GetCacheManager();
+		CacheManager.OnChunkOverwritten.AddRaw(this, &FAudioEventLogTraceProvider::OnCacheChunkOverwritten);
+#endif // WITH_EDITOR
+	}
+
+	FAudioEventLogTraceProvider::~FAudioEventLogTraceProvider()
+	{
+#if WITH_EDITOR
+		if (FModuleManager::Get().IsModuleLoaded("AudioInsights"))
+		{
+			FAudioInsightsCacheManager& CacheManager = IAudioInsightsModule::GetChecked().GetCacheManager();
+			CacheManager.OnChunkOverwritten.RemoveAll(this);
+		}
+#endif // WITH_EDITOR
+	}
+
 	UE::Trace::IAnalyzer* FAudioEventLogTraceProvider::ConstructAnalyzer(TraceServices::IAnalysisSession& InSession)
 	{
 		class FAudioEventLogTraceAnalyzer : public FTraceAnalyzerBase
@@ -93,6 +113,9 @@ namespace UE::Audio::Insights
 
 		PlayOrderToActorInfoMap.Empty();
 
+		CacheStartTimestamp = 0.0;
+		bCacheChunksUpdated = false;
+
 		return new FAudioEventLogTraceAnalyzer(AsShared(), InSession);
 	}
 
@@ -136,12 +159,57 @@ namespace UE::Audio::Insights
 		ProcessMessageQueue<FAudioEventLogMessage>(TraceMessages.AudioEventLogMessages, AddEntryFunc,
 		[](const FAudioEventLogMessage& Msg, TSharedPtr<FAudioEventLogDashboardEntry>* OutEntry)
 		{
-			// Do nothing - this provider does not clear or alter entries once they are in the log
+			// Do nothing
 		});
 
 		return true;
 	}
 
+	bool FAudioEventLogTraceProvider::ProcessManuallyUpdatedEntries()
+	{
+#if WITH_EDITOR
+		if (!bCacheChunksUpdated)
+		{
+			return false;
+		}
+
+		FAudioInsightsCacheManager& CacheManager = IAudioInsightsModule::GetChecked().GetCacheManager();
+		const uint32 HalfNumChunks = CacheManager.GetNumChunks() > 1 ? CacheManager.GetNumChunks() / 2 : 1;
+
+		bool bAnyEntriesRemoved = false;
+		for (auto& [DeviceID, MessageIDToDashboardEntryMap] : DeviceDataMap)
+		{
+			TArray<uint32> MessagesToRemove;
+			MessagesToRemove.Reserve(MessageIDToDashboardEntryMap.Num() / HalfNumChunks); // Reserve 2x average num messages per chunk to give us some wiggle room
+
+			for (auto& [MessageID, DashboardEntry] : MessageIDToDashboardEntryMap)
+			{
+				if (DashboardEntry->Timestamp < CacheStartTimestamp)
+				{
+					MessagesToRemove.Add(MessageID);
+				}
+				else
+				{
+					break;
+				}
+			}
+
+			for (const uint32 MessageID : MessagesToRemove)
+			{
+				MessageIDToDashboardEntryMap.Remove(MessageID);
+			}
+
+			bAnyEntriesRemoved = bAnyEntriesRemoved || MessagesToRemove.Num() > 0;
+		}
+
+		bCacheChunksUpdated = false;
+
+		return bAnyEntriesRemoved;
+#else
+		return false;
+#endif // WITH_EDITOR
+	}
+
 	TSharedPtr<FAudioEventLogDashboardEntry> FAudioEventLogTraceProvider::CreateNewEntry(const FAudioEventLogMessage& Msg)
 	{
 		TSharedPtr<FAudioEventLogDashboardEntry> Entry = MakeShared<FAudioEventLogDashboardEntry>();
@@ -190,6 +258,14 @@ namespace UE::Audio::Insights
 
 		return Entry;
 	}
+
+#if WITH_EDITOR
+	void FAudioEventLogTraceProvider::OnCacheChunkOverwritten(const double NewCacheStartTimestamp)
+	{
+		CacheStartTimestamp = NewCacheStartTimestamp;
+		bCacheChunksUpdated = true;
+	}
+#endif // WITH_EDITOR
 } // namespace UE::Audio::Insights
 
 #undef LOCTEXT_NAMESPACE
\ No newline at end of file
diff --git a/Engine/Plugins/AudioInsights/Source/AudioInsights/Private/Providers/AudioEventLogTraceProvider.h b/Engine/Plugins/AudioInsights/Source/AudioInsights/Private/Providers/AudioEventLogTraceProvider.h
--- a/Engine/Plugins/AudioInsights/Source/AudioInsights/Private/Providers/AudioEventLogTraceProvider.h
+++ b/Engine/Plugins/AudioInsights/Source/AudioInsights/Private/Providers/AudioEventLogTraceProvider.h
@@ -11,21 +11,23 @@ namespace UE::Audio::Insights
 		, public TSharedFromThis<FAudioEventLogTraceProvider>
 	{
 	public:
-		FAudioEventLogTraceProvider()
-			: TDeviceDataMapTraceProvider<uint32, TSharedPtr<FAudioEventLogDashboardEntry>>(GetName_Static())
-		{
-		}
+		FAudioEventLogTraceProvider();
+		virtual ~FAudioEventLogTraceProvider();
 
-		virtual ~FAudioEventLogTraceProvider() = default;
 		virtual UE::Trace::IAnalyzer* ConstructAnalyzer(TraceServices::IAnalysisSession& InSession) override;
 
 		static FName GetName_Static();
 
 	private:
 		virtual bool ProcessMessages() override;
+		virtual bool ProcessManuallyUpdatedEntries() override;
 
 		TSharedPtr<FAudioEventLogDashboardEntry> CreateNewEntry(const FAudioEventLogMessage& Msg);
 
+#if WITH_EDITOR
+		void OnCacheChunkOverwritten(const double NewCacheStartTimestamp);
+#endif // WITH_EDITOR
+
 		FAudioEventLogMessages TraceMessages;
 
 		struct FActorInformation
@@ -35,5 +37,8 @@ namespace UE::Audio::Insights
 		};
 
 		TMap<uint32, FActorInformation> PlayOrderToActorInfoMap;
+
+		double CacheStartTimestamp = 0.0;
+		bool bCacheChunksUpdated = false;
 	};
 } // namespace UE::Audio::Insights
diff --git a/Engine/Plugins/AudioInsights/Source/AudioInsights/Public/Cache/AudioInsightsCacheManager.h b/Engine/Plugins/AudioInsights/Source/AudioInsights/Public/Cache/AudioInsightsCacheManager.h
--- a/Engine/Plugins/AudioInsights/Source/AudioInsights/Public/Cache/AudioInsightsCacheManager.h
+++ b/Engine/Plugins/AudioInsights/Source/AudioInsights/Public/Cache/AudioInsightsCacheManager.h
@@ -25,6 +25,7 @@ namespace UE::Audio::Insights
 
 		UE_API uint32 GetUsedCacheSize() const;
 		UE_API uint32 GetMaxCacheSize() const;
+		UE_API uint32 GetNumChunks() const;
 
 		// Find the closest message to a timestamp in the cache for a specific message type
 		// Will look backwards from the chosen timestamp until a message is found
@@ -168,6 +169,9 @@ namespace UE::Audio::Insights
 			IterateOverRange<T>(MessageID, Start, GetCacheEndTimeStamp(), OnMessageFunc, ID);
 		}
 
+		DECLARE_MULTICAST_DELEGATE_OneParam(FOnChunkOverwritten, double /*NewStartTimestamp*/);
+		FOnChunkOverwritten OnChunkOverwritten;
+
 	private:
 		void CreateCache();
 
