[Audio Insights] Visualize the lifetime of an event log row in the cache

diff --git a/Engine/Plugins/AudioInsights/Source/AudioInsights/Private/Cache/AudioInsightsCacheManager.cpp b/Engine/Plugins/AudioInsights/Source/AudioInsights/Private/Cache/AudioInsightsCacheManager.cpp
--- a/Engine/Plugins/AudioInsights/Source/AudioInsights/Private/Cache/AudioInsightsCacheManager.cpp
+++ b/Engine/Plugins/AudioInsights/Source/AudioInsights/Private/Cache/AudioInsightsCacheManager.cpp
@@ -76,6 +76,28 @@ namespace UE::Audio::Insights
 		return Cache[WriteIndex]->GetChunkTimeRangeEnd();
 	}
 
+	const FAudioCachedMessageChunk* const FAudioInsightsCacheManager::GetChunk(const uint32 NumChunksFromStart) const
+	{
+		ensure(NumChunksFromStart < Cache.Capacity());
+		if (NumChunksFromStart >= Cache.Capacity())
+		{
+			return nullptr;
+		}
+
+		uint32 ReadIndex = StartIndex;
+		for (uint32 Index = 0u; Index < NumChunksFromStart; ++Index)
+		{
+			ReadIndex = Cache.GetNextIndex(ReadIndex);
+		}
+
+		if (!Cache[ReadIndex]->HasAnyData())
+		{
+			return nullptr;
+		}
+
+		return Cache[ReadIndex].Get();
+	}
+
 	uint32 FAudioInsightsCacheManager::GetUsedCacheSize() const
 	{
 		uint32 UsedSpace = 0u;
@@ -97,6 +119,20 @@ namespace UE::Audio::Insights
 		return CacheManagerPrivate::NumChunks;
 	}
 
+	uint32 FAudioInsightsCacheManager::GetNumUsedChunks() const
+	{
+		uint32 NumUsedChunks = 0u;
+		for (uint32 Index = 0u; Index < Cache.Capacity(); ++Index)
+		{
+			if (Cache[Index]->HasAnyData())
+			{
+				++NumUsedChunks;
+			}
+		}
+
+		return NumUsedChunks;
+	}
+
 	void FAudioInsightsCacheManager::CreateCache()
 	{
 		for (uint32 Index = 0u; Index < Cache.Capacity(); ++Index)
diff --git a/Engine/Plugins/AudioInsights/Source/AudioInsights/Private/Providers/AudioEventLogTraceProvider.cpp b/Engine/Plugins/AudioInsights/Source/AudioInsights/Private/Providers/AudioEventLogTraceProvider.cpp
--- a/Engine/Plugins/AudioInsights/Source/AudioInsights/Private/Providers/AudioEventLogTraceProvider.cpp
+++ b/Engine/Plugins/AudioInsights/Source/AudioInsights/Private/Providers/AudioEventLogTraceProvider.cpp
@@ -173,8 +173,13 @@ namespace UE::Audio::Insights
 			return false;
 		}
 
+		double AboutToBeDeletedTimestamp = TNumericLimits<double>::Lowest();
+		double NearingStartOfCacheTimestamp = TNumericLimits<double>::Lowest();
+		GetCacheStatusTimestampThresholds(NearingStartOfCacheTimestamp, AboutToBeDeletedTimestamp);
+
 		FAudioInsightsCacheManager& CacheManager = IAudioInsightsModule::GetChecked().GetCacheManager();
-		const uint32 HalfNumChunks = CacheManager.GetNumChunks() > 1 ? CacheManager.GetNumChunks() / 2 : 1;
+
+		const uint32 HalfNumChunks = CacheManager.GetNumChunks() / 2;
 
 		bool bAnyEntriesRemoved = false;
 		for (auto& [DeviceID, MessageIDToDashboardEntryMap] : DeviceDataMap)
@@ -188,6 +193,14 @@ namespace UE::Audio::Insights
 				{
 					MessagesToRemove.Add(MessageID);
 				}
+				else if (DashboardEntry->Timestamp <= AboutToBeDeletedTimestamp)
+				{
+					DashboardEntry->CachedState = EAudioEventCacheState::NextToBeDeleted;
+				}
+				else if (DashboardEntry->Timestamp <= NearingStartOfCacheTimestamp)
+				{
+					DashboardEntry->CachedState = EAudioEventCacheState::NearStartOfCache;
+				}
 				else
 				{
 					break;
@@ -265,6 +278,34 @@ namespace UE::Audio::Insights
 		CacheStartTimestamp = NewCacheStartTimestamp;
 		bCacheChunksUpdated = true;
 	}
+
+	void FAudioEventLogTraceProvider::GetCacheStatusTimestampThresholds(double& OutNearingStartOfCacheTimestamp, double& OutAboutToBeDeletedTimestamp)
+	{
+		FAudioInsightsCacheManager& CacheManager = IAudioInsightsModule::GetChecked().GetCacheManager();
+
+		constexpr int32 ChunkBeingDeletedFromCacheThreshold = 1;
+		constexpr int32 ChunkNearingStartOfCacheThreshold = 3;
+
+		const int32 NumUsedChunks = CacheManager.GetNumUsedChunks();
+		const int32 MaxNumChunks = CacheManager.GetNumChunks();
+
+		check(MaxNumChunks > ChunkNearingStartOfCacheThreshold && MaxNumChunks > ChunkBeingDeletedFromCacheThreshold);
+
+		// We don't mark anything as leaving the cache if the cache isn't close to overwriting chunks yet
+		if (NumUsedChunks < (MaxNumChunks - ChunkNearingStartOfCacheThreshold))
+		{
+			return;
+		}
+
+		const int32 BeingDeletedChunkIndex = NumUsedChunks - MaxNumChunks + ChunkBeingDeletedFromCacheThreshold - 1;
+		const int32 NearingStartOfCacheChunkIndex = NumUsedChunks - MaxNumChunks + ChunkNearingStartOfCacheThreshold - 1;
+
+		const FAudioCachedMessageChunk* const NextToBeDeletedChunk = BeingDeletedChunkIndex >= 0 ? CacheManager.GetChunk(BeingDeletedChunkIndex) : nullptr;
+		const FAudioCachedMessageChunk* const NearingStartOfCacheChunk = NearingStartOfCacheChunkIndex >= 0 ? CacheManager.GetChunk(NearingStartOfCacheChunkIndex) : nullptr;
+
+		OutAboutToBeDeletedTimestamp = NextToBeDeletedChunk ? NextToBeDeletedChunk->GetChunkTimeRangeEnd() : TNumericLimits<double>::Lowest();
+		OutNearingStartOfCacheTimestamp = NearingStartOfCacheChunk ? NearingStartOfCacheChunk->GetChunkTimeRangeEnd() : TNumericLimits<double>::Lowest();
+	}
 #endif // WITH_EDITOR
 } // namespace UE::Audio::Insights
 
diff --git a/Engine/Plugins/AudioInsights/Source/AudioInsights/Private/Providers/AudioEventLogTraceProvider.h b/Engine/Plugins/AudioInsights/Source/AudioInsights/Private/Providers/AudioEventLogTraceProvider.h
--- a/Engine/Plugins/AudioInsights/Source/AudioInsights/Private/Providers/AudioEventLogTraceProvider.h
+++ b/Engine/Plugins/AudioInsights/Source/AudioInsights/Private/Providers/AudioEventLogTraceProvider.h
@@ -26,6 +26,7 @@ namespace UE::Audio::Insights
 
 #if WITH_EDITOR
 		void OnCacheChunkOverwritten(const double NewCacheStartTimestamp);
+		void GetCacheStatusTimestampThresholds(double& OutNearingStartOfCacheTimestamp, double& OutAboutToBeDeletedTimestamp);
 #endif // WITH_EDITOR
 
 		FAudioEventLogMessages TraceMessages;
diff --git a/Engine/Plugins/AudioInsights/Source/AudioInsights/Private/Views/AudioEventLogDashboardViewFactory.cpp b/Engine/Plugins/AudioInsights/Source/AudioInsights/Private/Views/AudioEventLogDashboardViewFactory.cpp
--- a/Engine/Plugins/AudioInsights/Source/AudioInsights/Private/Views/AudioEventLogDashboardViewFactory.cpp
+++ b/Engine/Plugins/AudioInsights/Source/AudioInsights/Private/Views/AudioEventLogDashboardViewFactory.cpp
@@ -5,8 +5,8 @@
 #include "AudioInsightsModule.h"
 #include "AudioInsightsStyle.h"
 #include "AudioMixerTrace.h"
-#include "Messages/AudioEventLogTraceMessages.h"
 #include "Providers/AudioEventLogTraceProvider.h"
+#include "Widgets/Colors/SColorBlock.h"
 
 #define LOCTEXT_NAMESPACE "AudioInsightsEventLog"
 
@@ -29,7 +29,7 @@ namespace UE::Audio::Insights
 		const FString CategoryMessages = LOCTEXT("EventLogCategory_Messages", "Messages").ToString();
 		const FString CategoryCustom = LOCTEXT("EventLogCategory_Custom", "Custom").ToString();
 
-
+		const FName CacheStatusColumnID = "CacheStatus";
 		const FName MessageIDColumnID = "MessageID";
 		const FName TimestampColumnID = "Timestamp";
 		const FName PlayOrderColumnID = "PlayOrder";
@@ -37,6 +37,10 @@ namespace UE::Audio::Insights
 		const FName AssetNameColumnID = "Asset";
 		const FName ActorColumnID = "Actor";
 		const FName CategoryColumnID = "Category";
+
+		const FLinearColor AboutToBeDeletedColor = FSlateColor(EStyleColor::AccentRed).GetSpecifiedColor().CopyWithNewOpacity(0.66f);
+		const FLinearColor NearingStartOfCacheColor = FSlateColor(EStyleColor::AccentWhite).GetSpecifiedColor().CopyWithNewOpacity(0.33f);
+		constexpr float ScrollbarMarkerMaxWidth = 4.0f;
 	} // namespace AudioEventLogPrivate
 
 	const TMap<FString, TSet<FString>> FAudioEventLogDashboardViewFactory::GetInitEventTypeFilters()
@@ -266,6 +270,8 @@ namespace UE::Audio::Insights
 		});
 
 #if WITH_EDITOR
+		UpdateScrollBarSidebarWidget();
+
 		if (Reason == FTraceTableDashboardViewFactory::EProcessReason::FilterUpdated)
 		{
 			FAudioEventLogSettings::OnRequestWriteSettings.Broadcast();
@@ -287,6 +293,66 @@ namespace UE::Audio::Insights
 
 			return TMap<FName, FTraceTableDashboardViewFactory::FColumnData>
 			{
+#if WITH_EDITOR
+				{
+					CacheStatusColumnID,
+					{
+						FText::GetEmpty(),
+						[](const IDashboardDataViewEntry& InData) { return FText::GetEmpty(); }, /* GetDisplayValue */
+						[this](const IDashboardDataViewEntry& InData) -> FName
+						{
+							const FAudioEventLogDashboardEntry& Entry = CastEntry(InData);
+
+							switch (Entry.CachedState)
+							{
+								case EAudioEventCacheState::NearStartOfCache:
+									return FName("AudioInsights.Icon.EventLog.NearToStartOfCache");
+								case EAudioEventCacheState::NextToBeDeleted:
+									return FName("AudioInsights.Icon.EventLog.MarkedForDelete");
+								default:
+									break;
+							}
+
+							return NAME_None;
+						},
+						!VisibleColumns.GetIsVisible(CacheStatusColumnID) /* bDefaultHidden */,
+						0.075f /* FillWidth */,
+						HAlign_Center /* Alignment */,
+						[this](const IDashboardDataViewEntry& InData) -> FLinearColor /* GetIconColor */
+						{
+							const FAudioEventLogDashboardEntry& Entry = CastEntry(InData);
+
+							switch (Entry.CachedState)
+							{
+								case EAudioEventCacheState::NearStartOfCache:
+									return NearingStartOfCacheColor;
+								case EAudioEventCacheState::NextToBeDeleted:
+									return AboutToBeDeletedColor;
+								default:
+									break;
+							}
+
+							return FLinearColor::Transparent;
+						},
+						[this](const IDashboardDataViewEntry& InData) -> FText /* GetIconTooltip */
+						{
+							const FAudioEventLogDashboardEntry& Entry = CastEntry(InData);
+
+							switch (Entry.CachedState)
+							{
+								case EAudioEventCacheState::NearStartOfCache:
+									return LOCTEXT("AudioEventLog_NearEndOfCache", "This event is near the start of the cache, it may be marked for deletion soon.");
+								case EAudioEventCacheState::NextToBeDeleted:
+									return LOCTEXT("AudioEventLog_NextToBeDeleted", "This event has been marked for deletion the next time the cache is full.");
+								default:
+									break;
+							}
+
+							return FText::GetEmpty();
+						}
+					}
+				},
+#endif // WITH_EDITOR
 				{
 					MessageIDColumnID,
 					{
@@ -645,6 +711,22 @@ namespace UE::Audio::Insights
 		}
 	}
 
+	TSharedRef<SWidget> FAudioEventLogDashboardViewFactory::CreateVerticalScrollMarkerWidget()
+	{
+#if WITH_EDITOR
+		using namespace AudioEventLogPrivate;
+
+		return SNew(SHorizontalBox)
+
+		+ SHorizontalBox::Slot().MaxWidth(ScrollbarMarkerMaxWidth)
+		[
+			SAssignNew(ScrollMarkerWidget, SVerticalBox)
+		];
+#else
+		return SNullWidget::NullWidget;
+#endif // WITH_EDITOR
+	}
+
 	const FTableRowStyle* FAudioEventLogDashboardViewFactory::GetRowStyle() const
 	{
 		return &FSlateStyle::Get().GetWidgetStyle<FTableRowStyle>("TreeDashboard.TableViewRow");
@@ -705,6 +787,126 @@ namespace UE::Audio::Insights
 		return false;
 	}
 
+#if WITH_EDITOR
+	void FAudioEventLogDashboardViewFactory::UpdateScrollBarSidebarWidget()
+	{
+		using namespace AudioEventLogPrivate;
+
+		if (!FilteredEntriesListView.IsValid() || !ScrollMarkerWidget.IsValid() || !HeaderRowWidget.IsValid())
+		{
+			return;
+		}
+
+		ScrollMarkerWidget->ClearChildren();
+
+		const int32 TotalNumFilteredItems = FilteredEntriesListView->GetItems().Num();
+		if (TotalNumFilteredItems <= 0)
+		{
+			return;
+		}
+
+		ScrollMarkerWidget->AddSlot()
+		.AutoHeight()
+		[
+			SNew(SColorBlock)
+			.Color(FLinearColor::Transparent)
+			.Size_Lambda([this]()
+			{
+				return HeaderRowWidget.IsValid() ? FVector2D(ScrollbarMarkerMaxWidth, HeaderRowWidget->GetDesiredSize().Y) : FVector2D(ScrollbarMarkerMaxWidth, 0.0f);
+			})
+		];
+
+		const TArray<FAudioEventLogDashboardViewFactory::FScrollbarAreaStats> ScrollBarAreaStats = GetScrollbarAreaStats();
+
+		for (const FAudioEventLogDashboardViewFactory::FScrollbarAreaStats& AreaStats : ScrollBarAreaStats)
+		{
+			FLinearColor AreaColor;
+			FText Tooltip;
+
+			switch (AreaStats.AreaCacheState)
+			{
+				case EAudioEventCacheState::NearStartOfCache:
+					AreaColor = NearingStartOfCacheColor;
+					Tooltip = LOCTEXT("AudioEventLog_NearEndOfCacheArea", "These events are near the start of the cache, they may be marked for deletion soon.");
+					break;
+				case EAudioEventCacheState::NextToBeDeleted:
+					AreaColor = AboutToBeDeletedColor;
+					Tooltip = LOCTEXT("AudioEventLog_NextToBeDeletedArea", "These events have been marked for deletion the next time the cache is full.");
+					break;
+				default:
+					AreaColor = FLinearColor::Transparent;
+					Tooltip = FText::GetEmpty();
+					break;
+			}
+
+			const float AreaProp = static_cast<float>(AreaStats.NumRowsInArea) / static_cast<float>(TotalNumFilteredItems);
+			ScrollMarkerWidget->AddSlot()
+			.AutoHeight()
+			[
+				SNew(SColorBlock)
+				.Color(AreaColor)
+				.ToolTipText(Tooltip)
+				.Size_Lambda([this, AreaProp]()
+				{
+					if (!FilteredEntriesListView.IsValid() || !HeaderRowWidget.IsValid())
+					{
+						return FVector2D(ScrollbarMarkerMaxWidth, 0.0f);
+					}
+
+					const float TableHeight = FilteredEntriesListView->GetPaintSpaceGeometry().GetLocalSize().Y - HeaderRowWidget->GetDesiredSize().Y;
+					return FVector2D(ScrollbarMarkerMaxWidth, TableHeight * AreaProp);
+				})
+			];
+		}
+	}
+
+	TArray<FAudioEventLogDashboardViewFactory::FScrollbarAreaStats> FAudioEventLogDashboardViewFactory::GetScrollbarAreaStats() const
+	{
+		using namespace AudioEventLogPrivate;
+
+		TArray<FScrollbarAreaStats> ScrollBarAreaStats;
+		ScrollBarAreaStats.Reserve(3); // we only have 3 areas for now
+		int32 ScrollbarStatsIndex = -1;
+
+		for (const TSharedPtr<IDashboardDataViewEntry>& Entry : FilteredEntriesListView->GetItems())
+		{
+			if (!Entry.IsValid())
+			{
+				continue;
+			}
+
+			const FAudioEventLogDashboardEntry& EventLogEntry = CastEntry(*Entry);
+			if (ScrollBarAreaStats.IsEmpty())
+			{
+				ScrollBarAreaStats.Add({ EventLogEntry.CachedState, 1 });
+				ScrollbarStatsIndex++;
+			}
+			else
+			{
+				ensure(ScrollBarAreaStats.IsValidIndex(ScrollbarStatsIndex));
+				if (!ScrollBarAreaStats.IsValidIndex(ScrollbarStatsIndex))
+				{
+					return ScrollBarAreaStats;
+				}
+
+				FScrollbarAreaStats& AreaStats = ScrollBarAreaStats[ScrollbarStatsIndex];
+
+				if (AreaStats.AreaCacheState != EventLogEntry.CachedState)
+				{
+					ScrollBarAreaStats.Add({ EventLogEntry.CachedState, 1 });
+					ScrollbarStatsIndex++;
+				}
+				else
+				{
+					ScrollBarAreaStats[ScrollbarStatsIndex].NumRowsInArea++;
+				}
+			}
+		}
+
+		return ScrollBarAreaStats;
+	}
+#endif // WITH_EDITOR
+
 	void FAudioEventLogDashboardViewFactory::BuildVisibleColumnsMenuContent(FMenuBuilder& MenuBuilder)
 	{
 		if (VisibleColumnsSettingsMenu.IsValid())
diff --git a/Engine/Plugins/AudioInsights/Source/AudioInsights/Private/Views/TableDashboardViewFactory.cpp b/Engine/Plugins/AudioInsights/Source/AudioInsights/Private/Views/TableDashboardViewFactory.cpp
--- a/Engine/Plugins/AudioInsights/Source/AudioInsights/Private/Views/TableDashboardViewFactory.cpp
+++ b/Engine/Plugins/AudioInsights/Source/AudioInsights/Private/Views/TableDashboardViewFactory.cpp
@@ -60,6 +60,8 @@ namespace UE::Audio::Insights
 
 			? SNew(SImage)
 			.Image(FSlateStyle::Get().GetBrush(IconName))
+			.ColorAndOpacity(ColumnData.GetIconColor.IsSet() ? ColumnData.GetIconColor(InRowData.Get()) : FLinearColor::White)
+			.ToolTipText(ColumnData.GetIconTooltip.IsSet() ? ColumnData.GetIconTooltip(InRowData.Get()) : FText::GetEmpty())
 
 			: SNullWidget::NullWidget
 		]
@@ -297,7 +299,21 @@ namespace UE::Audio::Insights
 			.FillHeight(1.0f)
 			.Padding(0, 2)
 			[
-				EnableHorizontalScrollBox() ? CreateHorizontalScrollBox() : FilteredEntriesListView->AsShared()
+				SNew(SHorizontalBox)
+				
+				// List view contents
+				+ SHorizontalBox::Slot()
+				.FillContentWidth(1.0f)
+				[
+					EnableHorizontalScrollBox() ? CreateHorizontalScrollBox() : FilteredEntriesListView->AsShared()
+				]
+				
+				// optional scroll bar marker
+				+ SHorizontalBox::Slot()
+				.AutoWidth()
+				[
+					CreateVerticalScrollMarkerWidget()
+				]
 			];
 		}
 
@@ -429,12 +445,15 @@ namespace UE::Audio::Insights
 		// Icon
 		+ SHorizontalBox::Slot()
 		.Padding(0.0f, 2.0f)
+		.HAlign(ColumnData.Alignment)
 		.AutoWidth()
 		[
 			IconName != NAME_None
 
 			? SNew(SImage)
 			.Image(FSlateStyle::Get().GetBrush(IconName))
+			.ColorAndOpacity(ColumnData.GetIconColor.IsSet() ? ColumnData.GetIconColor(InRowData.Get()) : FLinearColor::White)
+			.ToolTipText(ColumnData.GetIconTooltip.IsSet() ? ColumnData.GetIconTooltip(InRowData.Get()) : FText::GetEmpty())
 
 			: SNullWidget::NullWidget
 		]
diff --git a/Engine/Plugins/AudioInsights/Source/AudioInsights/Public/AudioInsightsStyle.h b/Engine/Plugins/AudioInsights/Source/AudioInsights/Public/AudioInsightsStyle.h
--- a/Engine/Plugins/AudioInsights/Source/AudioInsights/Public/AudioInsightsStyle.h
+++ b/Engine/Plugins/AudioInsights/Source/AudioInsights/Public/AudioInsightsStyle.h
@@ -224,6 +224,10 @@ namespace UE::Audio::Insights
 			Set("AudioInsights.Icon.SoundDashboard.Transparent",      new IMAGE_BRUSH_SVG(TEXT("Icons/SoundDashboard/transparent"),      Icon16));
 			Set("AudioInsights.Icon.SoundDashboard.Visible",		  new IMAGE_BRUSH_SVG(TEXT("Icons/SoundDashboard/visible"),			 Icon16));
 			Set("AudioInsights.Icon.SoundDashboard.Invisible",		  new IMAGE_BRUSH_SVG(TEXT("Icons/SoundDashboard/invisible"),		 Icon16));
+
+			/* Event Log */
+			Set("AudioInsights.Icon.EventLog.MarkedForDelete",		  new IMAGE_BRUSH_SVG(TEXT("Icons/EventLog/cache_marked_for_delete"), Icon16));
+			Set("AudioInsights.Icon.EventLog.NearToStartOfCache",	  new IMAGE_BRUSH_SVG(TEXT("Icons/EventLog/cache_near_to_start"), Icon16));
 			
 			// Category colors
 			Set("SoundDashboard.MetaSoundColor", FLinearColor(0.008f, 0.76f, 0.078f));
diff --git a/Engine/Plugins/AudioInsights/Source/AudioInsights/Public/Cache/AudioInsightsCacheManager.h b/Engine/Plugins/AudioInsights/Source/AudioInsights/Public/Cache/AudioInsightsCacheManager.h
--- a/Engine/Plugins/AudioInsights/Source/AudioInsights/Public/Cache/AudioInsightsCacheManager.h
+++ b/Engine/Plugins/AudioInsights/Source/AudioInsights/Public/Cache/AudioInsightsCacheManager.h
@@ -23,10 +23,22 @@ namespace UE::Audio::Insights
 		UE_API double GetCacheStartTimeStamp() const;
 		UE_API double GetCacheEndTimeStamp() const;
 
+		// Returns a chunk by index relative to the start of the cache
+		// e.g. NumChunksFromStart = 0u will return the chunk with the earliest messages in the cache
+		UE_API const FAudioCachedMessageChunk* const GetChunk(const uint32 NumChunksFromStart) const;
+
+		// Returns the total used memory size of the cache in bytes
 		UE_API uint32 GetUsedCacheSize() const;
+
+		// Returns the maximum memory size of the cache in bytes
 		UE_API uint32 GetMaxCacheSize() const;
+
+		// Returns the maximum number of chunks inside the cache
 		UE_API uint32 GetNumChunks() const;
 
+		// Returns the number of chunks with data present inside the cache
+		UE_API uint32 GetNumUsedChunks() const;
+
 		// Find the closest message to a timestamp in the cache for a specific message type
 		// Will look backwards from the chosen timestamp until a message is found
 		// @param ID: When overloaded, the search algorithm will only return messages that have a matching ID.
diff --git a/Engine/Plugins/AudioInsights/Source/AudioInsights/Public/Messages/AudioEventLogTraceMessages.h b/Engine/Plugins/AudioInsights/Source/AudioInsights/Public/Messages/AudioEventLogTraceMessages.h
--- a/Engine/Plugins/AudioInsights/Source/AudioInsights/Public/Messages/AudioEventLogTraceMessages.h
+++ b/Engine/Plugins/AudioInsights/Source/AudioInsights/Public/Messages/AudioEventLogTraceMessages.h
@@ -25,6 +25,15 @@ namespace UE::Audio::Insights
 		SoundCueTemplate
 	};
 
+#if WITH_EDITOR
+	enum class EAudioEventCacheState : uint8
+	{
+		Cached,
+		NearStartOfCache,
+		NextToBeDeleted
+	};
+#endif // WITH_EDITOR
+
 	struct FAudioEventLogMessage : public IAudioCachedMessage
 	{
 		FAudioEventLogMessage() = default;
@@ -66,6 +75,10 @@ namespace UE::Audio::Insights
 		FString CategoryName;
 
 		EAudioEventLogSoundCategory Category = EAudioEventLogSoundCategory::None;
+
+#if WITH_EDITOR
+		EAudioEventCacheState CachedState = EAudioEventCacheState::Cached;
+#endif // WITH_EDITOR
 	};
 
 	class FAudioEventLogMessages
diff --git a/Engine/Plugins/AudioInsights/Source/AudioInsights/Public/Settings/AudioEventLogSettings.h b/Engine/Plugins/AudioInsights/Source/AudioInsights/Public/Settings/AudioEventLogSettings.h
--- a/Engine/Plugins/AudioInsights/Source/AudioInsights/Public/Settings/AudioEventLogSettings.h
+++ b/Engine/Plugins/AudioInsights/Source/AudioInsights/Public/Settings/AudioEventLogSettings.h
@@ -36,6 +36,9 @@ struct FAudioEventLogVisibleColumns : public FVisibleColumnsSettings
 {
 	GENERATED_BODY()
 
+	UPROPERTY(EditAnywhere, config, Category = EventLog)
+	bool bCacheStatus = true;
+
 	UPROPERTY(EditAnywhere, config, Category = EventLog)
 	bool bMessageID = false;
 
diff --git a/Engine/Plugins/AudioInsights/Source/AudioInsights/Public/Views/AudioEventLogDashboardViewFactory.h b/Engine/Plugins/AudioInsights/Source/AudioInsights/Public/Views/AudioEventLogDashboardViewFactory.h
--- a/Engine/Plugins/AudioInsights/Source/AudioInsights/Public/Views/AudioEventLogDashboardViewFactory.h
+++ b/Engine/Plugins/AudioInsights/Source/AudioInsights/Public/Views/AudioEventLogDashboardViewFactory.h
@@ -2,6 +2,7 @@
 #pragma once
 
 #include "Filters/GenericFilter.h"
+#include "Messages/AudioEventLogTraceMessages.h"
 #include "Settings/AudioEventLogSettings.h"
 #include "Settings/VisibleColumnsSettingsMenu.h"
 #include "Templates/SharedPointer.h"
@@ -83,6 +84,7 @@ namespace UE::Audio::Insights
 
 		UE_API virtual void OnSelectionChanged(TSharedPtr<IDashboardDataViewEntry> SelectedItem, ESelectInfo::Type SelectInfo) override;
 		virtual bool ClearSelectionOnClick() const override { return true; }
+		UE_API virtual TSharedRef<SWidget> CreateVerticalScrollMarkerWidget() override;
 		virtual bool EnableHorizontalScrollBox() const override { return false; }
 		UE_API virtual const FTableRowStyle* GetRowStyle() const override;
 
@@ -95,6 +97,18 @@ namespace UE::Audio::Insights
 		TSharedRef<FAudioEventLogFilter> CreateNewEventFilterType(const TSharedPtr<FFilterCategory>& FilterCategory, const FString& EventType);
 		bool PassesEventTypeFilter(const FEventLogFilterID EventLogID) const;
 
+#if WITH_EDITOR
+		void UpdateScrollBarSidebarWidget();
+
+		struct FScrollbarAreaStats
+		{
+			EAudioEventCacheState AreaCacheState = EAudioEventCacheState::Cached;
+			uint32 NumRowsInArea = 0u;
+		};
+
+		TArray<FScrollbarAreaStats> GetScrollbarAreaStats() const;
+#endif // WITH_EDITOR
+
 		void BuildVisibleColumnsMenuContent(FMenuBuilder& MenuBuilder);
 
 		void SortByPredicate(TFunctionRef<bool(const class FAudioEventLogDashboardEntry&, const class FAudioEventLogDashboardEntry&)> Predicate);
@@ -115,6 +129,7 @@ namespace UE::Audio::Insights
 		TSharedPtr<SAudioFilterBar<FEventLogFilterID>> FilterBar;
 
 		TSharedPtr<SWidget> SettingsAreaWidget;
+		TSharedPtr<SVerticalBox> ScrollMarkerWidget;
 
 		TMap<FString, TSharedPtr<FFilterCategory>> FilterCategories;
 		TMap<FString, FEventLogFilterID> EventFilterTypes;
diff --git a/Engine/Plugins/AudioInsights/Source/AudioInsights/Public/Views/TableDashboardViewFactory.h b/Engine/Plugins/AudioInsights/Source/AudioInsights/Public/Views/TableDashboardViewFactory.h
--- a/Engine/Plugins/AudioInsights/Source/AudioInsights/Public/Views/TableDashboardViewFactory.h
+++ b/Engine/Plugins/AudioInsights/Source/AudioInsights/Public/Views/TableDashboardViewFactory.h
@@ -136,6 +136,8 @@ namespace UE::Audio::Insights
 			bool bDefaultHidden = false;
 			const float FillWidth = 1.0f;
 			const EHorizontalAlignment Alignment = HAlign_Left;
+			const TFunction<FLinearColor(const IDashboardDataViewEntry&)> GetIconColor = nullptr;
+			const TFunction<FText(const IDashboardDataViewEntry&)> GetIconTooltip = nullptr;
 		};
 
 		virtual const TMap<FName, FColumnData>& GetColumns() const = 0;
@@ -160,6 +162,7 @@ namespace UE::Audio::Insights
 		virtual bool ClearSelectionOnClick() const { return false; }
 
 		virtual bool EnableHorizontalScrollBox() const { return true; }
+		virtual TSharedRef<SWidget> CreateVerticalScrollMarkerWidget() { return SNullWidget::NullWidget; }
 
 		virtual const FTableRowStyle* GetRowStyle() const { return nullptr; }
 		UE_API virtual FSlateColor GetRowColor(const TSharedPtr<IDashboardDataViewEntry>& InRowDataPtr);
