[Audio Insights] Split out checks to remove warnings in build

diff --git a/Engine/Plugins/AudioInsights/Source/AudioInsights/Private/Providers/AudioEventLogTraceProvider.cpp b/Engine/Plugins/AudioInsights/Source/AudioInsights/Private/Providers/AudioEventLogTraceProvider.cpp
--- a/Engine/Plugins/AudioInsights/Source/AudioInsights/Private/Providers/AudioEventLogTraceProvider.cpp
+++ b/Engine/Plugins/AudioInsights/Source/AudioInsights/Private/Providers/AudioEventLogTraceProvider.cpp
@@ -284,12 +284,12 @@ namespace UE::Audio::Insights
 		FAudioInsightsCacheManager& CacheManager = IAudioInsightsModule::GetChecked().GetCacheManager();
 
 		constexpr int32 ChunkBeingDeletedFromCacheThreshold = 1;
-		constexpr int32 ChunkNearingStartOfCacheThreshold = 3;
+		constexpr int32 ChunkNearingStartOfCacheThreshold = ChunkBeingDeletedFromCacheThreshold + 2;
 
 		const int32 NumUsedChunks = CacheManager.GetNumUsedChunks();
 		const int32 MaxNumChunks = CacheManager.GetNumChunks();
 
-		check(MaxNumChunks > ChunkNearingStartOfCacheThreshold && MaxNumChunks > ChunkBeingDeletedFromCacheThreshold);
+		check(MaxNumChunks > ChunkNearingStartOfCacheThreshold);
 
 		// We don't mark anything as leaving the cache if the cache isn't close to overwriting chunks yet
 		if (NumUsedChunks < (MaxNumChunks - ChunkNearingStartOfCacheThreshold))
