[Audio Insights] Implemented SaveSnapshot logic that works for both standalone and editor + added SaveSnapshot button to Audio Insights standalone UI

diff --git a/Engine/Plugins/AudioInsights/Source/AudioInsights/AudioInsights.Build.cs b/Engine/Plugins/AudioInsights/Source/AudioInsights/AudioInsights.Build.cs
--- a/Engine/Plugins/AudioInsights/Source/AudioInsights/AudioInsights.Build.cs
+++ b/Engine/Plugins/AudioInsights/Source/AudioInsights/AudioInsights.Build.cs
@@ -32,6 +32,7 @@ public class AudioInsights : ModuleRules
 				"SessionServices",
 				"Slate",
 				"SlateCore",
+				"Sockets",
 				"TraceAnalysis",
 				"TraceLog",
 				"ToolMenus",
diff --git a/Engine/Plugins/AudioInsights/Source/AudioInsights/Private/AudioInsightsDashboardFactory.cpp b/Engine/Plugins/AudioInsights/Source/AudioInsights/Private/AudioInsightsDashboardFactory.cpp
--- a/Engine/Plugins/AudioInsights/Source/AudioInsights/Private/AudioInsightsDashboardFactory.cpp
+++ b/Engine/Plugins/AudioInsights/Source/AudioInsights/Private/AudioInsightsDashboardFactory.cpp
@@ -7,6 +7,7 @@
 #include "Framework/MultiBox/MultiBoxBuilder.h"
 #include "Internationalization/Text.h"
 #include "Templates/SharedPointer.h"
+#include "Widgets/Images/SImage.h"
 #include "Widgets/Input/SButton.h"
 #include "Widgets/Input/SCheckBox.h"
 #include "Widgets/Layout/SBox.h"
@@ -26,6 +27,8 @@ namespace UE::Audio::Insights
 		static const FText EnableThemNowText = LOCTEXT("AudioDashboard_EnableNowText", "Enable them now?");
 		static const FText TraceControllerUnavailabledWarning = LOCTEXT("AudioDashboard_TraceControllerUnavailableWarning", "The Trace Controller API is currently unavailable.");
 		static const FText TryEnablingMessagingText = LOCTEXT("AudioDashboard_TryEnablingMessagingText", "Make sure you have launched this package with the -Messaging command line argument.");
+		static const FText SaveSnapshotText = LOCTEXT("AudioDashboard_SnapshotButton_Text", "Save Snapshot");
+		static const FText SaveSnapshotTooltipText = LOCTEXT("AudioDashboard_SnapshotButton_TooltipText", "Saves a snapshot of the most recent trace data to a .utrace file. By default snapshots are limited to 4MB and can be accessed in Unreal Insights. To specify a custom snapshot size (in MB) launch the game executable with the -tracetailmb=<size> option.");
 	}
 
 	::Audio::FDeviceId FDashboardFactory::GetDeviceId() const
@@ -56,19 +59,26 @@ namespace UE::Audio::Insights
 #endif	// !WITH_EDITOR
 
 				SNew(SVerticalBox)
+				// Menu bar (File | View | etc.)
 				+ SVerticalBox::Slot()
 				.AutoHeight()
 				[
 					MakeMenuBarWidget()
 				]
-
+				// Main Toolbar (for general Audio Insights buttons)
+				+ SVerticalBox::Slot()
+				.AutoHeight()
+				[
+					MakeMainToolbarWidget()
+				]
+				// Spacing
 				+ SVerticalBox::Slot()
-
 				.AutoHeight()
 				[
 					SNew(SBox)
 					.HeightOverride(4.0f)
 				]
+				// Dashboards area
 				+ SVerticalBox::Slot()
 				[
 					DashboardTabManager->RestoreFrom(TabLayout->AsShared(), Args.GetOwnerWindow()).ToSharedRef()
@@ -184,6 +194,62 @@ namespace UE::Audio::Insights
 
 		return MenuBarBuilder.MakeWidget();
 	}
+
+	TSharedRef<SWidget> FDashboardFactory::MakeMainToolbarWidget()
+	{
+		return SNew(SHorizontalBox)
+			// Save snapshot button
+			+ SHorizontalBox::Slot()
+			.HAlign(HAlign_Left)
+			.VAlign(VAlign_Center)
+			.AutoWidth()
+			.Padding(12.0f, 0.0f, 0.0f, 0.0f)
+			[
+				MakeSnapshotButtonWidget()
+			];
+	}
+
+	TSharedRef<SWidget> FDashboardFactory::MakeSnapshotButtonWidget()
+	{
+		using namespace DashboardFactoryPrivate;
+
+		return SNew(SButton)
+			.ButtonStyle(&FAppStyle::Get().GetWidgetStyle<FButtonStyle>("SimpleButton"))
+			.ToolTipText(SaveSnapshotTooltipText)
+			.OnClicked_Lambda([]()
+			{
+				FTraceModule& AudioInsightsTraceModule = static_cast<FTraceModule&>(FAudioInsightsModule::GetChecked().GetTraceModule());
+				AudioInsightsTraceModule.SaveTraceSnapshot();
+
+				return FReply::Handled();
+			})
+			[
+				SNew(SBox)
+				.HeightOverride(16.0f)
+				[
+					SNew(SHorizontalBox)
+					// Trace snapshot icon
+					+ SHorizontalBox::Slot()
+					.VAlign(VAlign_Center)
+					.AutoWidth()
+					[
+						SNew(SImage)
+						.Image(FSlateStyle::Get().GetBrush("AudioInsights.Icon.Trace.Snapshot"))
+						.ColorAndOpacity(FSlateColor::UseForeground())
+					]
+					// Trace snapshot text
+					+ SHorizontalBox::Slot()
+					.VAlign(VAlign_Center)
+					.AutoWidth()
+					.Padding(6.0f, 1.0f, 0.0f, 0.0f)
+					[
+						SNew(STextBlock)
+						.Justification(ETextJustify::Center)
+						.Text(SaveSnapshotText)
+					]
+				]
+			];
+	}
 	
 	TSharedPtr<FTabManager::FLayout> FDashboardFactory::GetDefaultTabLayout()
 	{
diff --git a/Engine/Plugins/AudioInsights/Source/AudioInsights/Private/AudioInsightsTraceModule.cpp b/Engine/Plugins/AudioInsights/Source/AudioInsights/Private/AudioInsightsTraceModule.cpp
--- a/Engine/Plugins/AudioInsights/Source/AudioInsights/Private/AudioInsightsTraceModule.cpp
+++ b/Engine/Plugins/AudioInsights/Source/AudioInsights/Private/AudioInsightsTraceModule.cpp
@@ -1,6 +1,7 @@
 // Copyright Epic Games, Inc. All Rights Reserved.
 #include "AudioInsightsTraceModule.h"
 
+#include "AudioInsightsLog.h"
 #include "AudioInsightsModule.h"
 #include "Cache/AudioInsightsCacheManager.h"
 #include "CoreGlobals.h"
@@ -9,15 +10,12 @@
 #include "ISessionServicesModule.h"
 #include "Misc/App.h"
 #include "Providers/VirtualLoopTraceProvider.h"
+#include "SocketSubsystem.h"
 #include "Templates/SharedPointer.h"
 #include "Trace/Trace.h"
 #include "TraceServices/Model/Channel.h"
 #include "TraceServices/Model/Diagnostics.h"
 
-#if WITH_EDITOR
-#include "AudioInsightsLog.h"
-#endif // WITH_EDITOR
-
 namespace UE::Audio::Insights
 {
 	namespace FTraceModulePrivate
@@ -44,6 +42,14 @@ namespace UE::Audio::Insights
 	{
 #if WITH_EDITOR
 		FTraceAuxiliary::OnTraceStopped.AddRaw(this, &FTraceModule::OnTraceStopped);
+#else
+		ISessionServicesModule& SessionServicesModule = FModuleManager::LoadModuleChecked<ISessionServicesModule>("SessionServices");
+
+		const TSharedPtr<ISessionManager> SessionManager = SessionServicesModule.GetSessionManager();
+		if (SessionManager.IsValid())
+		{
+			SessionManager->OnSessionsUpdated().AddRaw(this, &FTraceModule::OnSessionsUpdated);
+		}
 #endif // WITH_EDITOR
 	}
 
@@ -51,6 +57,17 @@ namespace UE::Audio::Insights
 	{
 #if WITH_EDITOR
 		FTraceAuxiliary::OnTraceStopped.RemoveAll(this);
+#else
+		if (FModuleManager::Get().IsModuleLoaded("SessionServices"))
+		{
+			ISessionServicesModule& SessionServicesModule = FModuleManager::GetModuleChecked<ISessionServicesModule>("SessionServices");
+
+			const TSharedPtr<ISessionManager> SessionManager = SessionServicesModule.GetSessionManager();
+			if (SessionManager.IsValid())
+			{
+				SessionManager->OnSessionsUpdated().RemoveAll(this);
+			}
+		}
 #endif // WITH_EDITOR
 	}
 
@@ -364,37 +381,6 @@ namespace UE::Audio::Insights
 
 		bTraceAnalysisHasStarted = false;
 	}
-
-	void FTraceModule::SaveTraceSnapshot()
-	{
-		bool bSnapshotSaved = false;
-
-		IUnrealInsightsModule& UnrealInsightsModule = FModuleManager::LoadModuleChecked<IUnrealInsightsModule>("TraceInsights");
-
-		if (const UE::Trace::FStoreClient* StoreClient = UnrealInsightsModule.GetStoreClient())
-		{
-			if (const UE::Trace::FStoreClient::FStatus* Status = StoreClient->GetStatus())
-			{
-				const FString TraceStorePath = FString(Status->GetStoreDir());
-				FString SnapshotFileName = FDateTime::Now().ToString(TEXT("%Y%m%d_%H%M%S_AudioInsights.utrace"));
-
-				FString SnapshotFilePath = FPaths::Combine(TraceStorePath, SnapshotFileName);
-				FPaths::NormalizeFilename(SnapshotFilePath);
-
-				bSnapshotSaved = FTraceAuxiliary::WriteSnapshot(*SnapshotFilePath);
-
-				if (bSnapshotSaved)
-				{
-					UE_LOG(LogAudioInsights, Log, TEXT("Audio Insights snapshot has been saved in %s"), *SnapshotFilePath);
-				}
-			}
-		}
-
-		if (!bSnapshotSaved)
-		{
-			UE_LOG(LogAudioInsights, Warning, TEXT("Failed to save Audio Insights snapshot"));
-		}
-	}
 #endif // WITH_EDITOR
 
 	void FTraceModule::OnOnlyTraceAudioChannelsStateChanged(const bool bOnlyTraceAudioChannels)
@@ -480,6 +466,35 @@ namespace UE::Audio::Insights
 
 		return TraceController->HasAvailableInstance(InstanceID);
 	}
+
+	void FTraceModule::OnSessionsUpdated()
+	{
+		if (!InstanceID.IsValid())
+		{
+			const IUnrealInsightsModule& UnrealInsightsModule = FModuleManager::LoadModuleChecked<IUnrealInsightsModule>("TraceInsights");
+
+			TSharedPtr<const TraceServices::IAnalysisSession> AnalysisSession = UnrealInsightsModule.GetAnalysisSession();
+			if (AnalysisSession.IsValid())
+			{
+				TraceServices::FAnalysisSessionReadScope SessionReadScope(*AnalysisSession.Get());
+
+				const TraceServices::IDiagnosticsProvider* DiagnosticsProvider = TraceServices::ReadDiagnosticsProvider(*AnalysisSession.Get());
+				if (DiagnosticsProvider && DiagnosticsProvider->IsSessionInfoAvailable())
+				{
+					InstanceID = DiagnosticsProvider->GetSessionInfo().InstanceId;
+				}
+			}
+		}
+
+		ISessionServicesModule& SessionServicesModule = FModuleManager::LoadModuleChecked<ISessionServicesModule>("SessionServices");
+
+		const TSharedPtr<ITraceController> TraceController = SessionServicesModule.GetTraceController();
+		if (TraceController.IsValid() && !TraceController->HasAvailableInstance(InstanceID))
+		{
+			TraceController->SendDiscoveryRequest();
+			TraceController->SendStatusUpdateRequest();
+		}
+	}
 #endif // !WITH_EDITOR
 
 	void FTraceModule::ExecuteConsoleCommand(const FString& InCommandStr)const 
@@ -507,6 +522,79 @@ namespace UE::Audio::Insights
 			}
 		}
 	}
+	
+	void FTraceModule::SaveTraceSnapshot()
+	{
+		auto GetSnapshotFilePath = []()
+		{
+			auto GetTraceHostAddress = []() -> FString
+			{
+				TSharedPtr<FInternetAddr> RecorderAddress;
+				
+				if (ISocketSubsystem* Sockets = ISocketSubsystem::Get())
+				{
+					bool bCanBindAll = false;
+					RecorderAddress = Sockets->GetLocalHostAddr(*GLog, bCanBindAll);
+				}
+				
+				if (RecorderAddress.IsValid())
+				{
+					constexpr bool bAppendPort = false;
+					return RecorderAddress->ToString(bAppendPort);
+				}
+				
+				return "127.0.0.1";
+			};
+			
+			if (const UE::Trace::FStoreClient* StoreClient = UE::Trace::FStoreClient::Connect(*GetTraceHostAddress()))
+			{
+				if (const UE::Trace::FStoreClient::FStatus* StoreClientStatus = StoreClient->GetStatus())
+				{
+					const FString TraceStorePath = FString(StoreClientStatus->GetStoreDir());
+					FString SnapshotFileName = FDateTime::Now().ToString(TEXT("%Y%m%d_%H%M%S_AudioInsights.utrace"));
+					
+					FString SnapshotFilePath = FPaths::Combine(TraceStorePath, SnapshotFileName);
+					FPaths::NormalizeFilename(SnapshotFilePath);
+					
+					return SnapshotFilePath;
+				}
+			}
+
+			UE_LOG(LogAudioInsights, Warning, TEXT("Failed to connect to the Trace Store Client when trying to save a snapshot, returning an empty snapshot path."));
+
+			return FString();
+		};
+
+		const FString SnapshotFilePath = GetSnapshotFilePath();
+		if (SnapshotFilePath.IsEmpty())
+		{
+			UE_LOG(LogAudioInsights, Warning, TEXT("Failed to save Audio Insights snapshot. Snapshot path is empty."));
+			return;
+		}
+
+#if WITH_EDITOR
+		const bool bWriteSnapshotSuccess = FTraceAuxiliary::WriteSnapshot(*SnapshotFilePath);
+		if (!bWriteSnapshotSuccess)
+		{
+			UE_LOG(LogAudioInsights, Warning, TEXT("Failed to save Audio Insights snapshot in %s"), *SnapshotFilePath);
+		}
+#else
+		ISessionServicesModule& SessionServicesModule = FModuleManager::LoadModuleChecked<ISessionServicesModule>("SessionServices");
+			
+		TSharedPtr<ITraceController> TraceController = SessionServicesModule.GetTraceController();
+		if (TraceController.IsValid())
+		{
+			TraceController->WithInstance(InstanceID, [SnapshotFilePath](const FTraceStatus& Status, ITraceControllerCommands& Commands)
+			{
+				Commands.SnapshotFile(SnapshotFilePath);
+			});
+		}
+		else
+		{
+			UE_LOG(LogAudioInsights, Warning, TEXT("Failed to save Audio Insights snapshot. TraceController is invalid."));
+		}
+#endif // WITH_EDITOR
+	}
 
 	void FTraceModule::GetLoggers(TArray<const TCHAR *>& OutLoggers)
 	{
diff --git a/Engine/Plugins/AudioInsights/Source/AudioInsights/Public/AudioInsightsDashboardFactory.h b/Engine/Plugins/AudioInsights/Source/AudioInsights/Public/AudioInsightsDashboardFactory.h
--- a/Engine/Plugins/AudioInsights/Source/AudioInsights/Public/AudioInsightsDashboardFactory.h
+++ b/Engine/Plugins/AudioInsights/Source/AudioInsights/Public/AudioInsightsDashboardFactory.h
@@ -25,6 +25,8 @@ namespace UE::Audio::Insights
 
 	private:
 		TSharedRef<SWidget> MakeMenuBarWidget();
+		TSharedRef<SWidget> MakeMainToolbarWidget();
+		TSharedRef<SWidget> MakeSnapshotButtonWidget();
 
 		TSharedPtr<FTabManager::FLayout> GetDefaultTabLayout();
 
diff --git a/Engine/Plugins/AudioInsights/Source/AudioInsights/Public/AudioInsightsTraceModule.h b/Engine/Plugins/AudioInsights/Source/AudioInsights/Public/AudioInsightsTraceModule.h
--- a/Engine/Plugins/AudioInsights/Source/AudioInsights/Public/AudioInsightsTraceModule.h
+++ b/Engine/Plugins/AudioInsights/Source/AudioInsights/Public/AudioInsightsTraceModule.h
@@ -53,13 +53,14 @@ namespace UE::Audio::Insights
 
 		UE_API virtual void ExecuteConsoleCommand(const FString& InCommandStr) const override;
 
+		UE_API virtual void SaveTraceSnapshot();
+
 #if !WITH_EDITOR
 		UE_API virtual void InitializeSessionInfo(const TraceServices::FSessionInfo& SessionInfo) override;
 		UE_API virtual void RequestChannelUpdate() override;
 		UE_API virtual void ResetTicker() override;
 		UE_API virtual bool TraceControllerIsAvailable() const override;
 #else
-		UE_API virtual void SaveTraceSnapshot() override;
 		virtual bool IsAudioInsightsTrace() const override { return bIsAudioInsightsTrace; }
 #endif // !WITH_EDITOR
 
@@ -81,6 +82,7 @@ namespace UE::Audio::Insights
 		UE_API const UE::Trace::FStoreClient::FSessionInfo* GetCurrentAnalysisSessionInfo() const;
 		UE_API bool GetAudioTracesAreEnabled() const;
 		UE_API void SendDiscoveryRequestToTraceController() const;
+		UE_API void OnSessionsUpdated();
 
 		UE_API bool Tick(float DeltaTime);
 #else
