[Audio Insights] FTraceModule::SaveTraceSnapshot: manage the lifetime of the FStoreClient pointer returned by FStoreClient::Connect method with a TUniquePtr (yes, FStoreClient::Connect calls new!). Detected by BugHawk.

diff --git a/Engine/Plugins/AudioInsights/Source/AudioInsights/Private/AudioInsightsTraceModule.cpp b/Engine/Plugins/AudioInsights/Source/AudioInsights/Private/AudioInsightsTraceModule.cpp
--- a/Engine/Plugins/AudioInsights/Source/AudioInsights/Private/AudioInsightsTraceModule.cpp
+++ b/Engine/Plugins/AudioInsights/Source/AudioInsights/Private/AudioInsightsTraceModule.cpp
@@ -546,7 +546,9 @@ namespace UE::Audio::Insights
 				return "127.0.0.1";
 			};
 			
-			if (const UE::Trace::FStoreClient* StoreClient = UE::Trace::FStoreClient::Connect(*GetTraceHostAddress()))
+			
+			const TUniquePtr<UE::Trace::FStoreClient> StoreClient(UE::Trace::FStoreClient::Connect(*GetTraceHostAddress()));
+			if (StoreClient.IsValid())
 			{
 				if (const UE::Trace::FStoreClient::FStatus* StoreClientStatus = StoreClient->GetStatus())
 				{
