[Audio Insights] Crash fix for threading issues in Audio Insights Cache Manager when messages are added/removed. Reading actions should not collide with write actions for now.

diff --git a/Engine/Plugins/AudioInsights/Source/AudioInsights/Private/Cache/AudioCachedMessageChunk.cpp b/Engine/Plugins/AudioInsights/Source/AudioInsights/Private/Cache/AudioCachedMessageChunk.cpp
--- a/Engine/Plugins/AudioInsights/Source/AudioInsights/Private/Cache/AudioCachedMessageChunk.cpp
+++ b/Engine/Plugins/AudioInsights/Source/AudioInsights/Private/Cache/AudioCachedMessageChunk.cpp
@@ -23,13 +23,21 @@ namespace UE::Audio::Insights
 
 		CurrentChunkSize += Message->GetSizeOf();
 
-		TSharedPtr<TraceServices::TPagedArray<TSharedPtr<IAudioCachedMessage>>>& CachedMessages = Messages.FindOrAdd(Message->GetMessageName(), MakeShared<TraceServices::TPagedArray<TSharedPtr<IAudioCachedMessage>>>(Allocator, 4096));
-		CachedMessages->EmplaceBack(Message);
+		{
+			FScopeLock Lock(&MessagesCriticalSection);
+
+			TSharedPtr<TraceServices::TPagedArray<TSharedPtr<IAudioCachedMessage>>>& CachedMessages = Messages.FindOrAdd(Message->GetMessageName(), MakeShared<TraceServices::TPagedArray<TSharedPtr<IAudioCachedMessage>>>(Allocator, 4096));
+			CachedMessages->EmplaceBack(Message);
+		}
     }
 
 	void FAudioCachedMessageChunk::ClearChunk()
 	{
-		Messages.Reset();
+		{
+			FScopeLock Lock(&MessagesCriticalSection);
+			Messages.Reset();
+		}
+		
 		CurrentChunkSize = 0u;
 
 		RangeStartTimestamp = INVALID_TIMESTAMP;
diff --git a/Engine/Plugins/AudioInsights/Source/AudioInsights/Public/Cache/AudioCachedMessageChunk.h b/Engine/Plugins/AudioInsights/Source/AudioInsights/Public/Cache/AudioCachedMessageChunk.h
--- a/Engine/Plugins/AudioInsights/Source/AudioInsights/Public/Cache/AudioCachedMessageChunk.h
+++ b/Engine/Plugins/AudioInsights/Source/AudioInsights/Public/Cache/AudioCachedMessageChunk.h
@@ -78,6 +78,8 @@ namespace UE::Audio::Insights
 
 		uint32 CurrentChunkSize = 0u;
 		const uint32 MaxChunkSize = 0u;
+
+		FCriticalSection MessagesCriticalSection;
 	};
 } // namespace UE::Audio::Insights
 
