Backout 45370452

diff --git a/Engine/Plugins/AudioInsights/Source/AudioInsights/Private/Cache/AudioCachedMessageChunk.cpp b/Engine/Plugins/AudioInsights/Source/AudioInsights/Private/Cache/AudioCachedMessageChunk.cpp
--- a/Engine/Plugins/AudioInsights/Source/AudioInsights/Private/Cache/AudioCachedMessageChunk.cpp
+++ b/Engine/Plugins/AudioInsights/Source/AudioInsights/Private/Cache/AudioCachedMessageChunk.cpp
@@ -23,20 +23,13 @@ namespace UE::Audio::Insights
 
 		CurrentChunkSize += Message->GetSizeOf();
 
-		{
-			FScopeLock Lock(&MessagesCriticalSection);
-
-			TSharedPtr<TraceServices::TPagedArray<TSharedPtr<IAudioCachedMessage>>>& CachedMessages = Messages.FindOrAdd(Message->GetMessageName(), MakeShared<TraceServices::TPagedArray<TSharedPtr<IAudioCachedMessage>>>(Allocator, 4096));
-			CachedMessages->EmplaceBack(Message);
-		}
+		TSharedPtr<TraceServices::TPagedArray<TSharedPtr<IAudioCachedMessage>>>& CachedMessages = Messages.FindOrAdd(Message->GetMessageName(), MakeShared<TraceServices::TPagedArray<TSharedPtr<IAudioCachedMessage>>>(Allocator, 4096));
+		CachedMessages->EmplaceBack(Message);
     }
 
 	void FAudioCachedMessageChunk::ClearChunk()
 	{
-		{
-			FScopeLock Lock(&MessagesCriticalSection);
-			Messages.Reset();
-		}
+		Messages.Reset();
 		
 		CurrentChunkSize = 0u;
 
diff --git a/Engine/Plugins/AudioInsights/Source/AudioInsights/Public/Cache/AudioCachedMessageChunk.h b/Engine/Plugins/AudioInsights/Source/AudioInsights/Public/Cache/AudioCachedMessageChunk.h
--- a/Engine/Plugins/AudioInsights/Source/AudioInsights/Public/Cache/AudioCachedMessageChunk.h
+++ b/Engine/Plugins/AudioInsights/Source/AudioInsights/Public/Cache/AudioCachedMessageChunk.h
@@ -78,8 +78,6 @@ namespace UE::Audio::Insights
 
 		uint32 CurrentChunkSize = 0u;
 		const uint32 MaxChunkSize = 0u;
-
-		FCriticalSection MessagesCriticalSection;
 	};
 } // namespace UE::Audio::Insights
 
