[Audio Insights] - Support adding Submix messages to AudioInsightsCacheManager

diff --git a/Engine/Plugins/AudioInsights/Source/AudioInsights/Private/Messages/SubmixTraceMessages.cpp b/Engine/Plugins/AudioInsights/Source/AudioInsights/Private/Messages/SubmixTraceMessages.cpp
new file mode 100644
--- /dev/null
+++ b/Engine/Plugins/AudioInsights/Source/AudioInsights/Private/Messages/SubmixTraceMessages.cpp
@@ -0,0 +1,84 @@
+// Copyright Epic Games, Inc. All Rights Reserved.
+#include "Messages/SubmixTraceMessages.h"
+
+namespace UE::Audio::Insights
+{
+	namespace SubmixMessageNames
+	{
+		const FName SubmixLoaded = "SubmixLoaded";
+		const FName SubmixIsAlivePing = "SubmixAlivePing";
+		const FName SubmixHasActivity = "SubmixHasActivity";
+		const FName SubmixEnvelopeValues = "SubmixEnvelopeValues";
+		const FName SubmixUnloaded = "SubmixUnloaded";
+	}
+
+	FSubmixMessageBase::FSubmixMessageBase(const UE::Trace::IAnalyzer::FOnEventContext& InContext)
+	{
+		const UE::Trace::IAnalyzer::FEventData& EventData = InContext.EventData;
+
+		DeviceId = static_cast<::Audio::FDeviceId>(EventData.GetValue<uint32>("DeviceId"));
+		SubmixId = EventData.GetValue<uint32>("SubmixId");
+		Timestamp = InContext.EventTime.AsSeconds(EventData.GetValue<uint64>("Timestamp"));
+	}
+
+	FSubmixLoadedMessage::FSubmixLoadedMessage(const UE::Trace::IAnalyzer::FOnEventContext& InContext)
+		: FSubmixMessageBase(InContext)
+	{
+		const UE::Trace::IAnalyzer::FEventData& EventData = InContext.EventData;
+
+		EventData.GetString("Name", Name);
+		NumChannels = EventData.GetValue<int32>("NumChannels");
+		bIsMainSubmix = EventData.GetValue<bool>("IsMainSubmix");
+	}
+
+	uint32 FSubmixLoadedMessage::GetSizeOf() const
+	{
+		// Start with the fixed size of the struct
+		uint32 MemorySize = sizeof(FSubmixLoadedMessage);
+
+		// Add any dynamically sized members
+		MemorySize += Name.GetAllocatedSize();
+		return MemorySize;
+	};
+
+	FSubmixHasActivityMessage::FSubmixHasActivityMessage(const UE::Trace::IAnalyzer::FOnEventContext& InContext)
+		: FSubmixMessageBase(InContext)
+	{
+		const UE::Trace::IAnalyzer::FEventData& EventData = InContext.EventData;
+
+		bHasActivity = EventData.GetValue<bool>("HasActivity");
+	}
+
+	FSubmixEnvelopeValuesMessage::FSubmixEnvelopeValuesMessage(const UE::Trace::IAnalyzer::FOnEventContext& InContext)
+		: FSubmixMessageBase(InContext)
+	{
+		const UE::Trace::IAnalyzer::FEventData& EventData = InContext.EventData;
+
+		const TArrayView<const float> EnvelopeValuesView = EventData.GetArrayView<float>("EnvelopeValues");
+
+		for (const auto& EnvelopeValue : EnvelopeValuesView)
+		{
+			EnvelopeValues.Emplace(EnvelopeValue);
+		}
+	}
+
+	uint32 FSubmixEnvelopeValuesMessage::GetSizeOf() const
+	{
+		// Start with the fixed size of the struct
+		uint32 MemorySize = sizeof(FSubmixEnvelopeValuesMessage);
+
+		// Add any dynamically sized members
+		if (EnvelopeValues.Num() > 0)
+		{
+			MemorySize += EnvelopeValues.Num() * sizeof(EnvelopeValues[0]);
+		}
+		return MemorySize;
+	};
+
+	FSubmixUnloadedMessage::FSubmixUnloadedMessage(const UE::Trace::IAnalyzer::FOnEventContext& InContext)
+		: FSubmixMessageBase(InContext)
+	{
+
+	};
+
+} // namespace UE::Audio::Insights
diff --git a/Engine/Plugins/AudioInsights/Source/AudioInsights/Private/Providers/SubmixTraceProvider.cpp b/Engine/Plugins/AudioInsights/Source/AudioInsights/Private/Providers/SubmixTraceProvider.cpp
--- a/Engine/Plugins/AudioInsights/Source/AudioInsights/Private/Providers/SubmixTraceProvider.cpp
+++ b/Engine/Plugins/AudioInsights/Source/AudioInsights/Private/Providers/SubmixTraceProvider.cpp
@@ -2,11 +2,14 @@
 
 #include "Providers/SubmixTraceProvider.h"
 
+#include "AudioInsightsModule.h"
+#include "Async/ParallelFor.h"
+#include "Cache/AudioInsightsCacheManager.h"
+
 #if WITH_EDITOR
 #include "AudioDevice.h"
 #include "AudioDeviceManager.h"
 #else
-#include "Async/ParallelFor.h"
 #include "Common/PagedArray.h"
 #endif // WITH_EDITOR
 
@@ -87,31 +90,54 @@ namespace UE::Audio::Insights
 			{
 				LLM_SCOPE_BYNAME(TEXT("Insights/FSubmixTraceAnalyzer"));
 
+#if WITH_EDITOR
+				if (ShouldIgnoreNewEvents())
+				{
+					return true;
+				}
+#endif // WITH_EDITOR
+
 				FSubmixMessages& Messages = GetProvider<FSubmixTraceProvider>().TraceMessages;
 
 				switch (RouteId)
 				{
 					case RouteId_Loaded:
 					{
+#if WITH_EDITOR
+						CacheMessage<FSubmixLoadedMessage>(Context, Messages.LoadedMessages);
+#else
 						Messages.LoadedMessages.Enqueue(FSubmixLoadedMessage{ Context });
+#endif // WITH_EDITOR
 						break;
 					}
 
 					case RouteId_HasActivity:
 					{
+#if WITH_EDITOR
+						CacheMessage<FSubmixHasActivityMessage>(Context, Messages.HasActivityMessages);
+#else
 						Messages.HasActivityMessages.Enqueue(FSubmixHasActivityMessage{ Context });
+#endif // WITH_EDITOR
 						break;
 					}
 
 					case RouteId_EnvelopeValues:
 					{
+#if WITH_EDITOR
+						CacheMessage<FSubmixEnvelopeValuesMessage>(Context, Messages.EnvelopeValuesMessages);
+#else
 						Messages.EnvelopeValuesMessages.Enqueue(FSubmixEnvelopeValuesMessage{ Context });
+#endif // WITH_EDITOR
 						break;
 					}
 
 					case RouteId_Unloaded:
 					{
+#if WITH_EDITOR
+						CacheMessage<FSubmixUnloadedMessage>(Context, Messages.UnloadedMessages);
+#else
 						Messages.UnloadedMessages.Enqueue(FSubmixUnloadedMessage{ Context });
+#endif // WITH_EDITOR
 						break;
 					}
 
@@ -148,8 +174,12 @@ namespace UE::Audio::Insights
 
 	bool FSubmixTraceProvider::ProcessMessages()
 	{
+		if (GetTimeControlMethod() != ETimeControlMethod::LatestMessageTime)
+		{
+			return true;
+		}
+
 		// Helper lambdas
-#if !WITH_EDITOR
 		auto CreateEntry = [this](const FSubmixMessageBase& Msg)
 		{
 			TSharedPtr<FSubmixDashboardEntry>* ToReturn = nullptr;
@@ -169,7 +199,6 @@ namespace UE::Audio::Insights
 
 			return ToReturn;
 		};
-#endif // !WITH_EDITOR
 
 		auto GetEntry = [this](const FSubmixMessageBase& Msg)
 		{
@@ -177,14 +206,15 @@ namespace UE::Audio::Insights
 		};
 
 		// Process messages
-#if !WITH_EDITOR
 		ProcessMessageQueue<FSubmixLoadedMessage>(TraceMessages.LoadedMessages, CreateEntry,
 		[this](const FSubmixLoadedMessage& Msg, TSharedPtr<FSubmixDashboardEntry>* OutEntry)
 		{
+#if !WITH_EDITOR
 			if (SessionCachedMessages.IsValid())
 			{
 				SessionCachedMessages->LoadedCachedMessages.EmplaceBack(Msg);
 			}
+#endif // !WITH_EDITOR
 
 			if (OutEntry)
 			{
@@ -194,7 +224,6 @@ namespace UE::Audio::Insights
 				EntryRef.bIsMainSubmix = Msg.bIsMainSubmix;
 			}
 		});
-#endif // !WITH_EDITOR
 
 		ProcessMessageQueue<FSubmixHasActivityMessage>(TraceMessages.HasActivityMessages, GetEntry,
 		[this](const FSubmixHasActivityMessage& Msg, TSharedPtr<FSubmixDashboardEntry>* OutEntry)
@@ -217,6 +246,13 @@ namespace UE::Audio::Insights
 		ProcessMessageQueue<FSubmixEnvelopeValuesMessage>(TraceMessages.EnvelopeValuesMessages, GetEntry,
 		[this](const FSubmixEnvelopeValuesMessage& Msg, TSharedPtr<FSubmixDashboardEntry>* OutEntry)
 		{
+#if !WITH_EDITOR
+			if (SessionCachedMessages.IsValid())
+			{
+				SessionCachedMessages->EnvelopeValuesCachedMessages.EmplaceBack(Msg);
+			}
+#endif // !WITH_EDITOR
+				
 			if (OutEntry)
 			{
 				FSubmixDashboardEntry& EntryRef = *OutEntry->Get();
@@ -226,21 +262,21 @@ namespace UE::Audio::Insights
 			}
 		});
 
-#if !WITH_EDITOR
 		ProcessMessageQueue<FSubmixUnloadedMessage>(TraceMessages.UnloadedMessages, GetEntry,
 		[this](const FSubmixUnloadedMessage& Msg, TSharedPtr<FSubmixDashboardEntry>* OutEntry)
 		{
+#if !WITH_EDITOR
 			if (SessionCachedMessages.IsValid())
 			{
 				SessionCachedMessages->UnloadedCachedMessages.EmplaceBack(Msg);
 			}
+#endif // !WITH_EDITOR
 
 			if (OutEntry && (*OutEntry)->Timestamp < Msg.Timestamp)
 			{
 				RemoveDeviceEntry(Msg.DeviceId, Msg.SubmixId);
 			}
 		});
-#endif // !WITH_EDITOR
 
 		return true;
 	}
@@ -344,42 +380,74 @@ namespace UE::Audio::Insights
 
 	void FSubmixTraceProvider::OnTimingViewTimeMarkerChanged(double TimeMarker)
 	{
-#if !WITH_EDITOR
 		using namespace FSubmixTraceProviderPrivate;
 
+#if !WITH_EDITOR
 		if (!SessionCachedMessages.IsValid())
 		{
 			return;
 		}
+#endif // !WITH_EDITOR
 
 		DeviceDataMap.Empty();
 
 		// Collect all the submix loaded messages registered until this point in time 
-		for (const FSubmixLoadedMessage& SubmixLoadedCachedMessage : SessionCachedMessages->LoadedCachedMessages)
+		auto ProcessSubmixLoadedMessage = [this](const FSubmixLoadedMessage& Msg)
 		{
-			if (SubmixLoadedCachedMessage.Timestamp > TimeMarker)
-			{
-				break;
-			}
-
-			UpdateDeviceEntry(SubmixLoadedCachedMessage.DeviceId, SubmixLoadedCachedMessage.SubmixId, [&SubmixLoadedCachedMessage](TSharedPtr<FSubmixDashboardEntry>& Entry)
+			UpdateDeviceEntry(Msg.DeviceId, Msg.SubmixId, [&Msg](TSharedPtr<FSubmixDashboardEntry>& Entry)
 			{
 				if (!Entry.IsValid())
 				{
 					Entry = MakeShared<FSubmixDashboardEntry>();
-					Entry->SubmixId = SubmixLoadedCachedMessage.SubmixId;
+					Entry->SubmixId = Msg.SubmixId;
 				}
 
-				Entry->Timestamp = SubmixLoadedCachedMessage.Timestamp;
+				Entry->Timestamp = Msg.Timestamp;
 
-				Entry->Name = *SubmixLoadedCachedMessage.Name;
-				Entry->AudioMeterInfo->NumChannels = SubmixLoadedCachedMessage.NumChannels;
-				Entry->bIsMainSubmix = SubmixLoadedCachedMessage.bIsMainSubmix;
+				Entry->Name = *Msg.Name;
+				Entry->AudioMeterInfo->NumChannels = Msg.NumChannels;
+				Entry->bIsMainSubmix = Msg.bIsMainSubmix;
 			});
+
+		};
+
+#if WITH_EDITOR
+		const FAudioInsightsCacheManager& CacheManager = FAudioInsightsModule::GetChecked().GetCacheManager();
+		CacheManager.IterateTo<FSubmixLoadedMessage>(SubmixMessageNames::SubmixLoaded, TimeMarker, [&ProcessSubmixLoadedMessage](const FSubmixLoadedMessage& SubmixLoadedCachedMessage)
+		{
+			ProcessSubmixLoadedMessage(SubmixLoadedCachedMessage);
+		});
+
+#else
+		for (const FSubmixLoadedMessage& SubmixLoadedCachedMessage : SessionCachedMessages->LoadedCachedMessages)
+		{
+			if (SubmixLoadedCachedMessage.Timestamp > TimeMarker)
+			{
+				break;
+			}
+
+			ProcessSubmixLoadedMessage(SubmixLoadedCachedMessage);
 		}
+#endif // WITH_EDITOR
 
 		// Selectively remove loaded messages collected in the step above by knowing which submixes were unloaded.
 		// With this we will know what are the available submixes at this point in time.
+		auto ProcessSubmixUnloadedMessage = [this](const FSubmixUnloadedMessage& Msg)
+		{
+			auto* OutEntry = FindDeviceEntry(Msg.DeviceId, Msg.SubmixId);
+
+			if (OutEntry && (*OutEntry)->Timestamp < Msg.Timestamp)
+			{
+				RemoveDeviceEntry(Msg.DeviceId, Msg.SubmixId);
+			}
+		};
+
+#if WITH_EDITOR
+		CacheManager.IterateTo<FSubmixUnloadedMessage>(SubmixMessageNames::SubmixUnloaded, TimeMarker, [&ProcessSubmixUnloadedMessage](const FSubmixUnloadedMessage& SubmixUnloadedCachedMessage)
+		{
+			ProcessSubmixUnloadedMessage(SubmixUnloadedCachedMessage);
+		});
+#else
 		for (const FSubmixUnloadedMessage& SubmixUnloadedCachedMessage : SessionCachedMessages->UnloadedCachedMessages)
 		{
 			if (SubmixUnloadedCachedMessage.Timestamp > TimeMarker)
@@ -387,13 +455,9 @@ namespace UE::Audio::Insights
 				break;
 			}
 
-			auto* OutEntry = FindDeviceEntry(SubmixUnloadedCachedMessage.DeviceId, SubmixUnloadedCachedMessage.SubmixId);
-
-			if (OutEntry && (*OutEntry)->Timestamp < SubmixUnloadedCachedMessage.Timestamp)
-			{
-				RemoveDeviceEntry(SubmixUnloadedCachedMessage.DeviceId, SubmixUnloadedCachedMessage.SubmixId);
-			}
+			ProcessSubmixUnloadedMessage(SubmixUnloadedCachedMessage);
 		}
+#endif // WITH_EDITOR
 
 		const FDeviceData* DeviceData = FindFilteredDeviceData();
 		if (DeviceData)
@@ -402,6 +466,7 @@ namespace UE::Audio::Insights
 			struct CachedEntryInfo
 			{
 				TOptional<FSubmixHasActivityMessage> HasActivityMessage;
+				TOptional<FSubmixEnvelopeValuesMessage> EnvelopeValuesMessage;
 			};
 
 			TArray<uint32> SubmixIdArray;
@@ -411,16 +476,35 @@ namespace UE::Audio::Insights
 			CachedEntryInfos.SetNum(SubmixIdArray.Num());
 
 			ParallelFor(SubmixIdArray.Num(),
+#if WITH_EDITOR
+			[&SubmixIdArray, &CachedEntryInfos, &CacheManager, TimeMarker, this](const int32 Index)
+#else
 			[&SubmixIdArray, &CachedEntryInfos, TimeMarker, this](const int32 Index)
+#endif // WITH_EDITOR
 			{
 				const uint32 SubmixId = SubmixIdArray[Index];
 
 				// HasActivity
+#if WITH_EDITOR
+				const FSubmixHasActivityMessage* FoundHasActivityCachedMessage = CacheManager.FindClosestMessage<FSubmixHasActivityMessage>(SubmixMessageNames::SubmixHasActivity, TimeMarker, SubmixId);
+#else
 				const FSubmixHasActivityMessage* FoundHasActivityCachedMessage = FindClosestMessageToTimestamp(SessionCachedMessages->HasActivityCachedMessages, TimeMarker, SubmixId);
+#endif // WITH_EDITOR
 				if (FoundHasActivityCachedMessage)
 				{
 					CachedEntryInfos[Index].HasActivityMessage = *FoundHasActivityCachedMessage;
 				}
+
+				// EnvelopeValues
+#if WITH_EDITOR
+				const FSubmixEnvelopeValuesMessage* FoundEnvelopeValuesCachedMessage = CacheManager.FindClosestMessage<FSubmixEnvelopeValuesMessage>(SubmixMessageNames::SubmixEnvelopeValues, TimeMarker, SubmixId);
+#else
+				const FSubmixEnvelopeValuesMessage* FoundEnvelopeValuesCachedMessage = FindClosestMessageToTimestamp(SessionCachedMessages->EnvelopeValuesCachedMessages, TimeMarker, SubmixId);
+#endif // WITH_EDITOR
+				if (FoundEnvelopeValuesCachedMessage)
+				{
+					CachedEntryInfos[Index].EnvelopeValuesMessage = *FoundEnvelopeValuesCachedMessage;
+				}
 			});
 
 			// Update the device entries with the collected info
@@ -443,12 +527,38 @@ namespace UE::Audio::Insights
 						Entry->bHasActivity = HasActivityMessage.bHasActivity;
 					});
 				}
+				if (CachedEntryInfo.EnvelopeValuesMessage.IsSet())
+				{
+					const FSubmixEnvelopeValuesMessage& EnvelopeValuesMessage = CachedEntryInfo.EnvelopeValuesMessage.GetValue();
+
+					UpdateDeviceEntry(EnvelopeValuesMessage.DeviceId, EnvelopeValuesMessage.SubmixId, [&EnvelopeValuesMessage](TSharedPtr<FSubmixDashboardEntry>& Entry)
+					{
+						if (!Entry.IsValid())
+						{
+							Entry = MakeShared<FSubmixDashboardEntry>();
+							Entry->SubmixId = EnvelopeValuesMessage.SubmixId;
+						}
+
+						Entry->Timestamp = EnvelopeValuesMessage.Timestamp;
+
+						Entry->AudioMeterInfo->NumChannels = EnvelopeValuesMessage.EnvelopeValues.Num();
+						for (const auto& EnvelopeValue : EnvelopeValuesMessage.EnvelopeValues)
+						{
+							Entry->AudioMeterInfo->EnvelopeValues.Emplace(EnvelopeValue);
+						}
+					});
+				}
 			}
 		}
 
-#endif // !WITH_EDITOR
-
 		// Call parent method to update LastMessageId
 		FTraceProviderBase::OnTimingViewTimeMarkerChanged(TimeMarker);
 	}
+
+#if WITH_EDITOR
+	void FSubmixTraceProvider::OnTimeControlMethodReset()
+	{
+		SubmixAssetProvider.RequestEntriesUpdate();
+	}
+#endif // WITH_EDITOR
 } // namespace UE::Audio::Insights
diff --git a/Engine/Plugins/AudioInsights/Source/AudioInsights/Private/Providers/SubmixTraceProvider.h b/Engine/Plugins/AudioInsights/Source/AudioInsights/Private/Providers/SubmixTraceProvider.h
--- a/Engine/Plugins/AudioInsights/Source/AudioInsights/Private/Providers/SubmixTraceProvider.h
+++ b/Engine/Plugins/AudioInsights/Source/AudioInsights/Private/Providers/SubmixTraceProvider.h
@@ -38,6 +38,8 @@ namespace UE::Audio::Insights
 		virtual bool ProcessMessages() override;
 
 #if WITH_EDITOR
+		virtual void OnTimeControlMethodReset() override;
+
 		void HandleOnSubmixAssetAdded(const TSharedPtr<FSubmixDashboardEntry>& InSubmixDashboardEntry);
 		void HandleOnSubmixAssetRemoved(const uint32 InSubmixId);
 		void HandleOnSubmixAssetListUpdated(const TArray<TSharedPtr<FSubmixDashboardEntry>>& InSubmixAssetEntries);
diff --git a/Engine/Plugins/AudioInsights/Source/AudioInsights/Public/Messages/SubmixTraceMessages.h b/Engine/Plugins/AudioInsights/Source/AudioInsights/Public/Messages/SubmixTraceMessages.h
--- a/Engine/Plugins/AudioInsights/Source/AudioInsights/Public/Messages/SubmixTraceMessages.h
+++ b/Engine/Plugins/AudioInsights/Source/AudioInsights/Public/Messages/SubmixTraceMessages.h
@@ -1,6 +1,7 @@
 // Copyright Epic Games, Inc. All Rights Reserved.
 #pragma once
 
+#include "Cache/IAudioCachedMessage.h"
 #include "Providers/AudioMeterProvider.h"
 #include "Messages/AnalyzerMessageQueue.h"
 #if WITH_EDITOR
@@ -16,37 +17,34 @@
 
 namespace UE::Audio::Insights
 {
-	struct FSubmixMessageBase
+	namespace SubmixMessageNames
 	{
-		FSubmixMessageBase() = default;
+		extern const FName SubmixLoaded;
+		extern const FName SubmixIsAlivePing;
+		extern const FName SubmixHasActivity;
+		extern const FName SubmixEnvelopeValues;
+		extern const FName SubmixUnloaded;
+	};
 
-		FSubmixMessageBase(const UE::Trace::IAnalyzer::FOnEventContext& InContext)
-		{
-			const UE::Trace::IAnalyzer::FEventData& EventData = InContext.EventData;
+	struct FSubmixMessageBase : public IAudioCachedMessage
+	{
+		FSubmixMessageBase() = default;
+		FSubmixMessageBase(const UE::Trace::IAnalyzer::FOnEventContext& InContext);
 
-			DeviceId  = static_cast<::Audio::FDeviceId>(EventData.GetValue<uint32>("DeviceId"));
-			SubmixId  = EventData.GetValue<uint32>("SubmixId");
-			Timestamp = InContext.EventTime.AsSeconds(EventData.GetValue<uint64>("Timestamp"));
-		}
+		virtual uint32 GetID() const override { return SubmixId; };
+		virtual uint32 GetSizeOf() const override { return sizeof(FSubmixMessageBase); };
 
 		::Audio::FDeviceId DeviceId = INDEX_NONE;
 		uint32 SubmixId = INDEX_NONE;
-		double Timestamp = 0.0;
 	};
 
 	struct FSubmixLoadedMessage : public FSubmixMessageBase
 	{
 		FSubmixLoadedMessage() = default;
+		FSubmixLoadedMessage(const UE::Trace::IAnalyzer::FOnEventContext& InContext);
 
-		FSubmixLoadedMessage(const UE::Trace::IAnalyzer::FOnEventContext& InContext)
-			: FSubmixMessageBase(InContext)
-		{
-			const UE::Trace::IAnalyzer::FEventData& EventData = InContext.EventData;
-
-			EventData.GetString("Name", Name);
-			NumChannels = EventData.GetValue<int32>("NumChannels");
-			bIsMainSubmix = EventData.GetValue<bool>("IsMainSubmix");
-		}
+		virtual uint32 GetSizeOf() const override;
+		virtual const FName GetMessageName() const override { return SubmixMessageNames::SubmixLoaded; }
 
 		FString Name;
 		int32 NumChannels = 0;
@@ -56,14 +54,10 @@ namespace UE::Audio::Insights
 	struct FSubmixHasActivityMessage : public FSubmixMessageBase
 	{
 		FSubmixHasActivityMessage() = default;
+		FSubmixHasActivityMessage(const UE::Trace::IAnalyzer::FOnEventContext& InContext);
 
-		FSubmixHasActivityMessage(const UE::Trace::IAnalyzer::FOnEventContext& InContext)
-			: FSubmixMessageBase(InContext)
-		{
-			const UE::Trace::IAnalyzer::FEventData& EventData = InContext.EventData;
-
-			bHasActivity = EventData.GetValue<bool>("HasActivity");
-		}
+		virtual uint32 GetSizeOf() const override { return sizeof(FSubmixHasActivityMessage); };
+		virtual const FName GetMessageName() const override { return SubmixMessageNames::SubmixHasActivity; }
 
 		bool bHasActivity = false;
 	};
@@ -71,24 +65,24 @@ namespace UE::Audio::Insights
 	struct FSubmixEnvelopeValuesMessage : public FSubmixMessageBase
 	{
 		FSubmixEnvelopeValuesMessage() = default;
+		FSubmixEnvelopeValuesMessage(const UE::Trace::IAnalyzer::FOnEventContext& InContext);
 
-		FSubmixEnvelopeValuesMessage(const UE::Trace::IAnalyzer::FOnEventContext& InContext)
-			: FSubmixMessageBase(InContext)
-		{
-			const UE::Trace::IAnalyzer::FEventData& EventData = InContext.EventData;
+		virtual uint32 GetSizeOf() const override;
+		virtual const FName GetMessageName() const override { return SubmixMessageNames::SubmixEnvelopeValues; }
 
-			const TArrayView<const float> EnvelopeValuesView = EventData.GetArrayView<float>("EnvelopeValues");
+		TArray<float> EnvelopeValues;
+	};
 
-			for (const auto& EnvelopeValue : EnvelopeValuesView)
-			{
-				EnvelopeValues.Emplace(EnvelopeValue);
-			}
-		}
+	struct FSubmixUnloadedMessage : public FSubmixMessageBase
+	{
+		FSubmixUnloadedMessage() = default;
+		FSubmixUnloadedMessage(const UE::Trace::IAnalyzer::FOnEventContext& InContext);
+
+		virtual uint32 GetSizeOf() const override { return sizeof(FSubmixUnloadedMessage); };
+		virtual const FName GetMessageName() const override { return SubmixMessageNames::SubmixUnloaded; }
 
-		TArray<float> EnvelopeValues;
 	};
 
-	using FSubmixUnloadedMessage = FSubmixMessageBase;
 
 	class FSubmixMessages
 	{
@@ -106,6 +100,7 @@ namespace UE::Audio::Insights
 		FSubmixSessionCachedMessages(TraceServices::IAnalysisSession& InSession)
 			: LoadedCachedMessages(InSession.GetLinearAllocator(), 4096)
 			, HasActivityCachedMessages(InSession.GetLinearAllocator(), 4096)
+			, EnvelopeValuesCachedMessages(InSession.GetLinearAllocator(), 4096)
 			, UnloadedCachedMessages(InSession.GetLinearAllocator(), 4096)
 		{
 
@@ -113,6 +108,7 @@ namespace UE::Audio::Insights
 
 		TraceServices::TPagedArray<FSubmixLoadedMessage> LoadedCachedMessages;
 		TraceServices::TPagedArray<FSubmixHasActivityMessage> HasActivityCachedMessages;
+		TraceServices::TPagedArray<FSubmixEnvelopeValuesMessage> EnvelopeValuesCachedMessages;
 		TraceServices::TPagedArray<FSubmixUnloadedMessage> UnloadedCachedMessages;
 	};
 #endif // !WITH_EDITOR
@@ -134,11 +130,7 @@ namespace UE::Audio::Insights
 
 		bool IsMainSubmix() const
 		{
-#if !WITH_EDITOR
 			return bIsMainSubmix;
-#else
-			return GetDefault<UAudioSettings>() && (FSoftObjectPath(Name) == GetDefault<UAudioSettings>()->MasterSubmix);
-#endif
 		}
 
 		uint32 SubmixId = INDEX_NONE;
@@ -148,8 +140,6 @@ namespace UE::Audio::Insights
 
 		FString Name;
 		bool bHasActivity = false;
-#if !WITH_EDITOR
 		bool bIsMainSubmix = false;
-#endif // !WITH_EDITOR
 	};
 } // namespace UE::Audio::Insights
