[Audio Insights] Virtual Loops: added caching support via cache manager + misc. code adjustments for readability

diff --git a/Engine/Plugins/AudioInsights/Source/AudioInsights/Private/Messages/VirtualLoopTraceMessages.cpp b/Engine/Plugins/AudioInsights/Source/AudioInsights/Private/Messages/VirtualLoopTraceMessages.cpp
new file mode 100644
--- /dev/null
+++ b/Engine/Plugins/AudioInsights/Source/AudioInsights/Private/Messages/VirtualLoopTraceMessages.cpp
@@ -0,0 +1,80 @@
+// Copyright Epic Games, Inc. All Rights Reserved.
+#include "Messages/VirtualLoopTraceMessages.h"
+
+namespace UE::Audio::Insights
+{
+	namespace VirtualLoopMessageNames
+	{
+		extern const FName Virtualize    = "VirtualLoopVirtualize";
+		extern const FName StopOrRealize = "VirtualLoopStopOrRealize";
+		extern const FName Update        = "VirtualLoopUpdate";
+	};
+
+	// FVirtualLoopMessageBase
+	FVirtualLoopMessageBase::FVirtualLoopMessageBase(const Trace::IAnalyzer::FOnEventContext& InContext)
+	{
+		const Trace::IAnalyzer::FEventData& EventData = InContext.EventData;
+
+		DeviceId = static_cast<::Audio::FDeviceId>(EventData.GetValue<uint32>("DeviceId"));
+		Timestamp = InContext.EventTime.AsSeconds(EventData.GetValue<uint64>("Timestamp"));
+		PlayOrder = EventData.GetValue<uint32>("PlayOrder");
+	}
+
+	// FVirtualLoopVirtualizeMessage
+	FVirtualLoopVirtualizeMessage::FVirtualLoopVirtualizeMessage(const Trace::IAnalyzer::FOnEventContext& InContext)
+		: FVirtualLoopMessageBase(InContext)
+	{
+		const Trace::IAnalyzer::FEventData& EventData = InContext.EventData;
+
+		EventData.GetString("Name", Name);
+		ComponentId = EventData.GetValue<uint64>("ComponentId");
+	}
+	
+	uint32 FVirtualLoopVirtualizeMessage::GetSizeOf() const
+	{
+		// Fixed size of the struct
+		uint32 MemorySize = sizeof(FVirtualLoopVirtualizeMessage);
+
+		// Adde dynamically sized members
+		MemorySize += Name.GetAllocatedSize();
+
+		return MemorySize;
+	}
+	
+	// FVirtualLoopStopOrRealizeMessage
+	FVirtualLoopStopOrRealizeMessage::FVirtualLoopStopOrRealizeMessage(const Trace::IAnalyzer::FOnEventContext& InContext)
+		: FVirtualLoopMessageBase(InContext)
+	{
+	}
+
+	uint32 FVirtualLoopStopOrRealizeMessage::GetSizeOf() const
+	{
+		// Returning fixed size of the struct (no dynamically sized members in it)
+		return sizeof(FVirtualLoopStopOrRealizeMessage);
+	}
+
+	// FVirtualLoopUpdateMessage
+	FVirtualLoopUpdateMessage::FVirtualLoopUpdateMessage(const Trace::IAnalyzer::FOnEventContext& InContext)
+		: FVirtualLoopMessageBase(InContext)
+	{
+		const Trace::IAnalyzer::FEventData& EventData = InContext.EventData;
+
+		TimeVirtualized = EventData.GetValue<float>("TimeVirtualized");
+		PlaybackTime = EventData.GetValue<float>("PlaybackTime");
+		UpdateInterval = EventData.GetValue<float>("UpdateInterval");
+
+		LocationX = EventData.GetValue<double>("LocationX");
+		LocationY = EventData.GetValue<double>("LocationY");
+		LocationZ = EventData.GetValue<double>("LocationZ");
+
+		RotatorPitch = EventData.GetValue<double>("RotatorPitch");
+		RotatorYaw = EventData.GetValue<double>("RotatorYaw");
+		RotatorRoll = EventData.GetValue<double>("RotatorRoll");
+	}
+
+	uint32 FVirtualLoopUpdateMessage::GetSizeOf() const
+	{
+		// Returning fixed size of the struct (no dynamically sized members in it)
+		return sizeof(FVirtualLoopUpdateMessage);
+	}
+} // namespace UE::Audio::Insights
diff --git a/Engine/Plugins/AudioInsights/Source/AudioInsights/Private/Providers/VirtualLoopTraceProvider.cpp b/Engine/Plugins/AudioInsights/Source/AudioInsights/Private/Providers/VirtualLoopTraceProvider.cpp
--- a/Engine/Plugins/AudioInsights/Source/AudioInsights/Private/Providers/VirtualLoopTraceProvider.cpp
+++ b/Engine/Plugins/AudioInsights/Source/AudioInsights/Private/Providers/VirtualLoopTraceProvider.cpp
@@ -1,14 +1,16 @@
 // Copyright Epic Games, Inc. All Rights Reserved.
 #include "Providers/VirtualLoopTraceProvider.h"
 
+#include "Async/ParallelFor.h"
+#include "AudioInsightsModule.h"
 #include "Trace/Analyzer.h"
 #include "TraceServices/Model/AnalysisSession.h"
-#include "TraceServices/ModuleService.h"
 
-#if !WITH_EDITOR
+#if WITH_EDITOR
+#include "Cache/AudioInsightsCacheManager.h"
+#else
 #include "Common/PagedArray.h"
-#include "Async/ParallelFor.h"
-#endif // !WITH_EDITOR
+#endif // WITH_EDITOR
 
 namespace UE::Audio::Insights
 {
@@ -41,24 +43,16 @@ namespace UE::Audio::Insights
 
 	void FVirtualLoopTraceProvider::OnTimingViewTimeMarkerChanged(double TimeMarker)
 	{
-		using namespace FVirtualLoopTraceProviderPrivate;
-
 #if !WITH_EDITOR
 		if (!SessionCachedMessages.IsValid())
 		{
 			return;
 		}
+#endif // !WITH_EDITOR
 
-		DeviceDataMap.Empty();
-
-		// Collect all the virtualize messages registered until this point in time 
-		for (const FVirtualLoopVirtualizeMessage& VirtualizeCachedMessage : SessionCachedMessages->VirtualizeCachedMessages)
+		// Helper lambdas
+		auto ProcessVirtualizeMessage = [this](const FVirtualLoopVirtualizeMessage& VirtualizeCachedMessage)
 		{
-			if (VirtualizeCachedMessage.Timestamp > TimeMarker)
-			{
-				break;
-			}
-
 			UpdateDeviceEntry(VirtualizeCachedMessage.DeviceId, VirtualizeCachedMessage.PlayOrder, [&VirtualizeCachedMessage](TSharedPtr<FVirtualLoopDashboardEntry>& Entry)
 			{
 				if (!Entry.IsValid())
@@ -67,89 +61,146 @@ namespace UE::Audio::Insights
 					Entry->DeviceId  = VirtualizeCachedMessage.DeviceId;
 					Entry->PlayOrder = VirtualizeCachedMessage.PlayOrder;
 				}
+				
 				Entry->Timestamp = VirtualizeCachedMessage.Timestamp;
 
 				Entry->Name        = *VirtualizeCachedMessage.Name;
 				Entry->ComponentId = VirtualizeCachedMessage.ComponentId;
 			});
-		}
+		};
 
-		// Selectively remove virtualize messages collected in the step above by knowing which sounds were stopped/realized.
-		// With this we will know what are the virtualized sounds at this point in time.
-		for (const FVirtualLoopMessageBase& StopOrRealizeCachedMessage : SessionCachedMessages->StopOrRealizeCachedMessages)
+		auto ProcessStopOrRealizeMessage = [this](const FVirtualLoopStopOrRealizeMessage& StopOrRealizeCachedMessage)
 		{
-			if (StopOrRealizeCachedMessage.Timestamp > TimeMarker)
+			auto* OutEntry = FindDeviceEntry(StopOrRealizeCachedMessage.DeviceId, StopOrRealizeCachedMessage.PlayOrder);
+			
+			if (OutEntry && (*OutEntry)->Timestamp < StopOrRealizeCachedMessage.Timestamp)
+			{
+				RemoveDeviceEntry(StopOrRealizeCachedMessage.DeviceId, StopOrRealizeCachedMessage.PlayOrder);
+			}
+		};
+
+		// Process cached messages
+		DeviceDataMap.Empty();
+
+#if WITH_EDITOR
+		const FAudioInsightsCacheManager& CacheManager = FAudioInsightsModule::GetChecked().GetCacheManager();
+
+		// Virtualize cached messages
+		CacheManager.IterateTo<FVirtualLoopVirtualizeMessage>(VirtualLoopMessageNames::Virtualize, TimeMarker, [&ProcessVirtualizeMessage](const FVirtualLoopVirtualizeMessage& VirtualizeCachedMessage)
+		{
+			ProcessVirtualizeMessage(VirtualizeCachedMessage);
+		});
+
+		// Stop or realize cached messages
+		CacheManager.IterateTo<FVirtualLoopStopOrRealizeMessage>(VirtualLoopMessageNames::StopOrRealize, TimeMarker, [&ProcessStopOrRealizeMessage](const FVirtualLoopStopOrRealizeMessage& StopOrRealizeCachedMessage)
+		{
+			ProcessStopOrRealizeMessage(StopOrRealizeCachedMessage);
+		});
+#else
+		// Virtualize cached messages
+		for (const FVirtualLoopVirtualizeMessage& VirtualizeCachedMessage : SessionCachedMessages->VirtualizeCachedMessages)
+		{
+			if (VirtualizeCachedMessage.Timestamp > TimeMarker)
 			{
 				break;
 			}
 
-			auto* OutEntry = FindDeviceEntry(StopOrRealizeCachedMessage.DeviceId, StopOrRealizeCachedMessage.PlayOrder);
+			ProcessVirtualizeMessage(VirtualizeCachedMessage);
+		}
 
-			if (OutEntry && (*OutEntry)->Timestamp < StopOrRealizeCachedMessage.Timestamp)
+		// Stop or realize cached messages
+		for (const FVirtualLoopStopOrRealizeMessage& StopOrRealizeCachedMessage : SessionCachedMessages->StopOrRealizeCachedMessages)
+		{
+			if (StopOrRealizeCachedMessage.Timestamp > TimeMarker)
 			{
-				RemoveDeviceEntry(StopOrRealizeCachedMessage.DeviceId, StopOrRealizeCachedMessage.PlayOrder);
+				break;
 			}
+
+			ProcessStopOrRealizeMessage(StopOrRealizeCachedMessage);
 		}
+#endif // WITH_EDITOR
+
+		CollectParamsForTimestamp(TimeMarker);
+
+		// Call parent method to update LastMessageId
+		FTraceProviderBase::OnTimingViewTimeMarkerChanged(TimeMarker);
+	}
+
+	void FVirtualLoopTraceProvider::CollectParamsForTimestamp(const double InTimeMarker)
+	{
+		using namespace FVirtualLoopTraceProviderPrivate;
 
 		const FDeviceData* DeviceData = FindFilteredDeviceData();
-		if (DeviceData)
+		if (DeviceData == nullptr)
 		{
-			// Collect update messages from virtualized sounds (based on active sounds's PlayOrder)
-			struct CachedEntryInfo
-			{
-				FVirtualLoopUpdateMessage UpdateMessage;
-			};
+			return;
+		}
 
-			TArray<uint32> PlayOrderArray;
-			(*DeviceData).GenerateKeyArray(PlayOrderArray);
+		// Collect update messages from virtualized sounds (based on active sounds's PlayOrder)
+		struct CachedEntryInfo
+		{
+			TOptional<FVirtualLoopUpdateMessage> UpdateMessage;
+		};
 
-			TArray<CachedEntryInfo> CachedEntryInfos;
-			CachedEntryInfos.SetNumUninitialized(PlayOrderArray.Num());
+		TArray<uint32> PlayOrderArray;
+		DeviceData->GenerateKeyArray(PlayOrderArray);
 
-			ParallelFor(PlayOrderArray.Num(), 
-			[&PlayOrderArray, &CachedEntryInfos, TimeMarker, this](const int32 Index)
-			{
-				const uint32 PlayOrder = PlayOrderArray[Index];
+		TArray<CachedEntryInfo> CachedEntryInfos;
+		CachedEntryInfos.SetNumUninitialized(PlayOrderArray.Num());
 
-				// Update
-				const FVirtualLoopUpdateMessage* FoundUpdateCachedMessage = FindClosestMessageToTimestamp(SessionCachedMessages->UpdateCachedMessages, TimeMarker, PlayOrder);
-				if (FoundUpdateCachedMessage)
-				{
-					CachedEntryInfos[Index].UpdateMessage = *FoundUpdateCachedMessage;
-				}
-			});
+#if WITH_EDITOR
+		const FAudioInsightsCacheManager& CacheManager = FAudioInsightsModule::GetChecked().GetCacheManager();
+
+		ParallelFor(PlayOrderArray.Num(), [&PlayOrderArray, &CachedEntryInfos, &CacheManager, InTimeMarker, this](const int32 Index)
+#else
+		ParallelFor(PlayOrderArray.Num(), [&PlayOrderArray, &CachedEntryInfos, InTimeMarker, this](const int32 Index)
+#endif // WITH_EDITOR
+		{
+			const uint32 PlayOrder = PlayOrderArray[Index];
 
-			// Update the device entries with the collected info
-			for (const CachedEntryInfo& CachedEntryInfo : CachedEntryInfos)
+#if WITH_EDITOR
+			const FVirtualLoopUpdateMessage* FoundUpdateCachedMessage = CacheManager.FindClosestMessage<FVirtualLoopUpdateMessage>(VirtualLoopMessageNames::Update, InTimeMarker, PlayOrder);
+#else
+			const FVirtualLoopUpdateMessage* FoundUpdateCachedMessage = FindClosestMessageToTimestamp(SessionCachedMessages->UpdateCachedMessages, InTimeMarker, PlayOrder);
+#endif // WITH_EDITOR
+
+			if (FoundUpdateCachedMessage)
 			{
-				UpdateDeviceEntry(CachedEntryInfo.UpdateMessage.DeviceId, CachedEntryInfo.UpdateMessage.PlayOrder, [&CachedEntryInfo](TSharedPtr<FVirtualLoopDashboardEntry>& Entry)
+				CachedEntryInfos[Index].UpdateMessage = *FoundUpdateCachedMessage;
+			}
+		});
+
+		// Update the device entries with the collected info
+		for (const CachedEntryInfo& CachedEntryInfo : CachedEntryInfos)
+		{
+			if (CachedEntryInfo.UpdateMessage.IsSet())
+			{
+				const FVirtualLoopUpdateMessage& VirtualLoopUpdateMessage = CachedEntryInfo.UpdateMessage.GetValue();
+
+				UpdateDeviceEntry(VirtualLoopUpdateMessage.DeviceId, VirtualLoopUpdateMessage.PlayOrder, [&VirtualLoopUpdateMessage](TSharedPtr<FVirtualLoopDashboardEntry>& Entry)
 				{
-					if (!Entry.IsValid())
+					if (ensure(Entry.IsValid()))
 					{
-						Entry = MakeShared<FVirtualLoopDashboardEntry>();
-						Entry->DeviceId  = CachedEntryInfo.UpdateMessage.DeviceId;
-						Entry->PlayOrder = CachedEntryInfo.UpdateMessage.PlayOrder;
-					}
+						Entry->Timestamp = VirtualLoopUpdateMessage.Timestamp;
 
-					Entry->Timestamp = CachedEntryInfo.UpdateMessage.Timestamp;
-
-					Entry->PlaybackTime    = CachedEntryInfo.UpdateMessage.PlaybackTime;
-					Entry->TimeVirtualized = CachedEntryInfo.UpdateMessage.TimeVirtualized;
-					Entry->UpdateInterval  = CachedEntryInfo.UpdateMessage.UpdateInterval;
-					Entry->Location        = FVector{ CachedEntryInfo.UpdateMessage.LocationX, CachedEntryInfo.UpdateMessage.LocationY, CachedEntryInfo.UpdateMessage.LocationZ };
-					Entry->Rotator         = FRotator{ CachedEntryInfo.UpdateMessage.RotatorPitch, CachedEntryInfo.UpdateMessage.RotatorYaw, CachedEntryInfo.UpdateMessage.RotatorRoll };
+						Entry->PlaybackTime    = VirtualLoopUpdateMessage.PlaybackTime;
+						Entry->TimeVirtualized = VirtualLoopUpdateMessage.TimeVirtualized;
+						Entry->UpdateInterval  = VirtualLoopUpdateMessage.UpdateInterval;
+						Entry->Location        = FVector{ VirtualLoopUpdateMessage.LocationX, VirtualLoopUpdateMessage.LocationY, VirtualLoopUpdateMessage.LocationZ };
+						Entry->Rotator         = FRotator{ VirtualLoopUpdateMessage.RotatorPitch, VirtualLoopUpdateMessage.RotatorYaw, VirtualLoopUpdateMessage.RotatorRoll };
+					}
 				});
 			}
 		}
-
-#endif // !WITH_EDITOR
-
-		// Call parent method to update LastMessageId
-		FTraceProviderBase::OnTimingViewTimeMarkerChanged(TimeMarker);
 	}
 
 	bool FVirtualLoopTraceProvider::ProcessMessages()
 	{
+		if (GetTimeControlMethod() != ETimeControlMethod::LatestMessageTime)
+		{
+			return true;
+		}
+
 		auto RemoveEntryFunc = [this](const FVirtualLoopStopOrRealizeMessage& Msg, TSharedPtr<FVirtualLoopDashboardEntry>* OutEntry)
 		{
 #if !WITH_EDITOR
@@ -264,24 +315,43 @@ namespace UE::Audio::Insights
 			{
 				LLM_SCOPE_BYNAME(TEXT("Insights/FVirtualLoopTraceAnalyzer"));
 
+#if WITH_EDITOR
+				if (ShouldIgnoreNewEvents())
+				{
+					return true;
+				}
+#endif // WITH_EDITOR
+
 				FVirtualLoopMessages& Messages = GetProvider<FVirtualLoopTraceProvider>().TraceMessages;
 				switch (RouteId)
 				{
-					case RouteId_Stop:
+					case RouteId_Virtualize:
 					{
-						Messages.StopOrRealizeMessages.Enqueue(FVirtualLoopStopOrRealizeMessage { Context });
+#if WITH_EDITOR
+						CacheMessage<FVirtualLoopVirtualizeMessage>(Context, Messages.VirtualizeMessages);
+#else
+						Messages.VirtualizeMessages.Enqueue(FVirtualLoopVirtualizeMessage{ Context });
+#endif // WITH_EDITOR
 						break;
 					}
 
-					case RouteId_Update:
+					case RouteId_Stop:
 					{
-						Messages.UpdateMessages.Enqueue(FVirtualLoopUpdateMessage { Context });
+#if WITH_EDITOR
+						CacheMessage<FVirtualLoopStopOrRealizeMessage>(Context, Messages.StopOrRealizeMessages);
+#else
+						Messages.StopOrRealizeMessages.Enqueue(FVirtualLoopStopOrRealizeMessage{ Context });
+#endif // WITH_EDITOR
 						break;
 					}
 
-					case RouteId_Virtualize:
+					case RouteId_Update:
 					{
-						Messages.VirtualizeMessages.Enqueue(FVirtualLoopVirtualizeMessage { Context });
+#if WITH_EDITOR
+						CacheMessage<FVirtualLoopUpdateMessage>(Context, Messages.UpdateMessages);
+#else
+						Messages.UpdateMessages.Enqueue(FVirtualLoopUpdateMessage{ Context });
+#endif // WITH_EDITOR
 						break;
 					}
 
diff --git a/Engine/Plugins/AudioInsights/Source/AudioInsights/Private/Providers/VirtualLoopTraceProvider.h b/Engine/Plugins/AudioInsights/Source/AudioInsights/Private/Providers/VirtualLoopTraceProvider.h
--- a/Engine/Plugins/AudioInsights/Source/AudioInsights/Private/Providers/VirtualLoopTraceProvider.h
+++ b/Engine/Plugins/AudioInsights/Source/AudioInsights/Private/Providers/VirtualLoopTraceProvider.h
@@ -29,6 +29,7 @@ namespace UE::Audio::Insights
 
 	private:
 		virtual void OnTimingViewTimeMarkerChanged(double TimeMarker) override;
+		void CollectParamsForTimestamp(const double InTimeMarker);
 
 #if !WITH_EDITOR
 		TUniquePtr<FVirtualLoopSessionCachedMessages> SessionCachedMessages;
diff --git a/Engine/Plugins/AudioInsights/Source/AudioInsights/Public/Messages/VirtualLoopTraceMessages.h b/Engine/Plugins/AudioInsights/Source/AudioInsights/Public/Messages/VirtualLoopTraceMessages.h
--- a/Engine/Plugins/AudioInsights/Source/AudioInsights/Public/Messages/VirtualLoopTraceMessages.h
+++ b/Engine/Plugins/AudioInsights/Source/AudioInsights/Public/Messages/VirtualLoopTraceMessages.h
@@ -2,6 +2,7 @@
 #pragma once
 
 #include "AudioDefines.h"
+#include "Cache/IAudioCachedMessage.h"
 #include "Math/NumericLimits.h"
 #include "Messages/AnalyzerMessageQueue.h"
 #include "Trace/Analyzer.h"
@@ -14,58 +15,52 @@
 
 namespace UE::Audio::Insights
 {
-	struct FVirtualLoopMessageBase
+	namespace VirtualLoopMessageNames
+	{
+		extern const FName Virtualize;
+		extern const FName StopOrRealize;
+		extern const FName Update;
+	};
+
+	struct FVirtualLoopMessageBase : public IAudioCachedMessage
 	{
 		FVirtualLoopMessageBase() = default;
-		FVirtualLoopMessageBase(const Trace::IAnalyzer::FOnEventContext& InContext)
-		{
-			const Trace::IAnalyzer::FEventData& EventData = InContext.EventData;
-			DeviceId = static_cast<::Audio::FDeviceId>(EventData.GetValue<uint32>("DeviceId"));
-			Timestamp = InContext.EventTime.AsSeconds(EventData.GetValue<uint64>("Timestamp"));
-			PlayOrder = EventData.GetValue<uint32>("PlayOrder");
-		}
+		FVirtualLoopMessageBase(const Trace::IAnalyzer::FOnEventContext& InContext);
+
+		virtual uint32 GetID() const override { return PlayOrder; }
 
 		::Audio::FDeviceId DeviceId = INDEX_NONE;
 		uint32 PlayOrder = INDEX_NONE;
-		double Timestamp = 0.0;
 	};
 
-	using FVirtualLoopStopOrRealizeMessage = FVirtualLoopMessageBase;
-
 	struct FVirtualLoopVirtualizeMessage : public FVirtualLoopMessageBase
 	{
 		FVirtualLoopVirtualizeMessage() = default;
-		FVirtualLoopVirtualizeMessage(const Trace::IAnalyzer::FOnEventContext& InContext)
-			: FVirtualLoopMessageBase(InContext)
-		{
-			const Trace::IAnalyzer::FEventData& EventData = InContext.EventData;
-			EventData.GetString("Name", Name);
-			ComponentId = EventData.GetValue<uint64>("ComponentId");
-		}
+		FVirtualLoopVirtualizeMessage(const Trace::IAnalyzer::FOnEventContext& InContext);
+
+		virtual const FName GetMessageName() const override { return VirtualLoopMessageNames::Virtualize; }
+		virtual uint32 GetSizeOf() const override;
 
 		FString Name;
 		uint64 ComponentId = TNumericLimits<uint64>::Max();
 	};
 
+	struct FVirtualLoopStopOrRealizeMessage : public FVirtualLoopMessageBase
+	{
+		FVirtualLoopStopOrRealizeMessage() = default;
+		FVirtualLoopStopOrRealizeMessage(const Trace::IAnalyzer::FOnEventContext& InContext);
+
+		virtual const FName GetMessageName() const override { return VirtualLoopMessageNames::StopOrRealize; }
+		virtual uint32 GetSizeOf() const override;
+	};
+
 	struct FVirtualLoopUpdateMessage : public FVirtualLoopMessageBase
 	{
 		FVirtualLoopUpdateMessage() = default;
-		FVirtualLoopUpdateMessage(const Trace::IAnalyzer::FOnEventContext& InContext)
-			: FVirtualLoopMessageBase(InContext)
-		{
-			const Trace::IAnalyzer::FEventData& EventData = InContext.EventData;
-			TimeVirtualized = EventData.GetValue<float>("TimeVirtualized");
-			PlaybackTime = EventData.GetValue<float>("PlaybackTime");
-			UpdateInterval = EventData.GetValue<float>("UpdateInterval");
-
-			LocationX = EventData.GetValue<double>("LocationX");
-			LocationY = EventData.GetValue<double>("LocationY");
-			LocationZ = EventData.GetValue<double>("LocationZ");
-
-			RotatorPitch = EventData.GetValue<double>("RotatorPitch");
-			RotatorYaw = EventData.GetValue<double>("RotatorYaw");
-			RotatorRoll = EventData.GetValue<double>("RotatorRoll");
-		}
+		FVirtualLoopUpdateMessage(const Trace::IAnalyzer::FOnEventContext& InContext);
+
+		virtual const FName GetMessageName() const override { return VirtualLoopMessageNames::Update; }
+		virtual uint32 GetSizeOf() const override;
 
 		float TimeVirtualized = 0.0f;
 		float PlaybackTime = 0.0f;
