[Audio Insights] Support caching in the Audio Bus Dashboard

diff --git a/Engine/Plugins/AudioInsights/Source/AudioInsights/Private/Messages/AudioBusTraceMessages.cpp b/Engine/Plugins/AudioInsights/Source/AudioInsights/Private/Messages/AudioBusTraceMessages.cpp
new file mode 100644
--- /dev/null
+++ b/Engine/Plugins/AudioInsights/Source/AudioInsights/Private/Messages/AudioBusTraceMessages.cpp
@@ -0,0 +1,94 @@
+// Copyright Epic Games, Inc. All Rights Reserved.
+#include "Messages/AudioBusTraceMessages.h"
+
+namespace UE::Audio::Insights
+{
+	namespace AudioBusMessageNames
+	{
+		const FName Start = "AudioBusStart";
+		const FName HasActivity = "AudioBusHasActivity";
+		const FName EnvelopeValues = "AudioBusEnvelopeValues";
+		const FName Stop = "AudioBusStop";
+	};
+
+	FAudioBusMessageBase::FAudioBusMessageBase(const Trace::IAnalyzer::FOnEventContext& InContext)
+    {
+		const UE::Trace::IAnalyzer::FEventData& EventData = InContext.EventData;
+
+		DeviceId = static_cast<::Audio::FDeviceId>(EventData.GetValue<uint32>("DeviceId"));
+		AudioBusId = EventData.GetValue<uint32>("AudioBusId");
+		Timestamp = InContext.EventTime.AsSeconds(EventData.GetValue<uint64>("Timestamp"));
+    }
+
+	FAudioBusStartMessage::FAudioBusStartMessage(const Trace::IAnalyzer::FOnEventContext& InContext)
+		: FAudioBusMessageBase(InContext)
+	{
+		const UE::Trace::IAnalyzer::FEventData& EventData = InContext.EventData;
+
+		EventData.GetString("Name", Name);
+		NumChannels = EventData.GetValue<int32>("NumChannels");
+
+		AudioBusType = FSoftObjectPath(Name).IsAsset() ? EAudioBusType::AssetBased : EAudioBusType::CodeGenerated;
+	}
+
+	uint32 FAudioBusStartMessage::GetSizeOf() const
+	{
+		// Start with the fixed size of the struct
+		uint32 MemorySize = sizeof(FAudioBusStartMessage);
+
+		// Add any dynamically sized members
+		MemorySize += Name.GetAllocatedSize();
+
+		return MemorySize;
+	}
+
+	FAudioBusHasActivityMessage::FAudioBusHasActivityMessage(const Trace::IAnalyzer::FOnEventContext& InContext)
+		: FAudioBusMessageBase(InContext)
+	{
+		const UE::Trace::IAnalyzer::FEventData& EventData = InContext.EventData;
+
+		bHasActivity = EventData.GetValue<bool>("HasActivity");
+	}
+
+	uint32 FAudioBusHasActivityMessage::GetSizeOf() const
+	{
+		return sizeof(FAudioBusHasActivityMessage);
+	}
+
+	FAudioBusEnvelopeValuesMessage::FAudioBusEnvelopeValuesMessage(const Trace::IAnalyzer::FOnEventContext& InContext)
+		: FAudioBusMessageBase(InContext)
+	{
+		const UE::Trace::IAnalyzer::FEventData& EventData = InContext.EventData;
+
+		const TArrayView<const float> EnvelopeValuesView = EventData.GetArrayView<float>("EnvelopeValues");
+
+		for (const auto& EnvelopeValue : EnvelopeValuesView)
+		{
+			EnvelopeValues.Emplace(EnvelopeValue);
+		}
+	}
+
+	uint32 FAudioBusEnvelopeValuesMessage::GetSizeOf() const
+	{
+		// Start with the fixed size of the struct
+		uint32 MemorySize = sizeof(FAudioBusEnvelopeValuesMessage);
+
+		// Add any dynamically sized members
+		if (EnvelopeValues.Num() > 0)
+		{
+			MemorySize += EnvelopeValues.Num() * sizeof(EnvelopeValues[0]);
+		}
+
+		return MemorySize;
+	}
+
+	FAudioBusStopMessage::FAudioBusStopMessage(const Trace::IAnalyzer::FOnEventContext& InContext)
+		: FAudioBusMessageBase(InContext)
+	{
+	}
+
+	uint32 FAudioBusStopMessage::GetSizeOf() const
+	{
+		return sizeof(FAudioBusStopMessage);
+	}
+}
\ No newline at end of file
diff --git a/Engine/Plugins/AudioInsights/Source/AudioInsights/Private/Providers/AudioBusTraceProvider.cpp b/Engine/Plugins/AudioInsights/Source/AudioInsights/Private/Providers/AudioBusTraceProvider.cpp
--- a/Engine/Plugins/AudioInsights/Source/AudioInsights/Private/Providers/AudioBusTraceProvider.cpp
+++ b/Engine/Plugins/AudioInsights/Source/AudioInsights/Private/Providers/AudioBusTraceProvider.cpp
@@ -2,11 +2,14 @@
 
 #include "Providers/AudioBusTraceProvider.h"
 
+#include "AudioInsightsModule.h"
+#include "Async/ParallelFor.h"
+#include "Cache/AudioInsightsCacheManager.h"
+
 #if WITH_EDITOR
 #include "AudioDevice.h"
 #include "AudioDeviceManager.h"
 #else
-#include "Async/ParallelFor.h"
 #include "Common/PagedArray.h"
 #endif // WITH_EDITOR
 
@@ -86,31 +89,54 @@ namespace UE::Audio::Insights
 			{
 				LLM_SCOPE_BYNAME(TEXT("Insights/FAudioBusTraceAnalyzer"));
 
+#if WITH_EDITOR
+				if (ShouldIgnoreNewEvents())
+				{
+					return true;
+				}
+#endif // WITH_EDITOR
+
 				FAudioBusMessages& Messages = GetProvider<FAudioBusTraceProvider>().TraceMessages;
 
 				switch (RouteId)
 				{
 					case RouteId_Start:
 					{
+#if WITH_EDITOR
+						CacheMessage<FAudioBusStartMessage>(Context, Messages.StartMessages);
+#else
 						Messages.StartMessages.Enqueue(FAudioBusStartMessage{ Context });
+#endif // WITH_EDITOR
 						break;
 					}
 
 					case RouteId_HasActivity:
 					{
+#if WITH_EDITOR
+						CacheMessage<FAudioBusHasActivityMessage>(Context, Messages.HasActivityMessages);
+#else
 						Messages.HasActivityMessages.Enqueue(FAudioBusHasActivityMessage{ Context });
+#endif // WITH_EDITOR
 						break;
 					}
 
 					case RouteId_EnvelopeValues:
 					{
+#if WITH_EDITOR
+						CacheMessage<FAudioBusEnvelopeValuesMessage>(Context, Messages.EnvelopeValuesMessages);
+#else
 						Messages.EnvelopeValuesMessages.Enqueue(FAudioBusEnvelopeValuesMessage{ Context });
+#endif // WITH_EDITOR
 						break;
 					}
 
 					case RouteId_Stop:
 					{
+#if WITH_EDITOR
+						CacheMessage<FAudioBusStopMessage>(Context, Messages.StopMessages);
+#else
 						Messages.StopMessages.Enqueue(FAudioBusStopMessage{ Context });
+#endif // WITH_EDITOR
 						break;
 					}
 
@@ -147,6 +173,11 @@ namespace UE::Audio::Insights
 
 	bool FAudioBusTraceProvider::ProcessMessages()
 	{
+		if (GetTimeControlMethod() != ETimeControlMethod::LatestMessageTime)
+		{
+			return true;
+		}
+
 		// Helper lambdas
 		auto CreateEntry = [this](const FAudioBusMessageBase& Msg)
 		{
@@ -226,6 +257,13 @@ namespace UE::Audio::Insights
 		ProcessMessageQueue<FAudioBusEnvelopeValuesMessage>(TraceMessages.EnvelopeValuesMessages, GetEntry,
 		[this](const FAudioBusEnvelopeValuesMessage& Msg, TSharedPtr<FAudioBusDashboardEntry>* OutEntry)
 		{
+#if !WITH_EDITOR
+			if (SessionCachedMessages.IsValid())
+			{
+				SessionCachedMessages->EnvelopeValuesCachedMessages.EmplaceBack(Msg);
+			}
+#endif // !WITH_EDITOR
+
 			if (OutEntry)
 			{
 				FAudioBusDashboardEntry& EntryRef = *OutEntry->Get();
@@ -360,42 +398,72 @@ namespace UE::Audio::Insights
 
 	void FAudioBusTraceProvider::OnTimingViewTimeMarkerChanged(double TimeMarker)
 	{
-#if !WITH_EDITOR
 		using namespace FAudioBusTraceProviderPrivate;
 
+#if !WITH_EDITOR
 		if (!SessionCachedMessages.IsValid())
 		{
 			return;
 		}
+#endif // !WITH_EDITOR
 
 		DeviceDataMap.Empty();
 
 		// Collect all the audio bus start messages registered until this point in time 
-		for (const FAudioBusStartMessage& AudioBusStartCachedMessage : SessionCachedMessages->StartCachedMessages)
+		auto ProcessBusStartMessage = [this](const FAudioBusStartMessage& Msg)
 		{
-			if (AudioBusStartCachedMessage.Timestamp > TimeMarker)
-			{
-				break;
-			}
-
-			UpdateDeviceEntry(AudioBusStartCachedMessage.DeviceId, AudioBusStartCachedMessage.AudioBusId, [&AudioBusStartCachedMessage](TSharedPtr<FAudioBusDashboardEntry>& Entry)
+			UpdateDeviceEntry(Msg.DeviceId, Msg.AudioBusId, [&Msg](TSharedPtr<FAudioBusDashboardEntry>& Entry)
 			{
 				if (!Entry.IsValid())
 				{
 					Entry = MakeShared<FAudioBusDashboardEntry>();
-					Entry->AudioBusId = AudioBusStartCachedMessage.AudioBusId;
+					Entry->AudioBusId = Msg.AudioBusId;
 				}
 
-				Entry->Timestamp = AudioBusStartCachedMessage.Timestamp;
+				Entry->Timestamp = Msg.Timestamp;
 
-				Entry->Name = *AudioBusStartCachedMessage.Name;
-				Entry->AudioMeterInfo->NumChannels = AudioBusStartCachedMessage.NumChannels;
-				Entry->AudioBusType = AudioBusStartCachedMessage.AudioBusType;
+				Entry->Name = *Msg.Name;
+				Entry->AudioMeterInfo->NumChannels = Msg.NumChannels;
+				Entry->AudioBusType = Msg.AudioBusType;
 			});
+		};
+
+#if WITH_EDITOR
+		const FAudioInsightsCacheManager& CacheManager = FAudioInsightsModule::GetChecked().GetCacheManager();
+		CacheManager.IterateTo<FAudioBusStartMessage>(AudioBusMessageNames::Start, TimeMarker, [&ProcessBusStartMessage](const FAudioBusStartMessage& AudioBusStartCachedMessage)
+		{
+			ProcessBusStartMessage(AudioBusStartCachedMessage);
+		});
+#else
+		for (const FAudioBusStartMessage& AudioBusStartCachedMessage : SessionCachedMessages->StartCachedMessages)
+		{
+			if (AudioBusStartCachedMessage.Timestamp > TimeMarker)
+			{
+				break;
+			}
+
+			ProcessBusStartMessage(AudioBusStartCachedMessage);
 		}
+#endif // WITH_EDITOR
 
 		// Selectively remove start messages collected in the step above by knowing which audio buses were stopped.
 		// With this we will know what are the available audio buses at this point in time.
+		auto ProcessBusStopMessage = [this](const FAudioBusStopMessage& Msg)
+		{
+			auto* OutEntry = FindDeviceEntry(Msg.DeviceId, Msg.AudioBusId);
+
+			if (OutEntry && (*OutEntry)->Timestamp < Msg.Timestamp)
+			{
+				RemoveDeviceEntry(Msg.DeviceId, Msg.AudioBusId);
+			}
+		};
+
+#if WITH_EDITOR
+		CacheManager.IterateTo<FAudioBusStopMessage>(AudioBusMessageNames::Stop, TimeMarker, [&ProcessBusStopMessage](const FAudioBusStopMessage& AudioBusStopCachedMessage)
+		{
+			ProcessBusStopMessage(AudioBusStopCachedMessage);
+		});
+#else
 		for (const FAudioBusStopMessage& AudioBusStopCachedMessage : SessionCachedMessages->StopCachedMessages)
 		{
 			if (AudioBusStopCachedMessage.Timestamp > TimeMarker)
@@ -403,13 +471,9 @@ namespace UE::Audio::Insights
 				break;
 			}
 
-			auto* OutEntry = FindDeviceEntry(AudioBusStopCachedMessage.DeviceId, AudioBusStopCachedMessage.AudioBusId);
-
-			if (OutEntry && (*OutEntry)->Timestamp < AudioBusStopCachedMessage.Timestamp)
-			{
-				RemoveDeviceEntry(AudioBusStopCachedMessage.DeviceId, AudioBusStopCachedMessage.AudioBusId);
-			}
+			ProcessBusStopMessage(AudioBusStopCachedMessage);
 		}
+#endif // WITH_EDITOR
 
 		const FDeviceData* DeviceData = FindFilteredDeviceData();
 		if (DeviceData)
@@ -418,6 +482,7 @@ namespace UE::Audio::Insights
 			struct CachedEntryInfo
 			{
 				TOptional<FAudioBusHasActivityMessage> HasActivityMessage;
+				TOptional<FAudioBusEnvelopeValuesMessage> EnvelopeValuesMessage;
 			};
 
 			TArray<uint32> AudioBusIdArray;
@@ -427,16 +492,35 @@ namespace UE::Audio::Insights
 			CachedEntryInfos.SetNum(AudioBusIdArray.Num());
 
 			ParallelFor(AudioBusIdArray.Num(),
+#if WITH_EDITOR
+			[&AudioBusIdArray, &CachedEntryInfos, &CacheManager, TimeMarker, this](const int32 Index)
+#else
 			[&AudioBusIdArray, &CachedEntryInfos, TimeMarker, this](const int32 Index)
+#endif // WITH_EDITOR
 			{
 				const uint32 AudioBusId = AudioBusIdArray[Index];
 
 				// HasActivity
+#if WITH_EDITOR
+				const FAudioBusHasActivityMessage* FoundHasActivityCachedMessage = CacheManager.FindClosestMessage<FAudioBusHasActivityMessage>(AudioBusMessageNames::HasActivity, TimeMarker, AudioBusId);
+#else
 				const FAudioBusHasActivityMessage* FoundHasActivityCachedMessage = FindClosestMessageToTimestamp(SessionCachedMessages->HasActivityCachedMessages, TimeMarker, AudioBusId);
+#endif // WITH_EDITOR
 				if (FoundHasActivityCachedMessage)
 				{
 					CachedEntryInfos[Index].HasActivityMessage = *FoundHasActivityCachedMessage;
 				}
+
+				// Envelope Values
+#if WITH_EDITOR
+				const FAudioBusEnvelopeValuesMessage* FoundEnvelopeValuesCachedMessage = CacheManager.FindClosestMessage<FAudioBusEnvelopeValuesMessage>(AudioBusMessageNames::EnvelopeValues, TimeMarker, AudioBusId);
+#else
+				const FAudioBusEnvelopeValuesMessage* FoundEnvelopeValuesCachedMessage = FindClosestMessageToTimestamp(SessionCachedMessages->EnvelopeValuesCachedMessages, TimeMarker, AudioBusId);
+#endif // WITH_EDITOR
+				if (FoundEnvelopeValuesCachedMessage)
+				{
+					CachedEntryInfos[Index].EnvelopeValuesMessage = *FoundEnvelopeValuesCachedMessage;
+				}
 			});
 
 			// Update the device entries with the collected info
@@ -459,12 +543,36 @@ namespace UE::Audio::Insights
 						Entry->bHasActivity = HasActivityMessage.bHasActivity;
 					});
 				}
+
+				if (CachedEntryInfo.EnvelopeValuesMessage.IsSet())
+				{
+					const FAudioBusEnvelopeValuesMessage& EnvelopeValuesMessage = CachedEntryInfo.EnvelopeValuesMessage.GetValue();
+
+					UpdateDeviceEntry(EnvelopeValuesMessage.DeviceId, EnvelopeValuesMessage.AudioBusId, [&EnvelopeValuesMessage](TSharedPtr<FAudioBusDashboardEntry>& Entry)
+					{
+						if (!Entry.IsValid())
+						{
+							Entry = MakeShared<FAudioBusDashboardEntry>();
+							Entry->AudioBusId = EnvelopeValuesMessage.AudioBusId;
+						}
+
+						Entry->Timestamp = EnvelopeValuesMessage.Timestamp;
+
+						Entry->AudioMeterInfo->NumChannels = EnvelopeValuesMessage.EnvelopeValues.Num();
+						Entry->AudioMeterInfo->EnvelopeValues = EnvelopeValuesMessage.EnvelopeValues;
+					});
+				}
 			}
 		}
-#endif // !WITH_EDITOR
 
 		// Call parent method to update LastMessageId
 		FTraceProviderBase::OnTimingViewTimeMarkerChanged(TimeMarker);
 	}
 
+#if WITH_EDITOR
+	void FAudioBusTraceProvider::OnTimeControlMethodReset()
+	{
+		AudioBusAssetProvider.RequestEntriesUpdate();
+	}
+#endif
 } // namespace UE::Audio::Insights
diff --git a/Engine/Plugins/AudioInsights/Source/AudioInsights/Private/Providers/AudioBusTraceProvider.h b/Engine/Plugins/AudioInsights/Source/AudioInsights/Private/Providers/AudioBusTraceProvider.h
--- a/Engine/Plugins/AudioInsights/Source/AudioInsights/Private/Providers/AudioBusTraceProvider.h
+++ b/Engine/Plugins/AudioInsights/Source/AudioInsights/Private/Providers/AudioBusTraceProvider.h
@@ -41,6 +41,8 @@ namespace UE::Audio::Insights
 		virtual void OnTimingViewTimeMarkerChanged(double TimeMarker) override;
 
 #if WITH_EDITOR
+		virtual void OnTimeControlMethodReset() override;
+
 		void HandleOnAudioBusAssetAdded(const TSharedPtr<FAudioBusDashboardEntry>& InAudioBusDashboardEntry);
 		void HandleOnAudioBusAssetRemoved(const uint32 InAudioBusId);
 		void HandleOnAudioBusAssetListUpdated(const TArray<TSharedPtr<FAudioBusDashboardEntry>>& InAudioBusAssetEntries);
diff --git a/Engine/Plugins/AudioInsights/Source/AudioInsights/Public/Messages/AudioBusTraceMessages.h b/Engine/Plugins/AudioInsights/Source/AudioInsights/Public/Messages/AudioBusTraceMessages.h
--- a/Engine/Plugins/AudioInsights/Source/AudioInsights/Public/Messages/AudioBusTraceMessages.h
+++ b/Engine/Plugins/AudioInsights/Source/AudioInsights/Public/Messages/AudioBusTraceMessages.h
@@ -1,6 +1,7 @@
 // Copyright Epic Games, Inc. All Rights Reserved.
 #pragma once
 
+#include "Cache/IAudioCachedMessage.h"
 #include "Providers/AudioMeterProvider.h"
 #include "Messages/AnalyzerMessageQueue.h"
 #include "Trace/Analyzer.h"
@@ -13,6 +14,14 @@
 
 namespace UE::Audio::Insights
 {
+	namespace AudioBusMessageNames
+	{
+		extern const FName Start;
+		extern const FName HasActivity;
+		extern const FName EnvelopeValues;
+		extern const FName Stop;
+	};
+
 	enum class EAudioBusType : uint8
 	{
 		AssetBased,
@@ -20,38 +29,24 @@ namespace UE::Audio::Insights
 		None
 	};
 
-	struct FAudioBusMessageBase
+	struct FAudioBusMessageBase : public IAudioCachedMessage
 	{
 		FAudioBusMessageBase() = default;
+		FAudioBusMessageBase(const UE::Trace::IAnalyzer::FOnEventContext& InContext);
 
-		FAudioBusMessageBase(const UE::Trace::IAnalyzer::FOnEventContext& InContext)
-		{
-			const UE::Trace::IAnalyzer::FEventData& EventData = InContext.EventData;
-
-			DeviceId  = static_cast<::Audio::FDeviceId>(EventData.GetValue<uint32>("DeviceId"));
-			AudioBusId  = EventData.GetValue<uint32>("AudioBusId");
-			Timestamp = InContext.EventTime.AsSeconds(EventData.GetValue<uint64>("Timestamp"));
-		}
+		virtual uint32 GetID() const override { return AudioBusId; }
 
 		::Audio::FDeviceId DeviceId = INDEX_NONE;
 		uint32 AudioBusId = INDEX_NONE;
-		double Timestamp = 0.0;
 	};
 
 	struct FAudioBusStartMessage : public FAudioBusMessageBase
 	{
 		FAudioBusStartMessage() = default;
+		FAudioBusStartMessage(const UE::Trace::IAnalyzer::FOnEventContext& InContext);
 
-		FAudioBusStartMessage(const UE::Trace::IAnalyzer::FOnEventContext& InContext)
-			: FAudioBusMessageBase(InContext)
-		{
-			const UE::Trace::IAnalyzer::FEventData& EventData = InContext.EventData;
-
-			EventData.GetString("Name", Name);
-			NumChannels = EventData.GetValue<int32>("NumChannels");
-
-			AudioBusType = FSoftObjectPath(Name).IsAsset() ? EAudioBusType::AssetBased : EAudioBusType::CodeGenerated;
-		}
+		virtual uint32 GetSizeOf() const override;
+		virtual const FName GetMessageName() const override { return AudioBusMessageNames::Start; }
 
 		FString Name;
 		int32 NumChannels = 0;
@@ -61,14 +56,10 @@ namespace UE::Audio::Insights
 	struct FAudioBusHasActivityMessage : public FAudioBusMessageBase
 	{
 		FAudioBusHasActivityMessage() = default;
+		FAudioBusHasActivityMessage(const UE::Trace::IAnalyzer::FOnEventContext& InContext);
 
-		FAudioBusHasActivityMessage(const UE::Trace::IAnalyzer::FOnEventContext& InContext)
-			: FAudioBusMessageBase(InContext)
-		{
-			const UE::Trace::IAnalyzer::FEventData& EventData = InContext.EventData;
-
-			bHasActivity = EventData.GetValue<bool>("HasActivity");
-		}
+		virtual uint32 GetSizeOf() const override;
+		virtual const FName GetMessageName() const override { return AudioBusMessageNames::HasActivity; }
 
 		bool bHasActivity = false;
 	};
@@ -76,24 +67,22 @@ namespace UE::Audio::Insights
 	struct FAudioBusEnvelopeValuesMessage : public FAudioBusMessageBase
 	{
 		FAudioBusEnvelopeValuesMessage() = default;
+		FAudioBusEnvelopeValuesMessage(const UE::Trace::IAnalyzer::FOnEventContext& InContext);
 
-		FAudioBusEnvelopeValuesMessage(const UE::Trace::IAnalyzer::FOnEventContext& InContext)
-			: FAudioBusMessageBase(InContext)
-		{
-			const UE::Trace::IAnalyzer::FEventData& EventData = InContext.EventData;
-
-			const TArrayView<const float> EnvelopeValuesView = EventData.GetArrayView<float>("EnvelopeValues");
-
-			for (const auto& EnvelopeValue : EnvelopeValuesView)
-			{
-				EnvelopeValues.Emplace(EnvelopeValue);
-			}
-		}
+		virtual uint32 GetSizeOf() const override;
+		virtual const FName GetMessageName() const override { return AudioBusMessageNames::EnvelopeValues; }
 
 		TArray<float> EnvelopeValues;
 	};
 
-	using FAudioBusStopMessage = FAudioBusMessageBase;
+	struct FAudioBusStopMessage : public FAudioBusMessageBase
+	{
+		FAudioBusStopMessage() = default;
+		FAudioBusStopMessage(const UE::Trace::IAnalyzer::FOnEventContext& InContext);
+
+		virtual uint32 GetSizeOf() const override;
+		virtual const FName GetMessageName() const override { return AudioBusMessageNames::Stop; }
+	};
 
 	class FAudioBusMessages
 	{
@@ -111,6 +100,7 @@ namespace UE::Audio::Insights
 		FAudioBusSessionCachedMessages(TraceServices::IAnalysisSession& InSession)
 			: StartCachedMessages(InSession.GetLinearAllocator(), 4096)
 			, HasActivityCachedMessages(InSession.GetLinearAllocator(), 4096)
+			, EnvelopeValuesCachedMessages(InSession.GetLinearAllocator(), 4096)
 			, StopCachedMessages(InSession.GetLinearAllocator(), 4096)
 		{
 
@@ -118,6 +108,7 @@ namespace UE::Audio::Insights
 
 		TraceServices::TPagedArray<FAudioBusStartMessage> StartCachedMessages;
 		TraceServices::TPagedArray<FAudioBusHasActivityMessage> HasActivityCachedMessages;
+		TraceServices::TPagedArray<FAudioBusEnvelopeValuesMessage> EnvelopeValuesCachedMessages;
 		TraceServices::TPagedArray<FAudioBusStopMessage> StopCachedMessages;
 	};
 #endif // !WITH_EDITOR
