[Audio Insights] Clear the provider submix envelope value array before initializing from the trace message (BugHawk suggestion)

diff --git a/Engine/Plugins/AudioInsights/Source/AudioInsights/Private/Providers/SubmixTraceProvider.cpp b/Engine/Plugins/AudioInsights/Source/AudioInsights/Private/Providers/SubmixTraceProvider.cpp
--- a/Engine/Plugins/AudioInsights/Source/AudioInsights/Private/Providers/SubmixTraceProvider.cpp
+++ b/Engine/Plugins/AudioInsights/Source/AudioInsights/Private/Providers/SubmixTraceProvider.cpp
@@ -542,10 +542,7 @@ namespace UE::Audio::Insights
 						Entry->Timestamp = EnvelopeValuesMessage.Timestamp;
 
 						Entry->AudioMeterInfo->NumChannels = EnvelopeValuesMessage.EnvelopeValues.Num();
-						for (const auto& EnvelopeValue : EnvelopeValuesMessage.EnvelopeValues)
-						{
-							Entry->AudioMeterInfo->EnvelopeValues.Emplace(EnvelopeValue);
-						}
+						Entry->AudioMeterInfo->EnvelopeValues = EnvelopeValuesMessage.EnvelopeValues;
 					});
 				}
 			}
