[Audio Insights] Move methods for managing pausing/resuming trace data in Audio Insights providers to the base

diff --git a/Engine/Plugins/AudioInsights/Source/AudioInsights/Private/AudioInsightsTraceProviderBase.cpp b/Engine/Plugins/AudioInsights/Source/AudioInsights/Private/AudioInsightsTraceProviderBase.cpp
--- a/Engine/Plugins/AudioInsights/Source/AudioInsights/Private/AudioInsightsTraceProviderBase.cpp
+++ b/Engine/Plugins/AudioInsights/Source/AudioInsights/Private/AudioInsightsTraceProviderBase.cpp
@@ -43,6 +43,18 @@ namespace UE::Audio::Insights
 		Provider->Reset();
 	}
 
+	bool FTraceProviderBase::FTraceAnalyzerBase::OnEvent(uint16 RouteId, EStyle Style, const FOnEventContext& Context)
+	{
+#if WITH_EDITOR
+		if (ShouldIgnoreNewEvents())
+		{
+			return true;
+		}
+#endif // WITH_EDITOR
+
+		return OnHandleEvent(RouteId, Style, Context);
+	}
+
 	bool FTraceProviderBase::FTraceAnalyzerBase::OnEventSuccess(uint16 RouteId, EStyle Style, const FOnEventContext& Context)
 	{
 		++(Provider->LastMessageId);
@@ -61,7 +73,16 @@ namespace UE::Audio::Insights
 
 	bool FTraceProviderBase::FTraceAnalyzerBase::ShouldIgnoreNewEvents() const
 	{
-		return Provider->GetTimeControlMethod() != ETimeControlMethod::LatestMessageTime;
+		return Provider->ProcessMessagesIsPaused();
+	}
+
+	bool FTraceProviderBase::ProcessMessagesIsPaused() const
+	{
+#if WITH_EDITOR
+		return GetTimeControlMethod() != ETimeControlMethod::LatestMessageTime;
+#else
+		return false;
+#endif // WITH_EDITOR
 	}
 
 	ETimeControlMethod FTraceProviderBase::GetTimeControlMethod() const
diff --git a/Engine/Plugins/AudioInsights/Source/AudioInsights/Private/Providers/AudioBusTraceProvider.cpp b/Engine/Plugins/AudioInsights/Source/AudioInsights/Private/Providers/AudioBusTraceProvider.cpp
--- a/Engine/Plugins/AudioInsights/Source/AudioInsights/Private/Providers/AudioBusTraceProvider.cpp
+++ b/Engine/Plugins/AudioInsights/Source/AudioInsights/Private/Providers/AudioBusTraceProvider.cpp
@@ -85,17 +85,10 @@ namespace UE::Audio::Insights
 				Builder.RouteEvent(RouteId_Stop,           "Audio", "AudioBusStop");
 			}
 
-			virtual bool OnEvent(uint16 RouteId, EStyle Style, const FOnEventContext& Context) override
+			virtual bool OnHandleEvent(uint16 RouteId, EStyle Style, const FOnEventContext& Context) override
 			{
 				LLM_SCOPE_BYNAME(TEXT("Insights/FAudioBusTraceAnalyzer"));
 
-#if WITH_EDITOR
-				if (ShouldIgnoreNewEvents())
-				{
-					return true;
-				}
-#endif // WITH_EDITOR
-
 				FAudioBusMessages& Messages = GetProvider<FAudioBusTraceProvider>().TraceMessages;
 
 				switch (RouteId)
@@ -173,11 +166,6 @@ namespace UE::Audio::Insights
 
 	bool FAudioBusTraceProvider::ProcessMessages()
 	{
-		if (GetTimeControlMethod() != ETimeControlMethod::LatestMessageTime)
-		{
-			return true;
-		}
-
 		// Helper lambdas
 		auto CreateEntry = [this](const FAudioBusMessageBase& Msg)
 		{
diff --git a/Engine/Plugins/AudioInsights/Source/AudioInsights/Private/Providers/AudioEventLogTraceProvider.cpp b/Engine/Plugins/AudioInsights/Source/AudioInsights/Private/Providers/AudioEventLogTraceProvider.cpp
--- a/Engine/Plugins/AudioInsights/Source/AudioInsights/Private/Providers/AudioEventLogTraceProvider.cpp
+++ b/Engine/Plugins/AudioInsights/Source/AudioInsights/Private/Providers/AudioEventLogTraceProvider.cpp
@@ -58,17 +58,10 @@ namespace UE::Audio::Insights
 				Builder.RouteEvent(RouteId_SoundStop, "Audio", "SoundStop");
 			}
 
-			virtual bool OnEvent(uint16 RouteId, EStyle Style, const FOnEventContext& Context) override
+			virtual bool OnHandleEvent(uint16 RouteId, EStyle Style, const FOnEventContext& Context) override
 			{
 				LLM_SCOPE_BYNAME(TEXT("Insights/FAudioEventLogTraceAnalyzer"));
 
-#if WITH_EDITOR
-				if (ShouldIgnoreNewEvents())
-				{
-					return true;
-				}
-#endif // WITH_EDITOR
-
 				FAudioEventLogMessages& Messages = GetProvider<FAudioEventLogTraceProvider>().TraceMessages;
 				switch (RouteId)
 				{
@@ -126,13 +119,6 @@ namespace UE::Audio::Insights
 
 	bool FAudioEventLogTraceProvider::ProcessMessages()
 	{
-#if WITH_EDITOR
-		if (GetTimeControlMethod() != ETimeControlMethod::LatestMessageTime)
-		{
-			return true;
-		}
-#endif // WITH_EDITOR
-
 		auto AddEntryFunc = [this](const FAudioEventLogMessage& Msg)
 		{
 			TSharedPtr<FAudioEventLogDashboardEntry>* ToReturn = nullptr;
diff --git a/Engine/Plugins/AudioInsights/Source/AudioInsights/Private/Providers/SoundTraceProvider.cpp b/Engine/Plugins/AudioInsights/Source/AudioInsights/Private/Providers/SoundTraceProvider.cpp
--- a/Engine/Plugins/AudioInsights/Source/AudioInsights/Private/Providers/SoundTraceProvider.cpp
+++ b/Engine/Plugins/AudioInsights/Source/AudioInsights/Private/Providers/SoundTraceProvider.cpp
@@ -219,17 +219,10 @@ namespace UE::Audio::Insights
 				Builder.RouteEvent(RouteId_RelativeRenderCost,		"Audio", "SoundRelativeRenderCost");
 			}
 
-			virtual bool OnEvent(uint16 RouteId, EStyle Style, const FOnEventContext& Context) override
+			virtual bool OnHandleEvent(uint16 RouteId, EStyle Style, const FOnEventContext& Context) override
 			{
 				LLM_SCOPE_BYNAME(TEXT("Insights/FSoundTraceAnalyzer"));
 
-#if WITH_EDITOR
-				if (ShouldIgnoreNewEvents())
-				{
-					return true;
-				}
-#endif // WITH_EDITOR
-
 				FSoundMessages& Messages = GetProvider<FSoundTraceProvider>().TraceMessages;
 
 				switch (RouteId)
@@ -421,11 +414,6 @@ namespace UE::Audio::Insights
 
 	bool FSoundTraceProvider::ProcessMessages()
 	{
-		if (GetTimeControlMethod() != ETimeControlMethod::LatestMessageTime)
-		{
-			return true;
-		}
-
 		// Helper lambdas
 		TSharedPtr<FSoundDashboardEntry> ActiveSoundEntryReturn = nullptr;
 		auto GetActiveSoundEntryFromActiveStartMessage = [this, &ActiveSoundEntryReturn](const FSoundStartMessage& Msg) -> TSharedPtr<FSoundDashboardEntry>*
diff --git a/Engine/Plugins/AudioInsights/Source/AudioInsights/Private/Providers/SubmixTraceProvider.cpp b/Engine/Plugins/AudioInsights/Source/AudioInsights/Private/Providers/SubmixTraceProvider.cpp
--- a/Engine/Plugins/AudioInsights/Source/AudioInsights/Private/Providers/SubmixTraceProvider.cpp
+++ b/Engine/Plugins/AudioInsights/Source/AudioInsights/Private/Providers/SubmixTraceProvider.cpp
@@ -86,17 +86,10 @@ namespace UE::Audio::Insights
 				Builder.RouteEvent(RouteId_Unloaded,       "Audio", "SubmixUnloaded");
 			}
 
-			virtual bool OnEvent(uint16 RouteId, EStyle Style, const FOnEventContext& Context) override
+			virtual bool OnHandleEvent(uint16 RouteId, EStyle Style, const FOnEventContext& Context) override
 			{
 				LLM_SCOPE_BYNAME(TEXT("Insights/FSubmixTraceAnalyzer"));
 
-#if WITH_EDITOR
-				if (ShouldIgnoreNewEvents())
-				{
-					return true;
-				}
-#endif // WITH_EDITOR
-
 				FSubmixMessages& Messages = GetProvider<FSubmixTraceProvider>().TraceMessages;
 
 				switch (RouteId)
@@ -174,11 +167,6 @@ namespace UE::Audio::Insights
 
 	bool FSubmixTraceProvider::ProcessMessages()
 	{
-		if (GetTimeControlMethod() != ETimeControlMethod::LatestMessageTime)
-		{
-			return true;
-		}
-
 		// Helper lambdas
 		auto CreateEntry = [this](const FSubmixMessageBase& Msg)
 		{
diff --git a/Engine/Plugins/AudioInsights/Source/AudioInsights/Private/Providers/VirtualLoopTraceProvider.cpp b/Engine/Plugins/AudioInsights/Source/AudioInsights/Private/Providers/VirtualLoopTraceProvider.cpp
--- a/Engine/Plugins/AudioInsights/Source/AudioInsights/Private/Providers/VirtualLoopTraceProvider.cpp
+++ b/Engine/Plugins/AudioInsights/Source/AudioInsights/Private/Providers/VirtualLoopTraceProvider.cpp
@@ -196,11 +196,6 @@ namespace UE::Audio::Insights
 
 	bool FVirtualLoopTraceProvider::ProcessMessages()
 	{
-		if (GetTimeControlMethod() != ETimeControlMethod::LatestMessageTime)
-		{
-			return true;
-		}
-
 		auto RemoveEntryFunc = [this](const FVirtualLoopStopOrRealizeMessage& Msg, TSharedPtr<FVirtualLoopDashboardEntry>* OutEntry)
 		{
 #if !WITH_EDITOR
@@ -311,17 +306,10 @@ namespace UE::Audio::Insights
 				Builder.RouteEvent(RouteId_Virtualize, "Audio", "VirtualLoopVirtualize");
 			}
 
-			virtual bool OnEvent(uint16 RouteId, EStyle Style, const FOnEventContext& Context) override
+			virtual bool OnHandleEvent(uint16 RouteId, EStyle Style, const FOnEventContext& Context) override
 			{
 				LLM_SCOPE_BYNAME(TEXT("Insights/FVirtualLoopTraceAnalyzer"));
 
-#if WITH_EDITOR
-				if (ShouldIgnoreNewEvents())
-				{
-					return true;
-				}
-#endif // WITH_EDITOR
-
 				FVirtualLoopMessages& Messages = GetProvider<FVirtualLoopTraceProvider>().TraceMessages;
 				switch (RouteId)
 				{
diff --git a/Engine/Plugins/AudioInsights/Source/AudioInsights/Public/AudioInsightsTraceProviderBase.h b/Engine/Plugins/AudioInsights/Source/AudioInsights/Public/AudioInsightsTraceProviderBase.h
--- a/Engine/Plugins/AudioInsights/Source/AudioInsights/Public/AudioInsightsTraceProviderBase.h
+++ b/Engine/Plugins/AudioInsights/Source/AudioInsights/Public/AudioInsightsTraceProviderBase.h
@@ -95,6 +95,9 @@ namespace UE::Audio::Insights
 			AUDIOINSIGHTS_API virtual bool OnEventSuccess(uint16 RouteId, EStyle Style, const FOnEventContext& Context);
 			AUDIOINSIGHTS_API virtual bool OnEventFailure(uint16 RouteId, EStyle Style, const FOnEventContext& Context);
 
+			// Override in derived to handle incoming trace messages
+			AUDIOINSIGHTS_API virtual bool OnHandleEvent(uint16 RouteId, EStyle Style, const FOnEventContext& Context) = 0;
+
 			AUDIOINSIGHTS_API bool ShouldIgnoreNewEvents() const;
 
 			template <typename TProviderType>
@@ -118,6 +121,8 @@ namespace UE::Audio::Insights
 #endif // WITH_EDITOR
 
 		private:
+			AUDIOINSIGHTS_API virtual bool OnEvent(uint16 RouteId, EStyle Style, const FOnEventContext& Context) override;
+
 			TSharedRef<FTraceProviderBase> Provider;
 		};
 
@@ -126,6 +131,8 @@ namespace UE::Audio::Insights
 			return LastMessageId;
 		}
 
+		AUDIOINSIGHTS_API bool ProcessMessagesIsPaused() const;
+
 		uint64 LastUpdateId = 0;
 		bool bForceUpdate = false;
 
@@ -151,7 +158,7 @@ namespace UE::Audio::Insights
 		{
 			TickerHandle = FTSTicker::GetCoreTicker().AddTicker(*InName.ToString(), 0.0f, [this](float DeltaTime)
 			{
-				if (LastUpdateId != GetLastMessageId())
+				if (LastUpdateId != GetLastMessageId() && !ProcessMessagesIsPaused())
 				{
 					ProcessMessages();
 				}
diff --git a/Engine/Plugins/Runtime/AudioModulation/Source/AudioModulationEditor/Private/Insights/Providers/ControlBusTraceProvider.cpp b/Engine/Plugins/Runtime/AudioModulation/Source/AudioModulationEditor/Private/Insights/Providers/ControlBusTraceProvider.cpp
--- a/Engine/Plugins/Runtime/AudioModulation/Source/AudioModulationEditor/Private/Insights/Providers/ControlBusTraceProvider.cpp
+++ b/Engine/Plugins/Runtime/AudioModulation/Source/AudioModulationEditor/Private/Insights/Providers/ControlBusTraceProvider.cpp
@@ -90,7 +90,7 @@ namespace AudioModulationEditor
 				Builder.RouteEvent(RouteId_Update, "Audio", "ControlBusUpdate");
 			}
 
-			virtual bool OnEvent(uint16 RouteId, UE::Trace::IAnalyzer::EStyle Style, const UE::Trace::IAnalyzer::FOnEventContext& Context) override
+			virtual bool OnHandleEvent(uint16 RouteId, UE::Trace::IAnalyzer::EStyle Style, const UE::Trace::IAnalyzer::FOnEventContext& Context) override
 			{
 				LLM_SCOPE_BYNAME(TEXT("Insights/FControlBusTraceAnalyzer"));
 
diff --git a/Engine/Plugins/Runtime/AudioModulation/Source/AudioModulationEditor/Private/Insights/Providers/ModulationMatrixTraceProvider.cpp b/Engine/Plugins/Runtime/AudioModulation/Source/AudioModulationEditor/Private/Insights/Providers/ModulationMatrixTraceProvider.cpp
--- a/Engine/Plugins/Runtime/AudioModulation/Source/AudioModulationEditor/Private/Insights/Providers/ModulationMatrixTraceProvider.cpp
+++ b/Engine/Plugins/Runtime/AudioModulation/Source/AudioModulationEditor/Private/Insights/Providers/ModulationMatrixTraceProvider.cpp
@@ -290,7 +290,7 @@ namespace AudioModulationEditor
 				Builder.RouteEvent(RouteId_Deactivate, "Audio", "ModulatingSourceDeactivate");
 			}
 
-			virtual bool OnEvent(uint16 RouteId, UE::Trace::IAnalyzer::EStyle Style, const UE::Trace::IAnalyzer::FOnEventContext& Context) override
+			virtual bool OnHandleEvent(uint16 RouteId, UE::Trace::IAnalyzer::EStyle Style, const UE::Trace::IAnalyzer::FOnEventContext& Context) override
 			{
 				LLM_SCOPE_BYNAME(TEXT("Insights/FModulationMatrixTraceAnalyzer"));
 
