[Audio Insights] Disable auto-scroll when the user scrolls over the event log - re-enable on esc/unselect a row

diff --git a/Engine/Plugins/AudioInsights/Source/AudioInsights/Private/AudioEventLogEditorCommands.cpp b/Engine/Plugins/AudioInsights/Source/AudioInsights/Private/AudioEventLogEditorCommands.cpp
--- a/Engine/Plugins/AudioInsights/Source/AudioInsights/Private/AudioEventLogEditorCommands.cpp
+++ b/Engine/Plugins/AudioInsights/Source/AudioInsights/Private/AudioEventLogEditorCommands.cpp
@@ -18,6 +18,7 @@ namespace UE::Audio::Insights
 	{
 		UI_COMMAND(Browse, "Browse To Asset", "Browses to the selected sound asset in the content browser.", EUserInterfaceActionType::Button, FInputChord(EKeys::B, EModifierKey::Control));
 		UI_COMMAND(Edit, "Edit", "Opens the selected sound for edit.", EUserInterfaceActionType::Button, FInputChord(EKeys::E, EModifierKey::Control));
+		UI_COMMAND(ResetInspectTimestamp, "Reset Inspect Timestamp", "Resumes gathering the latest audio data and auto-scrolling in the Event Log if previously paused.", EUserInterfaceActionType::Button, FInputChord(EKeys::Escape));
 	}
 
 } // namespace UE::Audio::Insights
diff --git a/Engine/Plugins/AudioInsights/Source/AudioInsights/Private/AudioEventLogEditorCommands.h b/Engine/Plugins/AudioInsights/Source/AudioInsights/Private/AudioEventLogEditorCommands.h
--- a/Engine/Plugins/AudioInsights/Source/AudioInsights/Private/AudioEventLogEditorCommands.h
+++ b/Engine/Plugins/AudioInsights/Source/AudioInsights/Private/AudioEventLogEditorCommands.h
@@ -17,9 +17,11 @@ namespace UE::Audio::Insights
 
 		TSharedPtr<const FUICommandInfo> GetBrowseCommand() const { return Browse; }
 		TSharedPtr<const FUICommandInfo> GetEditCommand() const { return Edit; }
+		TSharedPtr<const FUICommandInfo> GetResetInspectTimestampCommand() const { return ResetInspectTimestamp; }
 
 	private:
 		TSharedPtr<FUICommandInfo> Browse;
 		TSharedPtr<FUICommandInfo> Edit;
+		TSharedPtr<FUICommandInfo> ResetInspectTimestamp;
 	};
 } // namespace UE::Audio::Insights
diff --git a/Engine/Plugins/AudioInsights/Source/AudioInsights/Private/Views/AudioEventLogDashboardViewFactory.cpp b/Engine/Plugins/AudioInsights/Source/AudioInsights/Private/Views/AudioEventLogDashboardViewFactory.cpp
--- a/Engine/Plugins/AudioInsights/Source/AudioInsights/Private/Views/AudioEventLogDashboardViewFactory.cpp
+++ b/Engine/Plugins/AudioInsights/Source/AudioInsights/Private/Views/AudioEventLogDashboardViewFactory.cpp
@@ -163,9 +163,23 @@ namespace UE::Audio::Insights
 #if WITH_EDITOR
 		CommandList->MapAction(Commands.GetBrowseCommand(), FExecuteAction::CreateLambda([this]() { BrowseToAsset(); }));
 		CommandList->MapAction(Commands.GetEditCommand(), FExecuteAction::CreateLambda([this]() { OpenAsset(); }));
+		CommandList->MapAction(Commands.GetResetInspectTimestampCommand(), FExecuteAction::CreateLambda([this]() { ResetInspectTimestamp(); }));
 #endif // WITH_EDITOR
 	}
 
+	void FAudioEventLogDashboardViewFactory::ResetInspectTimestamp()
+	{
+		bAutoScroll = true;
+		if (FilteredEntriesListView.IsValid())
+		{
+			FilteredEntriesListView->ClearSelection();
+			FilteredEntriesListView->ClearHighlightedItems();
+		}
+
+		FAudioInsightsModule& AudioInsightsModule = FAudioInsightsModule::GetChecked();
+		AudioInsightsModule.GetTimingViewExtender().ResetTimeControlMethod();
+	}
+
 	FName FAudioEventLogDashboardViewFactory::GetName() const
 	{
 		return "EventLog";
@@ -186,6 +200,11 @@ namespace UE::Audio::Insights
 		return FSlateStyle::Get().CreateIcon("AudioInsights.Icon.Log");
 	}
 
+	FReply FAudioEventLogDashboardViewFactory::OnDataRowKeyInput(const FGeometry& Geometry, const FKeyEvent& KeyEvent) const
+	{
+		return (CommandList && CommandList->ProcessCommandBindings(KeyEvent)) ? FReply::Handled() : FReply::Unhandled();
+	}
+
 	TSharedRef<SWidget> FAudioEventLogDashboardViewFactory::MakeWidget(TSharedRef<SDockTab> OwnerTab, const FSpawnTabArgs& SpawnTabArgs)
 	{
 		FAudioInsightsModule& AudioInsightsModule = FAudioInsightsModule::GetChecked();
@@ -225,18 +244,15 @@ namespace UE::Audio::Insights
 	{
 		FTraceObjectTableDashboardViewFactory::RefreshFilteredEntriesListView();
 
-		if (FilteredEntriesListView.IsValid())
+		if (FilteredEntriesListView.IsValid() && bAutoScroll)
 		{
-			if (FilteredEntriesListView->GetNumItemsSelected() == 0)
+			if (SortMode == EColumnSortMode::Ascending)
 			{
-				if (SortMode == EColumnSortMode::Ascending)
-				{
-					FilteredEntriesListView->ScrollToBottom();
-				}
-				else if (SortMode == EColumnSortMode::Descending)
-				{
-					FilteredEntriesListView->ScrollToTop();
-				}
+				FilteredEntriesListView->ScrollToBottom();
+			}
+			else if (SortMode == EColumnSortMode::Descending)
+			{
+				FilteredEntriesListView->ScrollToTop();
 			}
 		}
 	}
@@ -617,6 +633,8 @@ namespace UE::Audio::Insights
 
 			FAudioInsightsModule& AudioInsightsModule = FAudioInsightsModule::GetChecked();
 			AudioInsightsModule.GetTimingViewExtender().ResetTimeControlMethod();
+
+			bAutoScroll = true;
 			return;
 		}
 
@@ -642,6 +660,16 @@ namespace UE::Audio::Insights
 		{
 			FAudioInsightsModule& AudioInsightsModule = FAudioInsightsModule::GetChecked();
 			AudioInsightsModule.GetTimingViewExtender().SetTimeMarker(SelectedEventLogEntry.Timestamp);
+
+			bAutoScroll = false;
+		}
+	}
+
+	void FAudioEventLogDashboardViewFactory::OnListViewScrolled(double InScrollOffset)
+	{
+		if (FilteredEntriesListView.IsValid() && FilteredEntriesListView->IsUserScrolling())
+		{
+			bAutoScroll = false;
 		}
 	}
 
diff --git a/Engine/Plugins/AudioInsights/Source/AudioInsights/Private/Views/TableDashboardViewFactory.cpp b/Engine/Plugins/AudioInsights/Source/AudioInsights/Private/Views/TableDashboardViewFactory.cpp
--- a/Engine/Plugins/AudioInsights/Source/AudioInsights/Private/Views/TableDashboardViewFactory.cpp
+++ b/Engine/Plugins/AudioInsights/Source/AudioInsights/Private/Views/TableDashboardViewFactory.cpp
@@ -226,6 +226,7 @@ namespace UE::Audio::Insights
 			.ListItemsSource(&DataViewEntries)
 			.OnContextMenuOpening(this, &FTraceTableDashboardViewFactory::OnConstructContextMenu)
 			.OnSelectionChanged(this, &FTraceTableDashboardViewFactory::OnSelectionChanged)
+			.OnListViewScrolled(this, &FTraceTableDashboardViewFactory::OnListViewScrolled)
 			.ClearSelectionOnClick(ClearSelectionOnClick())
 			.OnGenerateRow_Lambda([this](TSharedPtr<IDashboardDataViewEntry> Item, const TSharedRef<STableViewBase>& OwnerTable)
 			{
diff --git a/Engine/Plugins/AudioInsights/Source/AudioInsights/Public/Views/AudioEventLogDashboardViewFactory.h b/Engine/Plugins/AudioInsights/Source/AudioInsights/Public/Views/AudioEventLogDashboardViewFactory.h
--- a/Engine/Plugins/AudioInsights/Source/AudioInsights/Public/Views/AudioEventLogDashboardViewFactory.h
+++ b/Engine/Plugins/AudioInsights/Source/AudioInsights/Public/Views/AudioEventLogDashboardViewFactory.h
@@ -59,6 +59,7 @@ namespace UE::Audio::Insights
 		UE_API virtual EDefaultDashboardTabStack GetDefaultTabStack() const override;
 		UE_API virtual FSlateIcon GetIcon() const override;
 		
+		UE_API virtual FReply OnDataRowKeyInput(const FGeometry& Geometry, const FKeyEvent& KeyEvent) const override;
 		UE_API virtual TSharedRef<SWidget> MakeWidget(TSharedRef<SDockTab> OwnerTab, const FSpawnTabArgs& SpawnTabArgs) override;
 		UE_API virtual void RefreshFilteredEntriesListView() override;
 
@@ -84,6 +85,7 @@ namespace UE::Audio::Insights
 		UE_API virtual void OnSelectionChanged(TSharedPtr<IDashboardDataViewEntry> SelectedItem, ESelectInfo::Type SelectInfo) override;
 		virtual bool ClearSelectionOnClick() const override { return true; }
 		virtual bool EnableHorizontalScrollBox() const override { return false; }
+		UE_API virtual void OnListViewScrolled(double InScrollOffset) override;
 		UE_API virtual const FTableRowStyle* GetRowStyle() const override;
 
 #if WITH_EDITOR
@@ -111,6 +113,9 @@ namespace UE::Audio::Insights
 
 		void OnTimingViewTimeMarkerChanged(double InTimeMarker);
 
+		void BindCommands();
+		void ResetInspectTimestamp();
+
 		TSharedPtr<SWidget> DashboardWidget;
 		TSharedPtr<SAudioFilterBar<FEventLogFilterID>> FilterBar;
 
@@ -122,9 +127,11 @@ namespace UE::Audio::Insights
 		TSharedPtr<FUICommandList> CommandList;
 		TSharedPtr<FVisibleColumnsSettingsMenu<FAudioEventLogVisibleColumns>> VisibleColumnsSettingsMenu;
 
-		void BindCommands();
+		
 		const TMap<FString, TSet<FString>> InitEventTypeFilters;
 		TMap<FString, FAudioEventLogCustomEvents> CachedCustomEventSettings;
+
+		bool bAutoScroll = true;
 	};
 } // namespace UE::Audio::Insights
 
diff --git a/Engine/Plugins/AudioInsights/Source/AudioInsights/Public/Views/TableDashboardViewFactory.h b/Engine/Plugins/AudioInsights/Source/AudioInsights/Public/Views/TableDashboardViewFactory.h
--- a/Engine/Plugins/AudioInsights/Source/AudioInsights/Public/Views/TableDashboardViewFactory.h
+++ b/Engine/Plugins/AudioInsights/Source/AudioInsights/Public/Views/TableDashboardViewFactory.h
@@ -160,6 +160,7 @@ namespace UE::Audio::Insights
 		virtual bool ClearSelectionOnClick() const { return false; }
 
 		virtual bool EnableHorizontalScrollBox() const { return true; }
+		virtual void OnListViewScrolled(double InScrollOffset) { }
 
 		virtual const FTableRowStyle* GetRowStyle() const { return nullptr; }
 		UE_API virtual FSlateColor GetRowColor(const TSharedPtr<IDashboardDataViewEntry>& InRowDataPtr);
