[Audio Insights] Audio Event Log dashboard: added Clear Log button (in toolbar and when right clicking the list) + fixed FAudioInsightsComponent::GetIsLiveSession so it works correctly when doing direct trace in standalone.

diff --git a/Engine/Plugins/AudioInsights/Source/AudioInsights/Private/AudioInsightsComponent.cpp b/Engine/Plugins/AudioInsights/Source/AudioInsights/Private/AudioInsightsComponent.cpp
--- a/Engine/Plugins/AudioInsights/Source/AudioInsights/Private/AudioInsightsComponent.cpp
+++ b/Engine/Plugins/AudioInsights/Source/AudioInsights/Private/AudioInsightsComponent.cpp
@@ -137,17 +137,18 @@ namespace UE::Audio::Insights
 
 	bool FAudioInsightsComponent::GetIsLiveSession() const
 	{
-		IUnrealInsightsModule& UnrealInsightsModule = FModuleManager::LoadModuleChecked<IUnrealInsightsModule>("TraceInsights");
+		const IUnrealInsightsModule& UnrealInsightsModule = FModuleManager::LoadModuleChecked<IUnrealInsightsModule>("TraceInsights");
 
 		TSharedPtr<const TraceServices::IAnalysisSession> Session = UnrealInsightsModule.GetAnalysisSession();
 		if (Session.IsValid())
 		{
 			TraceServices::FAnalysisSessionReadScope SessionReadScope(*Session.Get());
 
-			const UE::Trace::FStoreClient* StoreClient = UnrealInsightsModule.GetStoreClient();
-			const UE::Trace::FStoreClient::FSessionInfo* SessionInfo = StoreClient ? StoreClient->GetSessionInfoByTraceId(Session->GetTraceId()) : nullptr;
-
-			return !Session->IsAnalysisComplete() && StoreClient != nullptr && SessionInfo != nullptr;
+			const TraceServices::IDiagnosticsProvider* DiagnosticsProvider = TraceServices::ReadDiagnosticsProvider(*Session.Get());
+			if (DiagnosticsProvider && DiagnosticsProvider->IsSessionInfoAvailable())
+			{
+				return !Session->IsAnalysisComplete() && DiagnosticsProvider->IsSessionInfoAvailable();
+			}
 		}
 
 		return false;
diff --git a/Engine/Plugins/AudioInsights/Source/AudioInsights/Private/Providers/AudioEventLogTraceProvider.cpp b/Engine/Plugins/AudioInsights/Source/AudioInsights/Private/Providers/AudioEventLogTraceProvider.cpp
--- a/Engine/Plugins/AudioInsights/Source/AudioInsights/Private/Providers/AudioEventLogTraceProvider.cpp
+++ b/Engine/Plugins/AudioInsights/Source/AudioInsights/Private/Providers/AudioEventLogTraceProvider.cpp
@@ -239,7 +239,7 @@ namespace UE::Audio::Insights
 #endif // WITH_EDITOR
 			
 			// Cache the actor information associated with this PlayOrder ID so we can look it up later
-			PlayOrderToActorInfoMap.Add(Entry->PlayOrder, FActorInformation{ Msg.ActorLabel, FName(Msg.ActorIconName) });
+			PlayOrderToActorInfoMap.Add(Entry->PlayOrder, FActorInformation{ Msg.ActorLabel, FName(Msg.ActorIconName), Msg.DeviceId });
 		}
 
 		return Entry;
@@ -252,6 +252,26 @@ namespace UE::Audio::Insights
 		bCacheChunksUpdated = true;
 	}
 #endif // WITH_EDITOR
+
+	void FAudioEventLogTraceProvider::ClearActiveAudioDeviceEntries()
+	{
+#if WITH_EDITOR
+		const IAudioInsightsModule& AudioInsightsModule = IAudioInsightsModule::GetEditorChecked();
+		const ::Audio::FDeviceId AudioDeviceId = AudioInsightsModule.GetDeviceId();
+#else
+		const ::Audio::FDeviceId AudioDeviceId = GetAudioDeviceIdFromDeviceDataMap();
+#endif // WITH_EDITOR
+
+		for (auto It = PlayOrderToActorInfoMap.CreateIterator(); It; ++It)
+		{
+			if (It.Value().AudioDeviceId == AudioDeviceId)
+			{
+				It.RemoveCurrent();
+			}
+		}
+
+		Super::ClearActiveAudioDeviceEntries();
+	}
 } // namespace UE::Audio::Insights
 
 #undef LOCTEXT_NAMESPACE
\ No newline at end of file
diff --git a/Engine/Plugins/AudioInsights/Source/AudioInsights/Private/Providers/AudioEventLogTraceProvider.h b/Engine/Plugins/AudioInsights/Source/AudioInsights/Private/Providers/AudioEventLogTraceProvider.h
--- a/Engine/Plugins/AudioInsights/Source/AudioInsights/Private/Providers/AudioEventLogTraceProvider.h
+++ b/Engine/Plugins/AudioInsights/Source/AudioInsights/Private/Providers/AudioEventLogTraceProvider.h
@@ -18,6 +18,8 @@ namespace UE::Audio::Insights
 
 		static FName GetName_Static();
 
+		virtual void ClearActiveAudioDeviceEntries() override;
+
 	private:
 		virtual bool ProcessMessages() override;
 		virtual bool ProcessManuallyUpdatedEntries() override;
@@ -34,6 +36,7 @@ namespace UE::Audio::Insights
 		{
 			const FString ActorLabel;
 			const FName ActorIconName;
+			const ::Audio::FDeviceId AudioDeviceId = 0;
 		};
 
 		TMap<uint32, FActorInformation> PlayOrderToActorInfoMap;
diff --git a/Engine/Plugins/AudioInsights/Source/AudioInsights/Private/Views/AudioEventLogDashboardViewFactory.cpp b/Engine/Plugins/AudioInsights/Source/AudioInsights/Private/Views/AudioEventLogDashboardViewFactory.cpp
--- a/Engine/Plugins/AudioInsights/Source/AudioInsights/Private/Views/AudioEventLogDashboardViewFactory.cpp
+++ b/Engine/Plugins/AudioInsights/Source/AudioInsights/Private/Views/AudioEventLogDashboardViewFactory.cpp
@@ -8,6 +8,10 @@
 #include "Messages/AudioEventLogTraceMessages.h"
 #include "Providers/AudioEventLogTraceProvider.h"
 
+#if !WITH_EDITOR
+#include "AudioInsightsComponent.h"
+#endif // !WITH_EDITOR
+
 #define LOCTEXT_NAMESPACE "AudioInsightsEventLog"
 
 namespace UE::Audio::Insights
@@ -21,6 +25,17 @@ namespace UE::Audio::Insights
 			return static_cast<const FAudioEventLogDashboardEntry&>(InData);
 		};
 
+		bool CanClearLog()
+		{
+#if WITH_EDITOR
+			return true;
+#else
+			const TSharedPtr<const FAudioInsightsComponent> AudioInsightsComponent = FAudioInsightsModule::GetChecked().GetAudioInsightsComponent();
+
+			return AudioInsightsComponent.IsValid() && AudioInsightsComponent->GetIsLiveSession();
+#endif // WITH_EDITOR
+		}
+
 		const FString CategorySoundActivity = LOCTEXT("EventLogCategory_SoundActivity", "Sound Activity").ToString();
 		const FString CategoryVirtualization = LOCTEXT("EventLogCategory_Virtualization", "Virtualization").ToString();
 		const FString CategoryPlayRequests = LOCTEXT("EventLogCategory_PlayRequests", "Play Requests").ToString();
@@ -37,6 +52,9 @@ namespace UE::Audio::Insights
 		const FName AssetNameColumnID = "Asset";
 		const FName ActorColumnID = "Actor";
 		const FName CategoryColumnID = "Category";
+
+		const FText ClearLogText = LOCTEXT("AudioEventLog_ClearLog_Text", "Clear Log");
+		const FText ClearLogTooltipText = LOCTEXT("AudioEventLog_ClearLog_Tooltip", "Clears the event log for the currently selected world.");
 	} // namespace AudioEventLogPrivate
 
 	const TMap<FString, TSet<FString>> FAudioEventLogDashboardViewFactory::GetInitEventTypeFilters()
@@ -123,7 +141,7 @@ namespace UE::Audio::Insights
 
 		FTraceModule& AudioInsightsTraceModule = static_cast<FTraceModule&>(FAudioInsightsModule::GetChecked().GetTraceModule());
 
-		const TSharedPtr<FAudioEventLogTraceProvider> AudioEventLogTraceProvider = MakeShared<FAudioEventLogTraceProvider>();
+		AudioEventLogTraceProvider = MakeShared<FAudioEventLogTraceProvider>();
 
 		AudioInsightsTraceModule.AddTraceProvider(AudioEventLogTraceProvider);
 
@@ -560,6 +578,52 @@ namespace UE::Audio::Insights
 			]
 #endif // WITH_EDITOR
 
+			// Clear log Button
+			+ SHorizontalBox::Slot()
+			.HAlign(HAlign_Left)
+			.VAlign(VAlign_Center)
+			.AutoWidth()
+			.Padding(15.0f, 0.0f, 6.0f, 0.0f)
+			[
+				SNew(SButton)
+				.ButtonStyle(&FAppStyle::Get().GetWidgetStyle<FButtonStyle>("SimpleButton"))
+				.ToolTipText(AudioEventLogPrivate::ClearLogTooltipText)
+				.IsEnabled_Lambda([]()
+				{
+					return AudioEventLogPrivate::CanClearLog();
+				})
+				.OnClicked_Lambda([this]()
+				{
+					ClearActiveAudioDeviceEntries();
+
+					return FReply::Handled();
+				})
+				[
+					SNew(SBox)
+					[
+						SNew(SHorizontalBox)
+						// Clear log icon
+						+ SHorizontalBox::Slot()
+						.VAlign(VAlign_Center)
+						.AutoWidth()
+						[
+							SNew(SImage)
+							.Image(FSlateStyle::Get().GetBrush("AudioInsights.Icon.ClearLog"))
+						]
+						// Clear log text
+						+ SHorizontalBox::Slot()
+						.VAlign(VAlign_Center)
+						.AutoWidth()
+						.Padding(6.0f, 1.0f, 0.0f, 0.0f)
+						[
+							SNew(STextBlock)
+							.Justification(ETextJustify::Center)
+							.Text(AudioEventLogPrivate::ClearLogText)
+						]
+					]
+				]
+			]
+
 			// Settings button
 			+ SHorizontalBox::Slot()
 			.AutoWidth()
@@ -598,12 +662,12 @@ namespace UE::Audio::Insights
 		return LOCTEXT("AudioEventLog_Settings_TooltipText", "Event log settings");
 	}
 
-#if WITH_EDITOR
 	TSharedPtr<SWidget> FAudioEventLogDashboardViewFactory::OnConstructContextMenu()
 	{
 		constexpr bool bShouldCloseWindowAfterMenuSelection = true;
 		FMenuBuilder MenuBuilder(bShouldCloseWindowAfterMenuSelection, CommandList);
 
+#if WITH_EDITOR
 		const FAudioEventLogEditorCommands& Commands = FAudioEventLogEditorCommands::Get();
 
 		MenuBuilder.BeginSection("EventLogDashboardActions", LOCTEXT("EventLogDashboard_Actions_HeaderText", "Editor Options"));
@@ -614,9 +678,38 @@ namespace UE::Audio::Insights
 
 		MenuBuilder.EndSection();
 
+		MenuBuilder.AddSeparator();
+#endif // WITH_EDITOR
+
+		MenuBuilder.AddMenuEntry(
+			AudioEventLogPrivate::ClearLogText,
+			AudioEventLogPrivate::ClearLogTooltipText,
+			FSlateStyle::Get().CreateIcon("AudioInsights.Icon.ClearLog"),
+			FUIAction(
+				FExecuteAction::CreateSP(this, &FAudioEventLogDashboardViewFactory::ClearActiveAudioDeviceEntries),
+				FCanExecuteAction::CreateLambda([]()
+				{
+					return AudioEventLogPrivate::CanClearLog();
+				})
+			)
+		);
+
 		return MenuBuilder.MakeWidget();
 	}
-#endif // WITH_EDITOR
+
+	void FAudioEventLogDashboardViewFactory::ClearActiveAudioDeviceEntries()
+	{
+		if (AudioEventLogTraceProvider.IsValid())
+		{
+			AudioEventLogTraceProvider->ClearActiveAudioDeviceEntries();
+			DataViewEntries.Reset();
+
+			if (FilteredEntriesListView.IsValid())
+			{
+				RefreshFilteredEntriesListView();
+			}
+		}
+	}
 
 	void FAudioEventLogDashboardViewFactory::OnSelectionChanged(TSharedPtr<IDashboardDataViewEntry> SelectedItem, ESelectInfo::Type SelectInfo)
 	{
diff --git a/Engine/Plugins/AudioInsights/Source/AudioInsights/Public/AudioInsightsStyle.h b/Engine/Plugins/AudioInsights/Source/AudioInsights/Public/AudioInsightsStyle.h
--- a/Engine/Plugins/AudioInsights/Source/AudioInsights/Public/AudioInsightsStyle.h
+++ b/Engine/Plugins/AudioInsights/Source/AudioInsights/Public/AudioInsightsStyle.h
@@ -172,8 +172,9 @@ namespace UE::Audio::Insights
 			const FVector2D Icon64(64.0f, 64.0f);
 
 			Set("AudioInsights.Icon", new IMAGE_BRUSH_SVG(TEXT("Icons/audio_insights_icon"), Icon16));
-			Set("AudioInsights.Icon.Dashboard", new IMAGE_BRUSH_SVG(TEXT("Icons/audio_dashboard"), Icon16));
+			Set("AudioInsights.Icon.ClearLog", new IMAGE_BRUSH_SVG(TEXT("Icons/audio_clear_log"), Icon20));
 			Set("AudioInsights.Icon.ContentBrowser", new IMAGE_BRUSH_SVG(TEXT("Icons/content_browser"), Icon20));
+			Set("AudioInsights.Icon.Dashboard", new IMAGE_BRUSH_SVG(TEXT("Icons/audio_dashboard"), Icon16));
 			Set("AudioInsights.Icon.Event", new IMAGE_BRUSH_SVG(TEXT("Icons/audio_event"), Icon16));
 			Set("AudioInsights.Icon.Log", new IMAGE_BRUSH_SVG(TEXT("Icons/audio_log"), Icon16));
 			Set("AudioInsights.Icon.Open", new IMAGE_BRUSH_SVG(TEXT("Icons/open"), Icon20));
diff --git a/Engine/Plugins/AudioInsights/Source/AudioInsights/Public/AudioInsightsTraceProviderBase.h b/Engine/Plugins/AudioInsights/Source/AudioInsights/Public/AudioInsightsTraceProviderBase.h
--- a/Engine/Plugins/AudioInsights/Source/AudioInsights/Public/AudioInsightsTraceProviderBase.h
+++ b/Engine/Plugins/AudioInsights/Source/AudioInsights/Public/AudioInsightsTraceProviderBase.h
@@ -204,6 +204,18 @@ namespace UE::Audio::Insights
 			return DeviceDataMap.Find(AudioDeviceId);
 		}
 
+		virtual void ClearActiveAudioDeviceEntries()
+		{
+#if WITH_EDITOR
+			const IAudioInsightsModule& AudioInsightsModule = IAudioInsightsModule::GetEditorChecked();
+			const ::Audio::FDeviceId AudioDeviceId = AudioInsightsModule.GetDeviceId();
+#else
+			const ::Audio::FDeviceId AudioDeviceId = GetAudioDeviceIdFromDeviceDataMap();
+#endif // WITH_EDITOR
+
+			DeviceDataMap.Remove(AudioDeviceId);
+		}
+
 	protected:
 		using Super = TDeviceDataMapTraceProvider<EntryKey, EntryType>;
 
@@ -291,6 +303,25 @@ namespace UE::Audio::Insights
 			InEntryMutator(Entry);
 		}
 
+#if !WITH_EDITOR
+	const uint32 GetAudioDeviceIdFromDeviceDataMap() const
+	{
+		// In Audio Insights for Unreal Insights we don't have access to the Audio Device Manager, 
+		// we need to extract the AudioDeviceId from the already recorded trace info stored in DeviceDataMap
+		uint32 AudioDeviceId = 1; // 1 is the default audio device id
+
+		for (const auto& [DeviceId, DeviceData] : DeviceDataMap)
+		{
+			if (DeviceId > AudioDeviceId)
+			{
+				AudioDeviceId = DeviceId;
+			}
+		}
+
+		return AudioDeviceId;
+	}
+#endif // !WITH_EDITOR
+
 	private:
 #if WITH_EDITOR
 		void RegisterDelegates()
@@ -331,23 +362,6 @@ namespace UE::Audio::Insights
 				}
 			}
 		}
-#else
-		const uint32 GetAudioDeviceIdFromDeviceDataMap() const
-		{
-			// In Audio Insights for Unreal Insights we don't have access to the Audio Device Manager, 
-			// we need to extract the AudioDeviceId from the already recorded trace info stored in DeviceDataMap
-			uint32 AudioDeviceId = 1; // 1 is the default audio device id
-			
-			for (const auto& [DeviceId, DeviceData] : DeviceDataMap)
-			{
-				if (DeviceId > AudioDeviceId)
-				{
-					AudioDeviceId = DeviceId;
-				}
-			}
-
-			return AudioDeviceId;
-		}
 #endif // WITH_EDITOR
 
 #if WITH_EDITOR
diff --git a/Engine/Plugins/AudioInsights/Source/AudioInsights/Public/Views/AudioEventLogDashboardViewFactory.h b/Engine/Plugins/AudioInsights/Source/AudioInsights/Public/Views/AudioEventLogDashboardViewFactory.h
--- a/Engine/Plugins/AudioInsights/Source/AudioInsights/Public/Views/AudioEventLogDashboardViewFactory.h
+++ b/Engine/Plugins/AudioInsights/Source/AudioInsights/Public/Views/AudioEventLogDashboardViewFactory.h
@@ -12,6 +12,8 @@
 
 namespace UE::Audio::Insights
 {
+	class FAudioEventLogTraceProvider;
+
 	using FEventLogFilterID = int32;
 
 	class FAudioEventLogFilter : public FGenericFilter<FEventLogFilterID>
@@ -78,9 +80,7 @@ namespace UE::Audio::Insights
 		UE_API virtual TSharedRef<SWidget> OnGetSettingsMenuContent() override;
 		UE_API virtual FText OnGetSettingsMenuToolTip() override;
 
-#if WITH_EDITOR
 		UE_API virtual TSharedPtr<SWidget> OnConstructContextMenu() override;
-#endif //WITH_EDITOR
 
 		UE_API virtual void OnSelectionChanged(TSharedPtr<IDashboardDataViewEntry> SelectedItem, ESelectInfo::Type SelectInfo) override;
 		virtual bool ClearSelectionOnClick() const override { return true; }
@@ -116,6 +116,10 @@ namespace UE::Audio::Insights
 		void BindCommands();
 		void ResetInspectTimestamp();
 
+		void ClearActiveAudioDeviceEntries();
+
+		TSharedPtr<FAudioEventLogTraceProvider> AudioEventLogTraceProvider;
+
 		TSharedPtr<SWidget> DashboardWidget;
 		TSharedPtr<SAudioFilterBar<FEventLogFilterID>> FilterBar;
 
