[Audio Insights] AudioInsightsEditorDashboardFactory.cpp: Adjusted LOCTEXTs so they are unique to this cpp file and don't collide with the standalone one + update snapshot tooltip text

diff --git a/Engine/Plugins/AudioInsights/Source/AudioInsightsEditor/Private/AudioInsightsEditorDashboardFactory.cpp b/Engine/Plugins/AudioInsights/Source/AudioInsightsEditor/Private/AudioInsightsEditorDashboardFactory.cpp
--- a/Engine/Plugins/AudioInsights/Source/AudioInsightsEditor/Private/AudioInsightsEditorDashboardFactory.cpp
+++ b/Engine/Plugins/AudioInsights/Source/AudioInsightsEditor/Private/AudioInsightsEditorDashboardFactory.cpp
@@ -31,46 +31,46 @@ namespace UE::Audio::Insights
 {
 	namespace EditorDashboardFactoryPrivate
 	{
-		static const FText ToolNameText = LOCTEXT("AudioDashboard_ToolName_Text", "Audio Insights");
+		static const FText ToolNameText = LOCTEXT("AudioDashboardEditor_ToolName_Text", "Audio Insights");
 
-		static const FText PreviewDeviceText = LOCTEXT("AudioDashboard_PreviewDevice_Text", "[Preview Audio]");
+		static const FText PreviewDeviceText = LOCTEXT("AudioDashboardEditor_PreviewDevice_Text", "[Preview Audio]");
 
-		static const FText MainMenuFileText = LOCTEXT("AudioDashboard_MainMenu_File_Text", "File");
-		static const FText MainMenuViewText = LOCTEXT("AudioDashboard_MainMenu_View_Text", "View");
+		static const FText MainMenuFileText = LOCTEXT("AudioDashboardEditor_MainMenu_File_Text", "File");
+		static const FText MainMenuViewText = LOCTEXT("AudioDashboardEditor_MainMenu_View_Text", "View");
 
-		static const FText FileCloseText = LOCTEXT("AudioDashboard_File_Close_Text", "Close");
-		static const FText FileCloseTooltipText = LOCTEXT("AudioDashboard_File_Close_TooltipText", "Closes the Audio Insights dashboard.");
+		static const FText FileCloseText = LOCTEXT("AudioDashboardEditor_File_Close_Text", "Close");
+		static const FText FileCloseTooltipText = LOCTEXT("AudioDashboardEditor_File_Close_TooltipText", "Closes the Audio Insights dashboard.");
 
-		static const FText ViewResetLayoutText = LOCTEXT("AudioDashboard_View_ResetLayout_Text", "Reset Layout");
+		static const FText ViewResetLayoutText = LOCTEXT("AudioDashboardEditor_View_ResetLayout_Text", "Reset Layout");
 
-		static const FText OnlyTraceAudioChannelsText = LOCTEXT("AudioDashboard_OnlyTraceAudioChannels_Text", "Only trace audio channels:");
-		static const FText OnlyTraceAudioChannelsTooltipText = LOCTEXT("AudioDashboard_OnlyTraceAudioChannels_TooltipText", "Disable all non-command line trace channels apart from Audio, Audio Mixer and Bookmark when starting a live monitoring / live trace. This will reduce the file size of snapshots / trace sessions while using Audio Insights.");
+		static const FText OnlyTraceAudioChannelsText = LOCTEXT("AudioDashboardEditor_OnlyTraceAudioChannels_Text", "Only trace audio channels:");
+		static const FText OnlyTraceAudioChannelsTooltipText = LOCTEXT("AudioDashboardEditor_OnlyTraceAudioChannels_TooltipText", "Disable all non-command line trace channels apart from Audio, Audio Mixer and Bookmark when starting a live monitoring / live trace. This will reduce the file size of snapshots / trace sessions while using Audio Insights.");
 
-		static const FText StartDirectTraceText = LOCTEXT("AudioDashboard_StartDirectTrace_Text", "Start Monitoring");
-		static const FText StartDirectTraceTooltipText = LOCTEXT("AudioDashboard_StartDirectTrace_TooltipText", "Starts monitoring for audio trace data without saving the data to a .utrace file.");
+		static const FText StartDirectTraceText = LOCTEXT("AudioDashboardEditor_StartDirectTrace_Text", "Start Monitoring");
+		static const FText StartDirectTraceTooltipText = LOCTEXT("AudioDashboardEditor_StartDirectTrace_TooltipText", "Starts monitoring for audio trace data without saving the data to a .utrace file.");
 
-		static const FText StopDirectTraceText = LOCTEXT("AudioDashboard_StopDirectTrace_Text", "Stop Monitoring");
-		static const FText StopDirectTraceTooltipText = LOCTEXT("AudioDashboard_StopDirectTrace_TooltipText", "Stops the active monitoring session.");
+		static const FText StopDirectTraceText = LOCTEXT("AudioDashboardEditor_StopDirectTrace_Text", "Stop Monitoring");
+		static const FText StopDirectTraceTooltipText = LOCTEXT("AudioDashboardEditor_StopDirectTrace_TooltipText", "Stops the active monitoring session.");
 
-		static const FText StopLiveTraceText = LOCTEXT("AudioDashboard_StopLiveTrace_Text", "Stop Trace");
-		static const FText StopLiveTraceTooltipText = LOCTEXT("AudioDashboard_StopLiveTrace_TooltipText", "Stops the currently active trace recording.");
+		static const FText StopLiveTraceText = LOCTEXT("AudioDashboardEditor_StopLiveTrace_Text", "Stop Trace");
+		static const FText StopLiveTraceTooltipText = LOCTEXT("AudioDashboardEditor_StopLiveTrace_TooltipText", "Stops the currently active trace recording.");
 
-		static const FText SaveSnapshotText = LOCTEXT("AudioDashboard_SnapshotButton_Text", "Save Snapshot");
-		static const FText SaveSnapshotTooltipText = LOCTEXT("AudioDashboard_SnapshotButton_TooltipText", "Saves a snapshot of the most recent trace data to a .utrace file. Snapshots are limited to 32MB and can be accessed inside Unreal Insights.");
+		static const FText SaveSnapshotText = LOCTEXT("AudioDashboardEditor_SnapshotButton_Text", "Save Snapshot");
+		static const FText SaveSnapshotTooltipText = LOCTEXT("AudioDashboardEditor_SnapshotButton_TooltipText", "Saves a snapshot of the most recent trace data to a .utrace file. By default snapshots are limited to 32MB and can be accessed in Unreal Insights. To specify a custom snapshot size (in MB) launch the editor executable with the -tracetailmb=<size> option.");
 
-		static const FText SaveSnapshotNotificationText = LOCTEXT("AudioDashboard_SnapshotButton_NotificationText", "Audio Insights Snapshot saved.");
-		static const FText SaveSnapshotNotificationSubText = LOCTEXT("AudioDashboard_SnapshotButton_NotificationSubText", "A snapshot .utrace with the most recent events has been saved to your trace server.");
+		static const FText SaveSnapshotNotificationText = LOCTEXT("AudioDashboardEditor_SnapshotButton_NotificationText", "Audio Insights Snapshot saved.");
+		static const FText SaveSnapshotNotificationSubText = LOCTEXT("AudioDashboardEditor_SnapshotButton_NotificationSubText", "A snapshot .utrace with the most recent events has been saved to your trace server.");
 
-		static const FText SaveBookmarkText = LOCTEXT("AudioDashboard_BookmarkButton_Text", "Bookmark");
-		static const FText SaveBookmarkTooltipText = LOCTEXT("AudioDashboard_BookmarkButton_TooltipText", "Ctrl+M - Creates an 'Audio Insights' bookmark in the .utrace file generated by Save Snapshot or Trace at the timestamp the button was pressed.");
+		static const FText SaveBookmarkText = LOCTEXT("AudioDashboardEditor_BookmarkButton_Text", "Bookmark");
+		static const FText SaveBookmarkTooltipText = LOCTEXT("AudioDashboardEditor_BookmarkButton_TooltipText", "Ctrl+M - Creates an 'Audio Insights' bookmark in the .utrace file generated by Save Snapshot or Trace at the timestamp the button was pressed.");
 
-		static const FText BookmarkLabelText = LOCTEXT("AudioDashboard_BookmarkLabelTitle_Text", "Audio Insights Bookmark");
-		static const FText SaveBookmarkNotificationText = LOCTEXT("AudioDashboard_BookmarkButton_NotificationText", "Audio Insights Bookmark created.");
-		static const FText SaveBookmarkNotificationDirectTraceSubText = LOCTEXT("AudioDashboard_BookmarkButton_NotificationSubText", "bookmark has been added to the Audio Insights Snapshot.");
-		static const FText SaveBookmarkNotificationLiveTraceSubText = LOCTEXT("AudioDashboard_BookmarkButton_LiveTrace_NotificationSubText", "bookmark has been added to the live .utrace file on your trace server.");
+		static const FText BookmarkLabelText = LOCTEXT("AudioDashboardEditor_BookmarkLabelTitle_Text", "Audio Insights Bookmark");
+		static const FText SaveBookmarkNotificationText = LOCTEXT("AudioDashboardEditor_BookmarkButton_NotificationText", "Audio Insights Bookmark created.");
+		static const FText SaveBookmarkNotificationDirectTraceSubText = LOCTEXT("AudioDashboardEditor_BookmarkButton_NotificationSubText", "bookmark has been added to the Audio Insights Snapshot.");
+		static const FText SaveBookmarkNotificationLiveTraceSubText = LOCTEXT("AudioDashboardEditor_BookmarkButton_LiveTrace_NotificationSubText", "bookmark has been added to the live .utrace file on your trace server.");
 
-		static const FText WorldFilterText = LOCTEXT("AudioDashboard_WorldFilter_Text", "World Filter:");
-		static const FText WorldFilterTooltipText = LOCTEXT("AudioDashboard_WorldFilter_TooltipText", "Select world(s) to monitor (worlds may share audio output).");
+		static const FText WorldFilterText = LOCTEXT("AudioDashboardEditor_WorldFilter_Text", "World Filter:");
+		static const FText WorldFilterTooltipText = LOCTEXT("AudioDashboardEditor_WorldFilter_TooltipText", "Select world(s) to monitor (worlds may share audio output).");
 
 		FText GetDebugNameFromDeviceId(::Audio::FDeviceId InDeviceId)
 		{
