[Audio Insights] FAudioEventLogDashboardViewFactory: added ability to filter by PlayOrder

diff --git a/Engine/Plugins/AudioInsights/Source/AudioInsights/Private/Views/AudioEventLogDashboardViewFactory.cpp b/Engine/Plugins/AudioInsights/Source/AudioInsights/Private/Views/AudioEventLogDashboardViewFactory.cpp
--- a/Engine/Plugins/AudioInsights/Source/AudioInsights/Private/Views/AudioEventLogDashboardViewFactory.cpp
+++ b/Engine/Plugins/AudioInsights/Source/AudioInsights/Private/Views/AudioEventLogDashboardViewFactory.cpp
@@ -290,7 +290,11 @@ namespace UE::Audio::Insights
 
 			if (bPassesEventTypeFilter)
 			{
-				if (EventLogEntry.GetDisplayName().ToString().Contains(FilterString) || EventLogEntry.ActorLabel.Contains(FilterString))
+				if (FilterString.IsNumeric() && LexToString(EventLogEntry.PlayOrder).Equals(FilterString))
+				{
+					return true;
+				}
+				else if (EventLogEntry.GetDisplayName().ToString().Contains(FilterString) || EventLogEntry.ActorLabel.Contains(FilterString))
 				{
 					return true;
 				}
