[Audio Insights] Improvements to thread safety of the Audio Insights Cache Manager - move the addition/removal of messages to the cache into the same thread as reads from the cache

diff --git a/Engine/Plugins/AudioInsights/Source/AudioInsights/Private/Cache/AudioCachedMessageChunk.cpp b/Engine/Plugins/AudioInsights/Source/AudioInsights/Private/Cache/AudioCachedMessageChunk.cpp
--- a/Engine/Plugins/AudioInsights/Source/AudioInsights/Private/Cache/AudioCachedMessageChunk.cpp
+++ b/Engine/Plugins/AudioInsights/Source/AudioInsights/Private/Cache/AudioCachedMessageChunk.cpp
@@ -23,21 +23,13 @@ namespace UE::Audio::Insights
 
 		CurrentChunkSize += Message->GetSizeOf();
 
-		{
-			FScopeLock Lock(&MessagesCriticalSection);
-
-			TSharedPtr<TraceServices::TPagedArray<TSharedPtr<IAudioCachedMessage>>>& CachedMessages = Messages.FindOrAdd(Message->GetMessageName(), MakeShared<TraceServices::TPagedArray<TSharedPtr<IAudioCachedMessage>>>(Allocator, 4096));
-			CachedMessages->EmplaceBack(Message);
-		}
+		TSharedPtr<TraceServices::TPagedArray<TSharedPtr<IAudioCachedMessage>>>& CachedMessages = Messages.FindOrAdd(Message->GetMessageName(), MakeShared<TraceServices::TPagedArray<TSharedPtr<IAudioCachedMessage>>>(Allocator, 4096));
+		CachedMessages->EmplaceBack(Message);
     }
 
 	void FAudioCachedMessageChunk::ClearChunk()
 	{
-		{
-			FScopeLock Lock(&MessagesCriticalSection);
-			Messages.Reset();
-		}
-		
+		Messages.Reset();
 		CurrentChunkSize = 0u;
 
 		RangeStartTimestamp = INVALID_TIMESTAMP;
diff --git a/Engine/Plugins/AudioInsights/Source/AudioInsights/Private/Cache/AudioInsightsCacheManager.cpp b/Engine/Plugins/AudioInsights/Source/AudioInsights/Private/Cache/AudioInsightsCacheManager.cpp
--- a/Engine/Plugins/AudioInsights/Source/AudioInsights/Private/Cache/AudioInsightsCacheManager.cpp
+++ b/Engine/Plugins/AudioInsights/Source/AudioInsights/Private/Cache/AudioInsightsCacheManager.cpp
@@ -18,36 +18,22 @@ namespace UE::Audio::Insights
 		check(CacheManagerPrivate::NumChunks > 0u);
 
 		CreateCache();
-	}
-
-	void FAudioInsightsCacheManager::AddMessageToCache(const TSharedPtr<IAudioCachedMessage> Message)
-	{
-		check(WriteIndex < Cache.Capacity());
 
-		FAudioCachedMessageChunk& Chunk = *Cache[WriteIndex].Get();
-
-		if (Chunk.IsChunkFull())
+		TickerHandle = FTSTicker::GetCoreTicker().AddTicker(TEXT("FAudioInsightsCacheManager::ProcessPendingMessages"), 0.0f, [this](float DeltaTime)
 		{
-			// Move to next chunk
-			WriteIndex = Cache.GetNextIndex(WriteIndex);
-
-			FAudioCachedMessageChunk& NextChunk = *Cache[WriteIndex].Get();
-
-			// If this new chunk already has data, clear it before we begin writing
-			if (NextChunk.HasAnyData())
-			{
-				StartIndex = Cache.GetNextIndex(WriteIndex);
-				NextChunk.ClearChunk();
-			}
+			ProcessPendingMessages(DeltaTime);
+			return true;
+		});
+	}
 
-			NextChunk.AddMessageToChunk(Message);
+	FAudioInsightsCacheManager::~FAudioInsightsCacheManager()
+	{
+		FTSTicker::GetCoreTicker().RemoveTicker(TickerHandle);
+	}
 
-			OnChunkOverwritten.Broadcast(GetCacheStartTimeStamp());
-		}
-		else
-		{
-			Chunk.AddMessageToChunk(Message);
-		}
+	void FAudioInsightsCacheManager::AddMessageToCache(TSharedPtr<IAudioCachedMessage> Message)
+	{
+		NewMessageQueue.Enqueue(MoveTemp(Message));
 	}
 
 	void FAudioInsightsCacheManager::ClearCache()
@@ -107,4 +93,38 @@ namespace UE::Audio::Insights
 		WriteIndex = 0u;
 		StartIndex = 0u;
 	}
+
+	void FAudioInsightsCacheManager::ProcessPendingMessages(float DeltaTime)
+	{
+		check(WriteIndex < Cache.Capacity());
+
+		FAudioCachedMessageChunk* Chunk = Cache[WriteIndex].Get();
+
+		const TArray<TSharedPtr<IAudioCachedMessage>> NewMessages = NewMessageQueue.DequeueAll();
+		for (const TSharedPtr<IAudioCachedMessage>& Message : NewMessages)
+		{
+			if (Chunk->IsChunkFull())
+			{
+				// Move to next chunk
+				WriteIndex = Cache.GetNextIndex(WriteIndex);
+
+				Chunk = Cache[WriteIndex].Get();
+
+				// If this new chunk already has data, clear it before we begin writing
+				if (Chunk->HasAnyData())
+				{
+					StartIndex = Cache.GetNextIndex(WriteIndex);
+					Chunk->ClearChunk();
+				}
+
+				Chunk->AddMessageToChunk(Message);
+
+				OnChunkOverwritten.Broadcast(GetCacheStartTimeStamp());
+			}
+			else
+			{
+				Chunk->AddMessageToChunk(Message);
+			}
+		}
+	}
 } // namespace UE::Audio::Insights
diff --git a/Engine/Plugins/AudioInsights/Source/AudioInsights/Public/Cache/AudioCachedMessageChunk.h b/Engine/Plugins/AudioInsights/Source/AudioInsights/Public/Cache/AudioCachedMessageChunk.h
--- a/Engine/Plugins/AudioInsights/Source/AudioInsights/Public/Cache/AudioCachedMessageChunk.h
+++ b/Engine/Plugins/AudioInsights/Source/AudioInsights/Public/Cache/AudioCachedMessageChunk.h
@@ -78,8 +78,6 @@ namespace UE::Audio::Insights
 
 		uint32 CurrentChunkSize = 0u;
 		const uint32 MaxChunkSize = 0u;
-
-		FCriticalSection MessagesCriticalSection;
 	};
 } // namespace UE::Audio::Insights
 
diff --git a/Engine/Plugins/AudioInsights/Source/AudioInsights/Public/Cache/AudioInsightsCacheManager.h b/Engine/Plugins/AudioInsights/Source/AudioInsights/Public/Cache/AudioInsightsCacheManager.h
--- a/Engine/Plugins/AudioInsights/Source/AudioInsights/Public/Cache/AudioInsightsCacheManager.h
+++ b/Engine/Plugins/AudioInsights/Source/AudioInsights/Public/Cache/AudioInsightsCacheManager.h
@@ -4,6 +4,8 @@
 #include "Cache/AudioCachedMessageChunk.h"
 #include "Cache/IAudioCachedMessage.h"
 #include "Containers/CircularBuffer.h"
+#include "Containers/Ticker.h"
+#include "Messages/AnalyzerMessageQueue.h"
 #include "TraceServices/Containers/SlabAllocator.h"
 
 #define UE_API AUDIOINSIGHTS_API
@@ -14,9 +16,9 @@ namespace UE::Audio::Insights
 	{
 	public:
 		UE_API FAudioInsightsCacheManager();
-		virtual ~FAudioInsightsCacheManager() = default;
+		UE_API virtual ~FAudioInsightsCacheManager();
 
-		UE_API void AddMessageToCache(const TSharedPtr<IAudioCachedMessage> Message);
+		UE_API void AddMessageToCache(TSharedPtr<IAudioCachedMessage> Message);
 		UE_API void ClearCache();
 
 		UE_API float GetCacheDuration() const;
@@ -175,11 +177,17 @@ namespace UE::Audio::Insights
 	private:
 		void CreateCache();
 
+		void ProcessPendingMessages(float DeltaTime);
+
 		TraceServices::FSlabAllocator Allocator;
 		TCircularBuffer<TUniquePtr<FAudioCachedMessageChunk>> Cache;
 
+		TCacheMessageQueue<TSharedPtr<IAudioCachedMessage>> NewMessageQueue { 5.0 };
+
 		uint32 WriteIndex = 0u;
 		uint32 StartIndex = 0u;
+
+		FTSTicker::FDelegateHandle TickerHandle;
 	};
 } // namespace UE::Audio::Insights
 
diff --git a/Engine/Plugins/AudioInsights/Source/AudioInsights/Public/Messages/AnalyzerMessageQueue.h b/Engine/Plugins/AudioInsights/Source/AudioInsights/Public/Messages/AnalyzerMessageQueue.h
--- a/Engine/Plugins/AudioInsights/Source/AudioInsights/Public/Messages/AnalyzerMessageQueue.h
+++ b/Engine/Plugins/AudioInsights/Source/AudioInsights/Public/Messages/AnalyzerMessageQueue.h
@@ -5,20 +5,24 @@
 #include "Math/UnrealMathUtility.h"
 #include "Misc/ScopeLock.h"
 
-
 namespace UE::Audio::Insights
 {
 	template <typename T>
-	struct TAnalyzerMessageQueue
+	struct TMessageQueueBase
 	{
 		static constexpr double MaxHistoryLimitSec = 5.0;
 
-		explicit TAnalyzerMessageQueue(double InHistoryLimitSec = MaxHistoryLimitSec)
-			: HistoryLimitSec(FMath::Clamp(InHistoryLimitSec, UE_DOUBLE_KINDA_SMALL_NUMBER, MaxHistoryLimitSec))
+		explicit TMessageQueueBase(double InHistoryLimitSec)
+			: HistoryLimitSec(FMath::Clamp(InHistoryLimitSec, UE_DOUBLE_KINDA_SMALL_NUMBER, 5.0))
 		{
 		}
 
-	private:
+		virtual ~TMessageQueueBase() = default;
+
+	protected:
+		virtual double GetTimestampFromData() const = 0;
+		virtual double GetTimestampFromMessage(const T& Message) const = 0;
+
 		double HistoryLimitSec = 0.0;
 		TMpscQueue<T> Data;
 
@@ -39,19 +43,19 @@ namespace UE::Audio::Insights
 
 				if (!Data.IsEmpty())
 				{
-					double FirstTimeStamp = Data.Peek()->Timestamp;
+					double FirstTimeStamp = GetTimestampFromData();
 
 					T Message;
 					while (Data.Dequeue(Message))
 					{
-						if (HistoryLimitSec > Message.Timestamp - FirstTimeStamp)
+						if (HistoryLimitSec > GetTimestampFromMessage(Message) - FirstTimeStamp)
 						{
 							Output.Add(MoveTemp(Message));
 						}
 
 						if (!Data.IsEmpty())
 						{
-							FirstTimeStamp = Data.Peek()->Timestamp;
+							FirstTimeStamp = GetTimestampFromData();
 						}
 					}
 				}
@@ -72,14 +76,14 @@ namespace UE::Audio::Insights
 
 				if (!Data.IsEmpty())
 				{
-					double FirstTimeStamp = Data.Peek()->Timestamp;
-					while (InMessage.Timestamp - FirstTimeStamp > MaxHistoryLimitSec)
+					double FirstTimeStamp = GetTimestampFromData();
+					while (GetTimestampFromMessage(InMessage) - FirstTimeStamp > MaxHistoryLimitSec)
 					{
 						if (!Data.Dequeue().IsSet() || Data.IsEmpty())
 						{
 							break;
 						}
-						FirstTimeStamp = Data.Peek()->Timestamp;
+						FirstTimeStamp = GetTimestampFromData();
 					}
 				}
 			}
@@ -93,4 +97,66 @@ namespace UE::Audio::Insights
 			Data.Empty();
 		}
 	};
+
+	template<typename T>
+	concept TIsValidMessage = requires (T Message)
+	{
+		{ Message.Timestamp } -> std::same_as<double&>;
+	};
+
+	template <TIsValidMessage T>
+	struct TAnalyzerMessageQueue : public TMessageQueueBase<T>
+	{
+		explicit TAnalyzerMessageQueue(double InHistoryLimitSec = TMessageQueueBase<T>::MaxHistoryLimitSec)
+			: TMessageQueueBase<T>(InHistoryLimitSec)
+		{
+		}
+
+	protected:
+		virtual double GetTimestampFromData() const override
+		{
+			if (!TMessageQueueBase<T>::Data.IsEmpty())
+			{
+				return TMessageQueueBase<T>::Data.Peek()->Timestamp;
+			}
+
+			return 0.0;
+		}
+
+		virtual double GetTimestampFromMessage(const T& Message) const
+		{
+			return Message.Timestamp;
+		}
+	};
+
+	template<typename T>
+	concept TIsValidMessagePointer = requires (T Message)
+	{
+		{ Message->Timestamp } -> std::same_as<double&>;
+	};
+
+	template <TIsValidMessagePointer T>
+	struct TCacheMessageQueue : public TMessageQueueBase<T>
+	{
+		explicit TCacheMessageQueue(double InHistoryLimitSec = TMessageQueueBase<T>::MaxHistoryLimitSec)
+			: TMessageQueueBase<T>(InHistoryLimitSec)
+		{
+		}
+
+	protected:
+		virtual double GetTimestampFromData() const override
+		{
+			if (!TMessageQueueBase<T>::Data.IsEmpty())
+			{
+				return (*TMessageQueueBase<T>::Data.Peek())->Timestamp;
+			}
+
+			return 0.0;
+		}
+
+		virtual double GetTimestampFromMessage(const T& Message) const
+		{
+			return Message->Timestamp;
+		}
+	};
 } // namespace UE::Audio::Insights
