[Audio Insights] Add options for how pausing the cache should work part 01 - update API for cache manager to help us detect when we are in the oldest chunk

diff --git a/Engine/Plugins/AudioInsights/Source/AudioInsights/Private/Cache/AudioInsightsCacheManager.cpp b/Engine/Plugins/AudioInsights/Source/AudioInsights/Private/Cache/AudioInsightsCacheManager.cpp
--- a/Engine/Plugins/AudioInsights/Source/AudioInsights/Private/Cache/AudioInsightsCacheManager.cpp
+++ b/Engine/Plugins/AudioInsights/Source/AudioInsights/Private/Cache/AudioInsightsCacheManager.cpp
@@ -83,6 +83,38 @@ namespace UE::Audio::Insights
 		return CacheManagerPrivate::NumChunks;
 	}
 
+	uint32 FAudioInsightsCacheManager::GetNumUsedChunks() const
+	{
+		uint32 NumUsedChunks = 0u;
+		for (uint32 Index = 0u; Index < Cache.Capacity(); ++Index)
+		{
+			if (Cache[Index]->HasAnyData())
+			{
+				++NumUsedChunks;
+			}
+		}
+
+		return NumUsedChunks;
+	}
+
+	uint32 FAudioInsightsCacheManager::GetNumChunksFromStartForTimestamp(const double Timestamp) const
+	{
+		uint32 ReadIndex = StartIndex;
+		for (uint32 Index = 0u; Index < Cache.Capacity(); ++Index)
+		{
+			if (Cache[ReadIndex]->TimestampIsInChunkRange(Timestamp))
+			{
+				return Index;
+			}
+			else
+			{
+				ReadIndex = Cache.GetNextIndex(ReadIndex);
+			}
+		}
+
+		return INDEX_NONE;
+	}
+
 	void FAudioInsightsCacheManager::CreateCache()
 	{
 		for (uint32 Index = 0u; Index < Cache.Capacity(); ++Index)
diff --git a/Engine/Plugins/AudioInsights/Source/AudioInsights/Public/Cache/AudioInsightsCacheManager.h b/Engine/Plugins/AudioInsights/Source/AudioInsights/Public/Cache/AudioInsightsCacheManager.h
--- a/Engine/Plugins/AudioInsights/Source/AudioInsights/Public/Cache/AudioInsightsCacheManager.h
+++ b/Engine/Plugins/AudioInsights/Source/AudioInsights/Public/Cache/AudioInsightsCacheManager.h
@@ -25,10 +25,21 @@ namespace UE::Audio::Insights
 		UE_API double GetCacheStartTimeStamp() const;
 		UE_API double GetCacheEndTimeStamp() const;
 
+		// Returns the total used memory size of the cache in bytes
 		UE_API uint32 GetUsedCacheSize() const;
+
+		// Returns the maximum memory size of the cache in bytes
 		UE_API uint32 GetMaxCacheSize() const;
+
+		// Returns the maximum number of chunks inside the cache
 		UE_API uint32 GetNumChunks() const;
 
+		// Returns the number of chunks with data present inside the cache
+		UE_API uint32 GetNumUsedChunks() const;
+
+		// Returns the number of chunks from the start of the cache for a specific timestamp, where lower indexed chunks contain older messages
+		UE_API uint32 GetNumChunksFromStartForTimestamp(const double Timestamp) const;
+
 		// Find the closest message to a timestamp in the cache for a specific message type
 		// Will look backwards from the chosen timestamp until a message is found
 		// @param ID: When overloaded, the search algorithm will only return messages that have a matching ID.
