[Audio Insights] Add options for how pausing the cache should work part 02 - Set up new commands and icon for pausing/resuming caching behaviour

diff --git a/Engine/Plugins/AudioInsights/Source/AudioInsights/Private/AudioEventLogEditorCommands.cpp b/Engine/Plugins/AudioInsights/Source/AudioInsights/Private/AudioEventLogEditorCommands.cpp
--- a/Engine/Plugins/AudioInsights/Source/AudioInsights/Private/AudioEventLogEditorCommands.cpp
+++ b/Engine/Plugins/AudioInsights/Source/AudioInsights/Private/AudioEventLogEditorCommands.cpp
@@ -19,6 +19,10 @@ namespace UE::Audio::Insights
 		UI_COMMAND(Browse, "Browse To Asset", "Browses to the selected sound asset in the content browser.", EUserInterfaceActionType::Button, FInputChord(EKeys::B, EModifierKey::Control));
 		UI_COMMAND(Edit, "Edit", "Opens the selected sound for edit.", EUserInterfaceActionType::Button, FInputChord(EKeys::E, EModifierKey::Control));
 		UI_COMMAND(ResetInspectTimestamp, "Reset Inspect Timestamp", "Resumes gathering the latest audio data and auto-scrolling in the Event Log if previously paused.", EUserInterfaceActionType::Button, FInputChord(EKeys::Escape));
+
+		UI_COMMAND(AutoStopCachingWhenLastInCache, "When last in cache", "Automatically stops caching new data when inspecting an event that is about to be deleted from the cache.", EUserInterfaceActionType::RadioButton, FInputChord(EKeys::F1, EModifierKey::Control));
+		UI_COMMAND(AutoStopCachingOnInspect, "On inspect", "Automatically stops caching new data when inspecting an event.", EUserInterfaceActionType::RadioButton, FInputChord(EKeys::F2, EModifierKey::Control));
+		UI_COMMAND(AutoStopCachingDisabled, "Never", "Disable automatically stopping the cache. Inspected events may be deleted from the cache during inspection.", EUserInterfaceActionType::RadioButton, FInputChord(EKeys::F3, EModifierKey::Control));
 	}
 
 } // namespace UE::Audio::Insights
diff --git a/Engine/Plugins/AudioInsights/Source/AudioInsights/Private/AudioEventLogEditorCommands.h b/Engine/Plugins/AudioInsights/Source/AudioInsights/Private/AudioEventLogEditorCommands.h
--- a/Engine/Plugins/AudioInsights/Source/AudioInsights/Private/AudioEventLogEditorCommands.h
+++ b/Engine/Plugins/AudioInsights/Source/AudioInsights/Private/AudioEventLogEditorCommands.h
@@ -19,9 +19,17 @@ namespace UE::Audio::Insights
 		TSharedPtr<const FUICommandInfo> GetEditCommand() const { return Edit; }
 		TSharedPtr<const FUICommandInfo> GetResetInspectTimestampCommand() const { return ResetInspectTimestamp; }
 
+		TSharedPtr<const FUICommandInfo> GetAutoStopCachingWhenLastInCacheCommand() const { return AutoStopCachingWhenLastInCache; }
+		TSharedPtr<const FUICommandInfo> GetAutoStopCachingOnInspectCommand() const { return AutoStopCachingOnInspect; }
+		TSharedPtr<const FUICommandInfo> GetAutoStopCachingDisabledCommand() const { return AutoStopCachingDisabled; }
+
 	private:
 		TSharedPtr<FUICommandInfo> Browse;
 		TSharedPtr<FUICommandInfo> Edit;
 		TSharedPtr<FUICommandInfo> ResetInspectTimestamp;
+
+		TSharedPtr<FUICommandInfo> AutoStopCachingWhenLastInCache;
+		TSharedPtr<FUICommandInfo> AutoStopCachingOnInspect;
+		TSharedPtr<FUICommandInfo> AutoStopCachingDisabled;
 	};
 } // namespace UE::Audio::Insights
diff --git a/Engine/Plugins/AudioInsights/Source/AudioInsights/Public/AudioInsightsStyle.h b/Engine/Plugins/AudioInsights/Source/AudioInsights/Public/AudioInsightsStyle.h
--- a/Engine/Plugins/AudioInsights/Source/AudioInsights/Public/AudioInsightsStyle.h
+++ b/Engine/Plugins/AudioInsights/Source/AudioInsights/Public/AudioInsightsStyle.h
@@ -178,6 +178,7 @@ namespace UE::Audio::Insights
 			Set("AudioInsights.Icon.Event", new IMAGE_BRUSH_SVG(TEXT("Icons/audio_event"), Icon16));
 			Set("AudioInsights.Icon.Log", new IMAGE_BRUSH_SVG(TEXT("Icons/audio_log"), Icon16));
 			Set("AudioInsights.Icon.Open", new IMAGE_BRUSH_SVG(TEXT("Icons/open"), Icon20));
+			Set("AudioInsights.Icon.Pause", new IMAGE_BRUSH_SVG(TEXT("Icons/pause"), Icon20));
 			Set("AudioInsights.Icon.Record.Active", new IMAGE_BRUSH_SVG(TEXT("Icons/record_active"), Icon20));
 			Set("AudioInsights.Icon.Record.Inactive", new IMAGE_BRUSH_SVG(TEXT("Icons/record_inactive"), Icon20));
 			Set("AudioInsights.Icon.Settings", new IMAGE_BRUSH_SVG(TEXT("Icons/settings"), Icon20));
