[Audio Insights] Add options for how pausing the cache should work part 03 - Caching and processing can now be controlled in different ways inside the event log

diff --git a/Engine/Plugins/AudioInsights/Source/AudioInsights/Private/AudioInsightsTimingViewExtender.cpp b/Engine/Plugins/AudioInsights/Source/AudioInsights/Private/AudioInsightsTimingViewExtender.cpp
--- a/Engine/Plugins/AudioInsights/Source/AudioInsights/Private/AudioInsightsTimingViewExtender.cpp
+++ b/Engine/Plugins/AudioInsights/Source/AudioInsights/Private/AudioInsightsTimingViewExtender.cpp
@@ -33,9 +33,9 @@ namespace UE::Audio::Insights
 		}
 	}
 
-	void FAudioInsightsTimingViewExtender::ResetTimeControlMethod()
+	void FAudioInsightsTimingViewExtender::ResetMessageProcessType()
 	{
-		TimeControlMethod = ETimeControlMethod::LatestMessageTime;
+		CacheAndProcessMessageStatus = ECacheAndProcess::Latest;
 		OnTimeControlMethodReset.Broadcast();
 	}
 
@@ -46,8 +46,6 @@ namespace UE::Audio::Insights
 
 	void FAudioInsightsTimingViewExtender::SetTimeMarker(const double Timestamp)
 	{
-		TimeControlMethod = ETimeControlMethod::Manual;
-
 		// If you set the time marker on the TimingView, it will Broadcast OnTimingViewTimeMarkerChanged. Otherwise we gotta do it ourselves.
 		if (TimingView != nullptr)
 		{
diff --git a/Engine/Plugins/AudioInsights/Source/AudioInsights/Private/AudioInsightsTimingViewExtender.h b/Engine/Plugins/AudioInsights/Source/AudioInsights/Private/AudioInsightsTimingViewExtender.h
--- a/Engine/Plugins/AudioInsights/Source/AudioInsights/Private/AudioInsightsTimingViewExtender.h
+++ b/Engine/Plugins/AudioInsights/Source/AudioInsights/Private/AudioInsightsTimingViewExtender.h
@@ -14,10 +14,14 @@ class IAnalysisSession;
 
 namespace UE::Audio::Insights
 {
-	enum class ETimeControlMethod
+	// How the cache manage handles new incoming messages
+	// Cache - will save the incoming message inside the cache to be retrievable later
+	// Process - sends the message to the Audio Insights providers to be processed and displayed in the UI
+	enum class ECacheAndProcess
 	{
-		LatestMessageTime = 0,
-		Manual
+		Latest = 0,				// Cache and process all new messages
+		CacheLatestNoProcess,	// Cache new messages but do not process them
+		None					// Do not cache or process any new messages
 	};
 
 	class  FAudioInsightsTimingViewExtender : public UE::Insights::Timing::ITimingViewExtender
@@ -28,8 +32,11 @@ namespace UE::Audio::Insights
 		virtual void OnEndSession(UE::Insights::Timing::ITimingViewSession& InSession) override;
 		virtual void Tick(UE::Insights::Timing::ITimingViewSession& InSession, const TraceServices::IAnalysisSession& InAnalysisSession) override;
 
-		ETimeControlMethod GetTimeControlMethod() const { return TimeControlMethod; }
-		void ResetTimeControlMethod();
+		ECacheAndProcess GetMessageCacheAndProcessingStatus() const { return CacheAndProcessMessageStatus; }
+
+		void StopProcessingNewMessages() { CacheAndProcessMessageStatus = ECacheAndProcess::CacheLatestNoProcess; }
+		void StopCachingAndProcessingNewMessages() { CacheAndProcessMessageStatus = ECacheAndProcess::None; }
+		void ResetMessageProcessType();
 
 		double GetCurrentDurationSeconds() const;
 
@@ -47,7 +54,7 @@ namespace UE::Audio::Insights
 		const TraceServices::IAnalysisSession* AnalysisSession = nullptr;
 		UE::Insights::Timing::ITimingViewSession* TimingView = nullptr;
 
-		ETimeControlMethod TimeControlMethod = ETimeControlMethod::LatestMessageTime;
+		ECacheAndProcess CacheAndProcessMessageStatus = ECacheAndProcess::Latest;
 		double TraceDurationSeconds = 0.0;
 	};
 } // namespace UE::Audio::Insights
diff --git a/Engine/Plugins/AudioInsights/Source/AudioInsights/Private/AudioInsightsTraceProviderBase.cpp b/Engine/Plugins/AudioInsights/Source/AudioInsights/Private/AudioInsightsTraceProviderBase.cpp
--- a/Engine/Plugins/AudioInsights/Source/AudioInsights/Private/AudioInsightsTraceProviderBase.cpp
+++ b/Engine/Plugins/AudioInsights/Source/AudioInsights/Private/AudioInsightsTraceProviderBase.cpp
@@ -73,27 +73,32 @@ namespace UE::Audio::Insights
 
 	bool FTraceProviderBase::FTraceAnalyzerBase::ShouldIgnoreNewEvents() const
 	{
-		return Provider->ProcessMessagesIsPaused();
+		return Provider->GetMessageCacheAndProcessingStatus() == ECacheAndProcess::None;
 	}
 
-	bool FTraceProviderBase::ProcessMessagesIsPaused() const
+	bool FTraceProviderBase::FTraceAnalyzerBase::ShouldProcessNewEvents() const
+	{
+		return Provider->GetMessageCacheAndProcessingStatus() == ECacheAndProcess::Latest;
+	}
+
+	bool FTraceProviderBase::IsMessageProcessingPaused() const
 	{
 #if WITH_EDITOR
-		return GetTimeControlMethod() != ETimeControlMethod::LatestMessageTime;
+		return GetMessageCacheAndProcessingStatus() != ECacheAndProcess::Latest;
 #else
 		return false;
 #endif // WITH_EDITOR
 	}
 
-	ETimeControlMethod FTraceProviderBase::GetTimeControlMethod() const
+	ECacheAndProcess FTraceProviderBase::GetMessageCacheAndProcessingStatus() const
 	{
 		const FAudioInsightsModule& AudioInsightsModule = FAudioInsightsModule::GetChecked();
-		return AudioInsightsModule.GetTimingViewExtender().GetTimeControlMethod();
+		return AudioInsightsModule.GetTimingViewExtender().GetMessageCacheAndProcessingStatus();
 	}
 
-	void FTraceProviderBase::ResetTimeControlMethod()
+	void FTraceProviderBase::ResetMessageProcessType()
 	{
 		FAudioInsightsModule& AudioInsightsModule = FAudioInsightsModule::GetChecked();
-		AudioInsightsModule.GetTimingViewExtender().ResetTimeControlMethod();
+		AudioInsightsModule.GetTimingViewExtender().ResetMessageProcessType();
 	}
 } // namespace UE::Audio::Insights
diff --git a/Engine/Plugins/AudioInsights/Source/AudioInsights/Private/Providers/AudioEventLogTraceProvider.cpp b/Engine/Plugins/AudioInsights/Source/AudioInsights/Private/Providers/AudioEventLogTraceProvider.cpp
--- a/Engine/Plugins/AudioInsights/Source/AudioInsights/Private/Providers/AudioEventLogTraceProvider.cpp
+++ b/Engine/Plugins/AudioInsights/Source/AudioInsights/Private/Providers/AudioEventLogTraceProvider.cpp
@@ -93,6 +93,12 @@ namespace UE::Audio::Insights
 				return OnEventSuccess(RouteId, Style, Context);
 			}
 
+		protected:
+			virtual bool ShouldProcessNewEvents() const override
+			{
+				return GetProvider<FAudioEventLogTraceProvider>().GetMessageCacheAndProcessingStatus() != ECacheAndProcess::None;
+			}
+
 		private:
 			enum : uint16
 			{
@@ -117,6 +123,11 @@ namespace UE::Audio::Insights
 		return "AudioEventLogProvider";
 	}
 
+	bool FAudioEventLogTraceProvider::IsMessageProcessingPaused() const
+	{
+		return GetMessageCacheAndProcessingStatus() == ECacheAndProcess::None;
+	}
+
 	bool FAudioEventLogTraceProvider::ProcessMessages()
 	{
 		auto AddEntryFunc = [this](const FAudioEventLogMessage& Msg)
diff --git a/Engine/Plugins/AudioInsights/Source/AudioInsights/Private/Providers/AudioEventLogTraceProvider.h b/Engine/Plugins/AudioInsights/Source/AudioInsights/Private/Providers/AudioEventLogTraceProvider.h
--- a/Engine/Plugins/AudioInsights/Source/AudioInsights/Private/Providers/AudioEventLogTraceProvider.h
+++ b/Engine/Plugins/AudioInsights/Source/AudioInsights/Private/Providers/AudioEventLogTraceProvider.h
@@ -20,6 +20,9 @@ namespace UE::Audio::Insights
 
 		virtual void ClearActiveAudioDeviceEntries() override;
 
+	protected:
+		virtual bool IsMessageProcessingPaused() const override;
+
 	private:
 		virtual bool ProcessMessages() override;
 		virtual bool ProcessManuallyUpdatedEntries() override;
diff --git a/Engine/Plugins/AudioInsights/Source/AudioInsights/Private/Views/AudioEventLogDashboardViewFactory.cpp b/Engine/Plugins/AudioInsights/Source/AudioInsights/Private/Views/AudioEventLogDashboardViewFactory.cpp
--- a/Engine/Plugins/AudioInsights/Source/AudioInsights/Private/Views/AudioEventLogDashboardViewFactory.cpp
+++ b/Engine/Plugins/AudioInsights/Source/AudioInsights/Private/Views/AudioEventLogDashboardViewFactory.cpp
@@ -159,6 +159,9 @@ namespace UE::Audio::Insights
 #if WITH_EDITOR
 		FAudioEventLogSettings::OnReadSettings.AddRaw(this, &FAudioEventLogDashboardViewFactory::OnReadEditorSettings);
 		FAudioEventLogSettings::OnWriteSettings.AddRaw(this, &FAudioEventLogDashboardViewFactory::OnWriteEditorSettings);
+
+		FAudioInsightsCacheManager& CacheManager = IAudioInsightsModule::GetChecked().GetCacheManager();
+		CacheManager.OnChunkOverwritten.AddRaw(this, &FAudioEventLogDashboardViewFactory::OnCacheChunkOverwritten);
 #endif // WITH_EDITOR
 	}
 
@@ -167,6 +170,12 @@ namespace UE::Audio::Insights
 #if WITH_EDITOR
 		FAudioEventLogSettings::OnReadSettings.RemoveAll(this);
 		FAudioEventLogSettings::OnWriteSettings.RemoveAll(this);
+
+		if (FModuleManager::Get().IsModuleLoaded("AudioInsights"))
+		{
+			FAudioInsightsCacheManager& CacheManager = IAudioInsightsModule::GetChecked().GetCacheManager();
+			CacheManager.OnChunkOverwritten.RemoveAll(this);
+		}
 #endif // WITH_EDITOR
 
 		FAudioEventLogEditorCommands::Unregister();
@@ -182,11 +191,30 @@ namespace UE::Audio::Insights
 		CommandList->MapAction(Commands.GetBrowseCommand(), FExecuteAction::CreateLambda([this]() { BrowseToAsset(); }));
 		CommandList->MapAction(Commands.GetEditCommand(), FExecuteAction::CreateLambda([this]() { OpenAsset(); }));
 		CommandList->MapAction(Commands.GetResetInspectTimestampCommand(), FExecuteAction::CreateLambda([this]() { ResetInspectTimestamp(); }));
+
+		CommandList->MapAction(
+			Commands.GetAutoStopCachingWhenLastInCacheCommand(),
+			FExecuteAction::CreateLambda([this]() { AutoStopCachingMode = EAutoStopCachingMode::WhenInLastChunk; }),
+			FCanExecuteAction(),
+			FIsActionChecked::CreateLambda([this]() { return AutoStopCachingMode == EAutoStopCachingMode::WhenInLastChunk; }));
+
+		CommandList->MapAction(
+			Commands.GetAutoStopCachingOnInspectCommand(),
+			FExecuteAction::CreateLambda([this]() { AutoStopCachingMode = EAutoStopCachingMode::OnInspect; }),
+			FCanExecuteAction(),
+			FIsActionChecked::CreateLambda([this]() { return AutoStopCachingMode == EAutoStopCachingMode::OnInspect; }));
+
+		CommandList->MapAction(
+			Commands.GetAutoStopCachingDisabledCommand(),
+			FExecuteAction::CreateLambda([this]() { AutoStopCachingMode = EAutoStopCachingMode::Never; }),
+			FCanExecuteAction(),
+			FIsActionChecked::CreateLambda([this]() { return AutoStopCachingMode == EAutoStopCachingMode::Never; }));
 #endif // WITH_EDITOR
 	}
 
 	void FAudioEventLogDashboardViewFactory::ResetInspectTimestamp()
 	{
+		FocusedItem.Reset();
 		bAutoScroll = true;
 		if (FilteredEntriesListView.IsValid())
 		{
@@ -195,7 +223,7 @@ namespace UE::Audio::Insights
 		}
 
 		FAudioInsightsModule& AudioInsightsModule = FAudioInsightsModule::GetChecked();
-		AudioInsightsModule.GetTimingViewExtender().ResetTimeControlMethod();
+		AudioInsightsModule.GetTimingViewExtender().ResetMessageProcessType();
 	}
 
 	FName FAudioEventLogDashboardViewFactory::GetName() const
@@ -235,6 +263,7 @@ namespace UE::Audio::Insights
 			if (FilteredEntriesListView.IsValid())
 			{
 				FilteredEntriesListView->SetSelectionMode(ESelectionMode::SingleToggle);
+				FilteredEntriesListView->OnMouseWheelDetected.AddSP(this, &FAudioEventLogDashboardViewFactory::OnUpdateAutoScroll);
 			}
 
 			if (HeaderRowWidget.IsValid())
@@ -262,7 +291,22 @@ namespace UE::Audio::Insights
 	{
 		FTraceObjectTableDashboardViewFactory::RefreshFilteredEntriesListView();
 
-		if (FilteredEntriesListView.IsValid() && bAutoScroll)
+		if (!FilteredEntriesListView.IsValid())
+		{
+			return;
+		}
+
+		if (FilteredEntriesListView->IsUserScrolling())
+		{
+			return;
+		}
+
+		const TSharedPtr<IDashboardDataViewEntry> FocusedItemShared = FocusedItem.Pin();
+		if (FocusedItemShared.IsValid())
+		{
+			FilteredEntriesListView->RequestScrollIntoView(FocusedItemShared);
+		}
+		else if (bAutoScroll)
 		{
 			if (SortMode == EColumnSortMode::Ascending)
 			{
@@ -526,8 +570,36 @@ namespace UE::Audio::Insights
 			SAssignNew(SettingsAreaWidget, SHorizontalBox)
 			
 #if WITH_EDITOR
+			// Cache paused icon
+			+ SHorizontalBox::Slot()
+			.VAlign(VAlign_Center)
+			.AutoWidth()
+			[
+				SNew(SButton)
+				.ButtonStyle(&FAppStyle::Get().GetWidgetStyle<FButtonStyle>("SimpleButton"))
+				.ToolTipText(LOCTEXT("AudioEventLog_CacheStats_PausedTooltip", "Caching has been paused - Audio Insights is no longer collecting data from the current session.\n\nClick/ESC to resume."))
+				.Visibility_Lambda([]()
+				{
+					FAudioInsightsModule& AudioInsightsModule = FAudioInsightsModule::GetChecked();
+					return AudioInsightsModule.GetTimingViewExtender().GetMessageCacheAndProcessingStatus() == ECacheAndProcess::None ? EVisibility::Visible : EVisibility::Hidden;
+				})
+				.OnClicked_Lambda([this]()
+				{
+					ResetInspectTimestamp();
+					return FReply::Handled();
+				})
+				[
+					SNew(SImage)
+					.Image(FSlateStyle::Get().GetBrush("AudioInsights.Icon.Pause"))
+					.ColorAndOpacity(StatTitleColor)
+					.DesiredSizeOverride(FVector2D(16.0f, 16.0f))
+				]
+				
+			]
+
 			// Cache stats
 			+ SHorizontalBox::Slot()
+			.Padding(4.0f, 0.0f, 0.0f, 0.0f)
 			.VAlign(VAlign_Center)
 			.AutoWidth()
 			[
@@ -654,6 +726,16 @@ namespace UE::Audio::Insights
 				LOCTEXT("AudioEventLog_Settings_VisibleColumnsTooltip", "Show/hide columns"),
 				FNewMenuDelegate::CreateRaw(this, &FAudioEventLogDashboardViewFactory::BuildVisibleColumnsMenuContent)
 			);
+
+#if WITH_EDITOR
+			// Auto-stop monitoring
+			MenuBuilder.AddSubMenu
+			(
+				LOCTEXT("AudioEventLog_Settings_AutoStopCaching", "Automatically Stop Caching"),
+				LOCTEXT("AudioEventLog_Settings_AutoStopCachingTooltip", "Choose if/how the event log automatically stops caching when inspecting an event."),
+				FNewMenuDelegate::CreateRaw(this, &FAudioEventLogDashboardViewFactory::BuildAutoStopCachingMenuContent)
+			);
+#endif // WITH_EDITOR
 		}
 
 		MenuBuilder.EndSection();
@@ -715,6 +797,29 @@ namespace UE::Audio::Insights
 		}
 	}
 
+	void FAudioEventLogDashboardViewFactory::OnUpdateAutoScroll()
+	{
+		if (!FilteredEntriesListView.IsValid())
+		{
+			return;
+		}
+
+		const int32 NumItems = FilteredEntriesListView->GetNumItemsBeingObserved();
+		if (NumItems == 0)
+		{
+			return;
+		}
+
+		if (SortMode == EColumnSortMode::Ascending)
+		{
+			bAutoScroll = FilteredEntriesListView->IsItemVisible(FilteredEntriesListView->GetItems()[NumItems - 1]);
+		}
+		else if (SortMode == EColumnSortMode::Descending)
+		{
+			bAutoScroll = FilteredEntriesListView->IsItemVisible(FilteredEntriesListView->GetItems()[0]);
+		}
+	}
+
 	void FAudioEventLogDashboardViewFactory::OnSelectionChanged(TSharedPtr<IDashboardDataViewEntry> SelectedItem, ESelectInfo::Type SelectInfo)
 	{
 		using namespace AudioEventLogPrivate;
@@ -726,15 +831,12 @@ namespace UE::Audio::Insights
 
 		if (!SelectedItem.IsValid())
 		{
-			FilteredEntriesListView->ClearHighlightedItems();
-
-			FAudioInsightsModule& AudioInsightsModule = FAudioInsightsModule::GetChecked();
-			AudioInsightsModule.GetTimingViewExtender().ResetTimeControlMethod();
-
-			bAutoScroll = true;
+			ResetInspectTimestamp();
 			return;
 		}
 
+		FocusedItem = SelectedItem;
+
 		const FAudioEventLogDashboardEntry& SelectedEventLogEntry = CastEntry(*SelectedItem);
 		if (SelectedEventLogEntry.PlayOrder == INDEX_NONE)
 		{
@@ -758,15 +860,64 @@ namespace UE::Audio::Insights
 			FAudioInsightsModule& AudioInsightsModule = FAudioInsightsModule::GetChecked();
 			AudioInsightsModule.GetTimingViewExtender().SetTimeMarker(SelectedEventLogEntry.Timestamp);
 
-			bAutoScroll = false;
+			if (AutoStopCachingMode == EAutoStopCachingMode::OnInspect)
+			{
+				AudioInsightsModule.GetTimingViewExtender().StopCachingAndProcessingNewMessages();
+				bAutoScroll = false;
+			}
+			else if (AutoStopCachingMode == EAutoStopCachingMode::WhenInLastChunk)
+			{
+				FAudioInsightsCacheManager& CacheManager = IAudioInsightsModule::GetChecked().GetCacheManager();
+
+				if (CacheManager.GetNumUsedChunks() == CacheManager.GetNumChunks() 
+					&& CacheManager.GetNumChunksFromStartForTimestamp(SelectedEventLogEntry.Timestamp) == 0u)
+				{
+					AudioInsightsModule.GetTimingViewExtender().StopCachingAndProcessingNewMessages();
+					bAutoScroll = false;
+				}
+				else
+				{
+					AudioInsightsModule.GetTimingViewExtender().StopProcessingNewMessages();
+				}
+			}
+			else
+			{
+				AudioInsightsModule.GetTimingViewExtender().StopProcessingNewMessages();
+			}
 		}
 	}
 
 	void FAudioEventLogDashboardViewFactory::OnListViewScrolled(double InScrollOffset)
 	{
-		if (FilteredEntriesListView.IsValid() && FilteredEntriesListView->IsUserScrolling())
+		if (!FilteredEntriesListView.IsValid())
+		{
+			return;
+		}
+
+		if (FilteredEntriesListView->IsUserScrolling())
+		{
+			OnUpdateAutoScroll();
+		}
+	}
+
+	void FAudioEventLogDashboardViewFactory::OnFinishedScrolling()
+	{
+		if (FilteredEntriesListView.IsValid() && FilteredEntriesListView->GetNumItemsBeingObserved() > 0)
 		{
-			bAutoScroll = false;
+			const TSharedPtr<IDashboardDataViewEntry> FocusedItemShared = FocusedItem.Pin();
+			
+			// Only set a focused item if the user has previously focused on a row
+			if (FocusedItemShared.IsValid())
+			{
+				for (const TSharedPtr<IDashboardDataViewEntry>& Entry : FilteredEntriesListView->GetItems())
+				{
+					if (FilteredEntriesListView->IsItemVisible(Entry))
+					{
+						FocusedItem = Entry;
+						break;
+					}
+				}
+			}
 		}
 	}
 
@@ -838,6 +989,15 @@ namespace UE::Audio::Insights
 		}
 	}
 
+	void FAudioEventLogDashboardViewFactory::BuildAutoStopCachingMenuContent(FMenuBuilder& MenuBuilder)
+	{
+		const FAudioEventLogEditorCommands& Commands = FAudioEventLogEditorCommands::Get();
+
+		MenuBuilder.AddMenuEntry(Commands.GetAutoStopCachingWhenLastInCacheCommand());
+		MenuBuilder.AddMenuEntry(Commands.GetAutoStopCachingOnInspectCommand());
+		MenuBuilder.AddMenuEntry(Commands.GetAutoStopCachingDisabledCommand());
+	}
+
 	void FAudioEventLogDashboardViewFactory::SortByPredicate(TFunctionRef<bool(const FAudioEventLogDashboardEntry&, const FAudioEventLogDashboardEntry&)> Predicate)
 	{
 		using namespace AudioEventLogPrivate;
@@ -914,6 +1074,18 @@ namespace UE::Audio::Insights
 		}
 	}
 
+	void FAudioEventLogDashboardViewFactory::OnCacheChunkOverwritten(const double NewCacheStartTimestamp)
+	{
+		if (IsInGameThread())
+		{
+			UpdateCacheMessageProcessMethod();
+		}
+		else
+		{
+			ExecuteOnGameThread(TEXT("FAudioEventLogDashboardViewFactory::UpdateCacheMessageProcessMethod"), [this] { UpdateCacheMessageProcessMethod(); });
+		}
+	}
+
 	void FAudioEventLogDashboardViewFactory::RemoveDeletedCustomEvents(const TMap<FString, FAudioEventLogCustomEvents>& CustomEventsFromSettings)
 	{
 		if (!FilterBar.IsValid())
@@ -1059,6 +1231,38 @@ namespace UE::Audio::Insights
 #endif // !WITH_EDITOR
 	}
 
+	void FAudioEventLogDashboardViewFactory::UpdateCacheMessageProcessMethod()
+	{
+		using namespace AudioEventLogPrivate;
+
+		if (AutoStopCachingMode != EAutoStopCachingMode::WhenInLastChunk || !FilteredEntriesListView.IsValid())
+		{
+			return;
+		}
+
+		FAudioInsightsModule& AudioInsightsModule = FAudioInsightsModule::GetChecked();
+		if (AudioInsightsModule.GetTimingViewExtender().GetMessageCacheAndProcessingStatus() != ECacheAndProcess::CacheLatestNoProcess)
+		{
+			return;
+		}
+
+		const TArray<TSharedPtr<IDashboardDataViewEntry>>& SelectedItems = FilteredEntriesListView->GetSelectedItems();
+		if (SelectedItems.Num() == 0 || !SelectedItems[0].IsValid())
+		{
+			return;
+		}
+
+		const FAudioEventLogDashboardEntry& SelectedEventLogEntry = CastEntry(*SelectedItems[0]);
+
+		FAudioInsightsCacheManager& CacheManager = IAudioInsightsModule::GetChecked().GetCacheManager();
+
+		if (CacheManager.GetNumUsedChunks() == CacheManager.GetNumChunks() 
+			&& CacheManager.GetNumChunksFromStartForTimestamp(SelectedEventLogEntry.Timestamp) == 0u)
+		{
+			AudioInsightsModule.GetTimingViewExtender().StopCachingAndProcessingNewMessages();
+		}
+	}
+
 } // namespace UE::Audio::Insights
 
 #undef LOCTEXT_NAMESPACE
diff --git a/Engine/Plugins/AudioInsights/Source/AudioInsights/Private/Views/SoundDashboardViewFactory.cpp b/Engine/Plugins/AudioInsights/Source/AudioInsights/Private/Views/SoundDashboardViewFactory.cpp
--- a/Engine/Plugins/AudioInsights/Source/AudioInsights/Private/Views/SoundDashboardViewFactory.cpp
+++ b/Engine/Plugins/AudioInsights/Source/AudioInsights/Private/Views/SoundDashboardViewFactory.cpp
@@ -1517,7 +1517,7 @@ namespace UE::Audio::Insights
 		{
 #if WITH_EDITOR
 			const FAudioInsightsModule& AudioInsightsModule = FAudioInsightsModule::GetChecked();
-			if (AudioInsightsModule.GetTimingViewExtender().GetTimeControlMethod() != ETimeControlMethod::LatestMessageTime)
+			if (AudioInsightsModule.GetTimingViewExtender().GetMessageCacheAndProcessingStatus() != ECacheAndProcess::Latest)
 			{
 				return;
 			}
diff --git a/Engine/Plugins/AudioInsights/Source/AudioInsights/Private/Views/TableDashboardViewFactory.cpp b/Engine/Plugins/AudioInsights/Source/AudioInsights/Private/Views/TableDashboardViewFactory.cpp
--- a/Engine/Plugins/AudioInsights/Source/AudioInsights/Private/Views/TableDashboardViewFactory.cpp
+++ b/Engine/Plugins/AudioInsights/Source/AudioInsights/Private/Views/TableDashboardViewFactory.cpp
@@ -222,11 +222,12 @@ namespace UE::Audio::Insights
 		{
 			FilterBarButton = CreateFilterBarButtonWidget();
 
-			FilteredEntriesListView = SNew(SListView<TSharedPtr<IDashboardDataViewEntry>>)
+			FilteredEntriesListView = SNew(SScrollableListView<TSharedPtr<IDashboardDataViewEntry>>)
 			.ListItemsSource(&DataViewEntries)
 			.OnContextMenuOpening(this, &FTraceTableDashboardViewFactory::OnConstructContextMenu)
 			.OnSelectionChanged(this, &FTraceTableDashboardViewFactory::OnSelectionChanged)
 			.OnListViewScrolled(this, &FTraceTableDashboardViewFactory::OnListViewScrolled)
+			.OnFinishedScrolling(this, &FTraceTableDashboardViewFactory::OnFinishedScrolling)
 			.ClearSelectionOnClick(ClearSelectionOnClick())
 			.OnGenerateRow_Lambda([this](TSharedPtr<IDashboardDataViewEntry> Item, const TSharedRef<STableViewBase>& OwnerTable)
 			{
diff --git a/Engine/Plugins/AudioInsights/Source/AudioInsights/Public/AudioInsightsTraceProviderBase.h b/Engine/Plugins/AudioInsights/Source/AudioInsights/Public/AudioInsightsTraceProviderBase.h
--- a/Engine/Plugins/AudioInsights/Source/AudioInsights/Public/AudioInsightsTraceProviderBase.h
+++ b/Engine/Plugins/AudioInsights/Source/AudioInsights/Public/AudioInsightsTraceProviderBase.h
@@ -22,7 +22,7 @@
 
 namespace UE::Audio::Insights
 {
-	enum class ETimeControlMethod;
+	enum class ECacheAndProcess;
 
 	class FTraceProviderBase : public TraceServices::IProvider, public TraceServices::IEditableProvider
 	{
@@ -79,8 +79,8 @@ namespace UE::Audio::Insights
 		virtual void OnTimingViewTimeMarkerChanged(double TimeMarker) { ++LastMessageId; };
 		virtual void OnTimeControlMethodReset() { };
 
-		ETimeControlMethod GetTimeControlMethod() const;
-		void ResetTimeControlMethod();
+		AUDIOINSIGHTS_API ECacheAndProcess GetMessageCacheAndProcessingStatus() const;
+		AUDIOINSIGHTS_API void ResetMessageProcessType();
 
 	protected:
 		class FTraceAnalyzerBase : public Trace::IAnalyzer
@@ -99,6 +99,7 @@ namespace UE::Audio::Insights
 			AUDIOINSIGHTS_API virtual bool OnHandleEvent(uint16 RouteId, EStyle Style, const FOnEventContext& Context) = 0;
 
 			AUDIOINSIGHTS_API bool ShouldIgnoreNewEvents() const;
+			AUDIOINSIGHTS_API virtual bool ShouldProcessNewEvents() const;
 
 			template <typename TProviderType>
 			TProviderType& GetProvider()
@@ -106,6 +107,12 @@ namespace UE::Audio::Insights
 				return *StaticCastSharedRef<TProviderType>(Provider);
 			}
 
+			template <typename TProviderType>
+			const TProviderType& GetProvider() const
+			{
+				return *StaticCastSharedRef<TProviderType>(Provider);
+			}
+
 #if WITH_EDITOR
 			template <TIsCacheableMessage T>
 			void CacheMessage(const FOnEventContext& Context, TAnalyzerMessageQueue<T>& OutMessageQueue)
@@ -115,8 +122,11 @@ namespace UE::Audio::Insights
 				TSharedPtr<T> Message = MakeShared<T>(Context);
 				CacheManager.AddMessageToCache(Message);
 
-				T MessageCopy = *Message;
-				OutMessageQueue.Enqueue(MoveTemp(MessageCopy));
+				if (ShouldProcessNewEvents())
+				{
+					T MessageCopy = *Message;
+					OutMessageQueue.Enqueue(MoveTemp(MessageCopy));
+				}
 			}
 #endif // WITH_EDITOR
 
@@ -131,7 +141,7 @@ namespace UE::Audio::Insights
 			return LastMessageId;
 		}
 
-		AUDIOINSIGHTS_API bool ProcessMessagesIsPaused() const;
+		AUDIOINSIGHTS_API virtual bool IsMessageProcessingPaused() const;
 
 		uint64 LastUpdateId = 0;
 		bool bForceUpdate = false;
@@ -158,7 +168,7 @@ namespace UE::Audio::Insights
 		{
 			TickerHandle = FTSTicker::GetCoreTicker().AddTicker(*InName.ToString(), 0.0f, [this](float DeltaTime)
 			{
-				if (LastUpdateId != GetLastMessageId() && !ProcessMessagesIsPaused())
+				if (LastUpdateId != GetLastMessageId() && !IsMessageProcessingPaused())
 				{
 					ProcessMessages();
 				}
diff --git a/Engine/Plugins/AudioInsights/Source/AudioInsights/Public/Views/AudioEventLogDashboardViewFactory.h b/Engine/Plugins/AudioInsights/Source/AudioInsights/Public/Views/AudioEventLogDashboardViewFactory.h
--- a/Engine/Plugins/AudioInsights/Source/AudioInsights/Public/Views/AudioEventLogDashboardViewFactory.h
+++ b/Engine/Plugins/AudioInsights/Source/AudioInsights/Public/Views/AudioEventLogDashboardViewFactory.h
@@ -86,6 +86,7 @@ namespace UE::Audio::Insights
 		virtual bool ClearSelectionOnClick() const override { return true; }
 		virtual bool EnableHorizontalScrollBox() const override { return false; }
 		UE_API virtual void OnListViewScrolled(double InScrollOffset) override;
+		UE_API virtual void OnFinishedScrolling();
 		UE_API virtual const FTableRowStyle* GetRowStyle() const override;
 
 #if WITH_EDITOR
@@ -98,6 +99,7 @@ namespace UE::Audio::Insights
 		bool PassesEventTypeFilter(const FEventLogFilterID EventLogID) const;
 
 		void BuildVisibleColumnsMenuContent(FMenuBuilder& MenuBuilder);
+		void BuildAutoStopCachingMenuContent(FMenuBuilder& MenuBuilder);
 
 		void SortByPredicate(TFunctionRef<bool(const class FAudioEventLogDashboardEntry&, const class FAudioEventLogDashboardEntry&)> Predicate);
 
@@ -105,6 +107,8 @@ namespace UE::Audio::Insights
 		void OnReadEditorSettings(const FAudioEventLogSettings& InSettings);
 		void OnWriteEditorSettings(FAudioEventLogSettings& OutSettings);
 
+		void OnCacheChunkOverwritten(const double NewCacheStartTimestamp);
+
 		void RemoveDeletedCustomEvents(const TMap<FString, FAudioEventLogCustomEvents>& CustomEventsFromSettings);
 		void AddNewCustomEvents(const TMap<FString, FAudioEventLogCustomEvents>& CustomEventsFromSettings);
 
@@ -113,11 +117,15 @@ namespace UE::Audio::Insights
 
 		void OnTimingViewTimeMarkerChanged(double InTimeMarker);
 
+		void UpdateCacheMessageProcessMethod();
+
 		void BindCommands();
 		void ResetInspectTimestamp();
 
 		void ClearActiveAudioDeviceEntries();
 
+		void OnUpdateAutoScroll();
+
 		TSharedPtr<FAudioEventLogTraceProvider> AudioEventLogTraceProvider;
 
 		TSharedPtr<SWidget> DashboardWidget;
@@ -131,11 +139,23 @@ namespace UE::Audio::Insights
 		TSharedPtr<FUICommandList> CommandList;
 		TSharedPtr<FVisibleColumnsSettingsMenu<FAudioEventLogVisibleColumns>> VisibleColumnsSettingsMenu;
 
+		TWeakPtr<IDashboardDataViewEntry> FocusedItem;
 		
 		const TMap<FString, TSet<FString>> InitEventTypeFilters;
 		TMap<FString, FAudioEventLogCustomEvents> CachedCustomEventSettings;
 
 		bool bAutoScroll = true;
+
+#if WITH_EDITOR
+		enum class EAutoStopCachingMode
+		{
+			WhenInLastChunk,
+			OnInspect,
+			Never
+		};
+
+		EAutoStopCachingMode AutoStopCachingMode = EAutoStopCachingMode::WhenInLastChunk;
+#endif // WITH_EDITOR
 	};
 } // namespace UE::Audio::Insights
 
diff --git a/Engine/Plugins/AudioInsights/Source/AudioInsights/Public/Views/TableDashboardViewFactory.h b/Engine/Plugins/AudioInsights/Source/AudioInsights/Public/Views/TableDashboardViewFactory.h
--- a/Engine/Plugins/AudioInsights/Source/AudioInsights/Public/Views/TableDashboardViewFactory.h
+++ b/Engine/Plugins/AudioInsights/Source/AudioInsights/Public/Views/TableDashboardViewFactory.h
@@ -3,6 +3,7 @@
 
 #include "AudioDefines.h"
 #include "AudioInsightsDataSource.h"
+#include "Delegates/DelegateCombinations.h"
 #include "Templates/SharedPointer.h"
 #include "Views/DashboardViewFactory.h"
 #include "Widgets/Input/SSearchBox.h"
@@ -13,6 +14,27 @@
 
 namespace UE::Audio::Insights
 {
+	
+	template <typename ItemType>
+	class SScrollableListView : public SListView<ItemType>
+	{
+	public:
+		void Construct(const typename SListView<ItemType>::FArguments& InArgs)
+		{
+			SListView<ItemType>::Construct(InArgs);
+		}
+
+		virtual FReply OnMouseWheel(const FGeometry& MyGeometry, const FPointerEvent& MouseEvent) override
+		{
+			OnMouseWheelDetected.Broadcast();
+
+			return SListView<ItemType>::OnMouseWheel(MyGeometry, MouseEvent);
+		}
+
+		DECLARE_MULTICAST_DELEGATE(FOnMouseWheel);
+		FOnMouseWheel OnMouseWheelDetected;
+	};
+
 	class IObjectDashboardEntry : public IDashboardDataViewEntry
 	{
 	public:
@@ -161,6 +183,7 @@ namespace UE::Audio::Insights
 
 		virtual bool EnableHorizontalScrollBox() const { return true; }
 		virtual void OnListViewScrolled(double InScrollOffset) { }
+		virtual void OnFinishedScrolling() {}
 
 		virtual const FTableRowStyle* GetRowStyle() const { return nullptr; }
 		UE_API virtual FSlateColor GetRowColor(const TSharedPtr<IDashboardDataViewEntry>& InRowDataPtr);
@@ -175,7 +198,7 @@ namespace UE::Audio::Insights
 
 		TSharedPtr<SWidget> DashboardWidget;
 		TSharedPtr<SHeaderRow> HeaderRowWidget;
-		TSharedPtr<SListView<TSharedPtr<IDashboardDataViewEntry>>> FilteredEntriesListView;
+		TSharedPtr<SScrollableListView<TSharedPtr<IDashboardDataViewEntry>>> FilteredEntriesListView;
 		TSharedPtr<SWidget> FilterBarButton;
 
 		FName SortByColumn;
