[Audio Insights] Fix build for non editor builds

diff --git a/Engine/Plugins/AudioInsights/Source/AudioInsights/Private/Views/AudioEventLogDashboardViewFactory.cpp b/Engine/Plugins/AudioInsights/Source/AudioInsights/Private/Views/AudioEventLogDashboardViewFactory.cpp
--- a/Engine/Plugins/AudioInsights/Source/AudioInsights/Private/Views/AudioEventLogDashboardViewFactory.cpp
+++ b/Engine/Plugins/AudioInsights/Source/AudioInsights/Private/Views/AudioEventLogDashboardViewFactory.cpp
@@ -860,6 +860,7 @@ namespace UE::Audio::Insights
 			FAudioInsightsModule& AudioInsightsModule = FAudioInsightsModule::GetChecked();
 			AudioInsightsModule.GetTimingViewExtender().SetTimeMarker(SelectedEventLogEntry.Timestamp);
 
+#if WITH_EDITOR
 			if (AutoStopCachingMode == EAutoStopCachingMode::OnInspect)
 			{
 				AudioInsightsModule.GetTimingViewExtender().StopCachingAndProcessingNewMessages();
@@ -884,6 +885,9 @@ namespace UE::Audio::Insights
 			{
 				AudioInsightsModule.GetTimingViewExtender().StopProcessingNewMessages();
 			}
+#else
+			bAutoScroll = false;
+#endif	// WITH_EDITOR
 		}
 	}
 
@@ -1235,6 +1239,7 @@ namespace UE::Audio::Insights
 	{
 		using namespace AudioEventLogPrivate;
 
+#if WITH_EDITOR
 		if (AutoStopCachingMode != EAutoStopCachingMode::WhenInLastChunk || !FilteredEntriesListView.IsValid())
 		{
 			return;
@@ -1261,6 +1266,7 @@ namespace UE::Audio::Insights
 		{
 			AudioInsightsModule.GetTimingViewExtender().StopCachingAndProcessingNewMessages();
 		}
+#endif // WITH_EDITOR
 	}
 
 } // namespace UE::Audio::Insights
