[Audio Insights] Bughawk suggestions for recent cache changes.
* Use TOptional instead of returning INDEX_NONE from AudioInsightsCacheManager::GetNumChunksFromStartForTimestamp to avoid possible errors in interpreting the returned index
* Make sure we clear the pending message queue inside the cache manager when we clear the cache
* Use Num() from SListView::GetItems rather than SListView::GetNumItemsBeingObserved to ensure we are definitely checking the right number of items before indexing into the array

diff --git a/Engine/Plugins/AudioInsights/Source/AudioInsights/Private/Cache/AudioInsightsCacheManager.cpp b/Engine/Plugins/AudioInsights/Source/AudioInsights/Private/Cache/AudioInsightsCacheManager.cpp
--- a/Engine/Plugins/AudioInsights/Source/AudioInsights/Private/Cache/AudioInsightsCacheManager.cpp
+++ b/Engine/Plugins/AudioInsights/Source/AudioInsights/Private/Cache/AudioInsightsCacheManager.cpp
@@ -38,6 +38,8 @@ namespace UE::Audio::Insights
 
 	void FAudioInsightsCacheManager::ClearCache()
 	{
+		NewMessageQueue.DequeueAll();
+
 		CreateCache();
 	}
 
@@ -97,7 +99,7 @@ namespace UE::Audio::Insights
 		return NumUsedChunks;
 	}
 
-	uint32 FAudioInsightsCacheManager::GetNumChunksFromStartForTimestamp(const double Timestamp) const
+	TOptional<uint32> FAudioInsightsCacheManager::TryGetNumChunksFromStartForTimestamp(const double Timestamp) const
 	{
 		uint32 ReadIndex = StartIndex;
 		for (uint32 Index = 0u; Index < Cache.Capacity(); ++Index)
@@ -112,7 +114,7 @@ namespace UE::Audio::Insights
 			}
 		}
 
-		return INDEX_NONE;
+		return TOptional<uint32>();
 	}
 
 	void FAudioInsightsCacheManager::CreateCache()
diff --git a/Engine/Plugins/AudioInsights/Source/AudioInsights/Private/Views/AudioEventLogDashboardViewFactory.cpp b/Engine/Plugins/AudioInsights/Source/AudioInsights/Private/Views/AudioEventLogDashboardViewFactory.cpp
--- a/Engine/Plugins/AudioInsights/Source/AudioInsights/Private/Views/AudioEventLogDashboardViewFactory.cpp
+++ b/Engine/Plugins/AudioInsights/Source/AudioInsights/Private/Views/AudioEventLogDashboardViewFactory.cpp
@@ -36,6 +36,22 @@ namespace UE::Audio::Insights
 #endif // WITH_EDITOR
 		}
 
+#if WITH_EDITOR
+		bool TimestampIsInChunkMarkedToDelete(const double Timestamp)
+		{
+			FAudioInsightsCacheManager& CacheManager = IAudioInsightsModule::GetChecked().GetCacheManager();
+
+			if (CacheManager.GetNumUsedChunks() != CacheManager.GetNumChunks())
+			{
+				// The cache is not full - no chunks are marked to delete
+				return false;
+			}
+
+			const TOptional<uint32> NumChunksFromStart = CacheManager.TryGetNumChunksFromStartForTimestamp(Timestamp);
+			return NumChunksFromStart.IsSet() && NumChunksFromStart.GetValue() == 0u;
+		}
+#endif // WITH_EDITOR
+
 		const FString CategorySoundActivity = LOCTEXT("EventLogCategory_SoundActivity", "Sound Activity").ToString();
 		const FString CategoryVirtualization = LOCTEXT("EventLogCategory_Virtualization", "Virtualization").ToString();
 		const FString CategoryPlayRequests = LOCTEXT("EventLogCategory_PlayRequests", "Play Requests").ToString();
@@ -804,7 +820,8 @@ namespace UE::Audio::Insights
 			return;
 		}
 
-		const int32 NumItems = FilteredEntriesListView->GetNumItemsBeingObserved();
+		const TArrayView<const TSharedPtr<IDashboardDataViewEntry>> FilteredItems = FilteredEntriesListView->GetItems();
+		const int32 NumItems = FilteredItems.Num();
 		if (NumItems == 0)
 		{
 			return;
@@ -812,11 +829,11 @@ namespace UE::Audio::Insights
 
 		if (SortMode == EColumnSortMode::Ascending)
 		{
-			bAutoScroll = FilteredEntriesListView->IsItemVisible(FilteredEntriesListView->GetItems()[NumItems - 1]);
+			bAutoScroll = FilteredEntriesListView->IsItemVisible(FilteredItems[NumItems - 1]);
 		}
 		else if (SortMode == EColumnSortMode::Descending)
 		{
-			bAutoScroll = FilteredEntriesListView->IsItemVisible(FilteredEntriesListView->GetItems()[0]);
+			bAutoScroll = FilteredEntriesListView->IsItemVisible(FilteredItems[0]);
 		}
 	}
 
@@ -868,10 +885,7 @@ namespace UE::Audio::Insights
 			}
 			else if (AutoStopCachingMode == EAutoStopCachingMode::WhenInLastChunk)
 			{
-				FAudioInsightsCacheManager& CacheManager = IAudioInsightsModule::GetChecked().GetCacheManager();
-
-				if (CacheManager.GetNumUsedChunks() == CacheManager.GetNumChunks() 
-					&& CacheManager.GetNumChunksFromStartForTimestamp(SelectedEventLogEntry.Timestamp) == 0u)
+				if (TimestampIsInChunkMarkedToDelete(SelectedEventLogEntry.Timestamp))
 				{
 					AudioInsightsModule.GetTimingViewExtender().StopCachingAndProcessingNewMessages();
 					bAutoScroll = false;
@@ -1259,10 +1273,7 @@ namespace UE::Audio::Insights
 
 		const FAudioEventLogDashboardEntry& SelectedEventLogEntry = CastEntry(*SelectedItems[0]);
 
-		FAudioInsightsCacheManager& CacheManager = IAudioInsightsModule::GetChecked().GetCacheManager();
-
-		if (CacheManager.GetNumUsedChunks() == CacheManager.GetNumChunks() 
-			&& CacheManager.GetNumChunksFromStartForTimestamp(SelectedEventLogEntry.Timestamp) == 0u)
+		if (TimestampIsInChunkMarkedToDelete(SelectedEventLogEntry.Timestamp))
 		{
 			AudioInsightsModule.GetTimingViewExtender().StopCachingAndProcessingNewMessages();
 		}
diff --git a/Engine/Plugins/AudioInsights/Source/AudioInsights/Public/Cache/AudioInsightsCacheManager.h b/Engine/Plugins/AudioInsights/Source/AudioInsights/Public/Cache/AudioInsightsCacheManager.h
--- a/Engine/Plugins/AudioInsights/Source/AudioInsights/Public/Cache/AudioInsightsCacheManager.h
+++ b/Engine/Plugins/AudioInsights/Source/AudioInsights/Public/Cache/AudioInsightsCacheManager.h
@@ -6,6 +6,7 @@
 #include "Containers/CircularBuffer.h"
 #include "Containers/Ticker.h"
 #include "Messages/AnalyzerMessageQueue.h"
+#include "Misc/Optional.h"
 #include "TraceServices/Containers/SlabAllocator.h"
 
 #define UE_API AUDIOINSIGHTS_API
@@ -37,8 +38,9 @@ namespace UE::Audio::Insights
 		// Returns the number of chunks with data present inside the cache
 		UE_API uint32 GetNumUsedChunks() const;
 
-		// Returns the number of chunks from the start of the cache for a specific timestamp, where lower indexed chunks contain older messages
-		UE_API uint32 GetNumChunksFromStartForTimestamp(const double Timestamp) const;
+		// Tries to find the number of chunks from the start of the cache for a specific timestamp, where lower indexed chunks contain older messages
+		// If the timestamp does not exist within the current cache, the returned result will be unset
+		UE_API TOptional<uint32> TryGetNumChunksFromStartForTimestamp(const double Timestamp) const;
 
 		// Find the closest message to a timestamp in the cache for a specific message type
 		// Will look backwards from the chosen timestamp until a message is found
