[Audio Insights] Fix for plots widget not displaying curves correctly in standalone
* FAudioInsightsComponent::GetIsLiveSession was changed recently to also include direct trace sessions in standalone, but new impl did not work with trace files being loaded in. Use trace ID == 0 to do additional checks if we are doing direct trace.
* Standalone requires all plot data to be processed when a file is being loaded in - remove check of whether we have a valid widget in SoundsDashboardViewFactory and send an empty array for selected items instead
* Only reset the game state to Running inside the plots widget view in standalone if we are connected to a live session

diff --git a/Engine/Plugins/AudioInsights/Source/AudioInsights/Private/AudioInsightsComponent.cpp b/Engine/Plugins/AudioInsights/Source/AudioInsights/Private/AudioInsightsComponent.cpp
--- a/Engine/Plugins/AudioInsights/Source/AudioInsights/Private/AudioInsightsComponent.cpp
+++ b/Engine/Plugins/AudioInsights/Source/AudioInsights/Private/AudioInsightsComponent.cpp
@@ -137,17 +137,27 @@ namespace UE::Audio::Insights
 
 	bool FAudioInsightsComponent::GetIsLiveSession() const
 	{
-		const IUnrealInsightsModule& UnrealInsightsModule = FModuleManager::LoadModuleChecked<IUnrealInsightsModule>("TraceInsights");
+		IUnrealInsightsModule& UnrealInsightsModule = FModuleManager::LoadModuleChecked<IUnrealInsightsModule>("TraceInsights");
 
 		TSharedPtr<const TraceServices::IAnalysisSession> Session = UnrealInsightsModule.GetAnalysisSession();
 		if (Session.IsValid())
 		{
 			TraceServices::FAnalysisSessionReadScope SessionReadScope(*Session.Get());
 
-			const TraceServices::IDiagnosticsProvider* DiagnosticsProvider = TraceServices::ReadDiagnosticsProvider(*Session.Get());
-			if (DiagnosticsProvider && DiagnosticsProvider->IsSessionInfoAvailable())
+			const UE::Trace::FStoreClient* StoreClient = UnrealInsightsModule.GetStoreClient();
+			if (StoreClient)
+			{
+				const UE::Trace::FStoreClient::FSessionInfo* SessionInfo = StoreClient->GetSessionInfoByTraceId(Session->GetTraceId());
+				return !Session->IsAnalysisComplete() && SessionInfo != nullptr;
+			}
+			// Direct traces have a trace ID of zero and a null StoreClient - check IDiagnosticsProvider to see if we are running a direct trace
+			else if (Session->GetTraceId() == 0)
 			{
-				return !Session->IsAnalysisComplete() && DiagnosticsProvider->IsSessionInfoAvailable();
+				const TraceServices::IDiagnosticsProvider* DiagnosticsProvider = TraceServices::ReadDiagnosticsProvider(*Session.Get());
+				if (DiagnosticsProvider && DiagnosticsProvider->IsSessionInfoAvailable())
+				{
+					return !Session->IsAnalysisComplete();
+				}
 			}
 		}
 
diff --git a/Engine/Plugins/AudioInsights/Source/AudioInsights/Private/Views/SoundDashboardViewFactory.cpp b/Engine/Plugins/AudioInsights/Source/AudioInsights/Private/Views/SoundDashboardViewFactory.cpp
--- a/Engine/Plugins/AudioInsights/Source/AudioInsights/Private/Views/SoundDashboardViewFactory.cpp
+++ b/Engine/Plugins/AudioInsights/Source/AudioInsights/Private/Views/SoundDashboardViewFactory.cpp
@@ -3189,10 +3189,24 @@ namespace UE::Audio::Insights
 
 	void FSoundDashboardViewFactory::ProcessPlotData()
 	{
+#if WITH_EDITOR
+		if (FilteredEntriesListView.IsValid())
+		{
+			OnProcessPlotData.Broadcast(DataViewEntries, FilteredEntriesListView->GetSelectedItems());
+		}
+#else
 		if (FilteredEntriesListView.IsValid())
 		{
 			OnProcessPlotData.Broadcast(DataViewEntries, FilteredEntriesListView->GetSelectedItems());
 		}
+		else
+		{
+			// In standalone, we still want to process entries as the trace file loads in
+			// Send an empty array for selected items and broadcast the data
+			const TArray<TSharedPtr<IDashboardDataTreeViewEntry>> EmptySelectedEntriesContainer;
+			OnProcessPlotData.Broadcast(DataViewEntries, EmptySelectedEntriesContainer);
+		}
+#endif // WITH_EDITOR
 	}
 
 	TMap<FName, FSoundPlotsWidgetView::FPlotColumnInfo> FSoundDashboardViewFactory::GetPlotColumnInfo()
diff --git a/Engine/Plugins/AudioInsights/Source/AudioInsights/Private/Views/SoundPlotsWidgetView.cpp b/Engine/Plugins/AudioInsights/Source/AudioInsights/Private/Views/SoundPlotsWidgetView.cpp
--- a/Engine/Plugins/AudioInsights/Source/AudioInsights/Private/Views/SoundPlotsWidgetView.cpp
+++ b/Engine/Plugins/AudioInsights/Source/AudioInsights/Private/Views/SoundPlotsWidgetView.cpp
@@ -303,8 +303,17 @@ namespace UE::Audio::Insights
 
 	void FSoundPlotsWidgetView::OnTimeControlMethodReset()
 	{
-		GameState = EGameState::Running;
 		bTimeMarkerMoved = false;
+
+#if WITH_EDITOR
+		GameState = EGameState::Running;
+#else
+		const TSharedPtr<const FAudioInsightsComponent> AudioInsightsComponent = FAudioInsightsModule::GetChecked().GetAudioInsightsComponent();
+		if (AudioInsightsComponent.IsValid() && AudioInsightsComponent->GetIsLiveSession())
+		{
+			GameState = EGameState::Running;
+		}
+#endif // WITH_EDITOR
 	}
 
 	TSharedRef<SWidget> FSoundPlotsWidgetView::MakePlotsWidget()
@@ -529,7 +538,15 @@ namespace UE::Audio::Insights
 			}
 			else if (GameState == EGameState::ControlledByTimemarker)
 			{
+#if WITH_EDITOR
 				DataPoints->Empty(DataPoints->Num());
+#else
+				const TSharedPtr<const FAudioInsightsComponent> AudioInsightsComponent = FAudioInsightsModule::GetChecked().GetAudioInsightsComponent();
+				if (AudioInsightsComponent.IsValid() && AudioInsightsComponent->GetIsLiveSession())
+				{
+					DataPoints->Empty(DataPoints->Num());
+				}
+#endif // WITH_EDITOR
 			}
 
 			auto DataFunc = GetPlotColumnDataFunc(ColumnName);
@@ -782,7 +799,7 @@ namespace UE::Audio::Insights
 
 		const bool bAnyMessageReceived = CurrentTimestamp != TNumericLimits<double>::Lowest();
 
-		if (GameState != EGameState::Paused && GameState != EGameState ::ControlledByTimemarker && bAnyMessageReceived)
+		if (GameState != EGameState::Paused && GameState != EGameState::ControlledByTimemarker && bAnyMessageReceived)
 		{
 			const double RelativeCurrentTime = CurrentTime - BeginTimestamp;
 			TimestampsDiff = RelativeCurrentTime - (CurrentTimestamp - BeginTimestamp);
