[Audio Insights] Make sure timestamps in the event log and plots view match

diff --git a/Engine/Plugins/AudioInsights/Source/AudioInsights/Private/Views/AudioEventLogDashboardViewFactory.cpp b/Engine/Plugins/AudioInsights/Source/AudioInsights/Private/Views/AudioEventLogDashboardViewFactory.cpp
--- a/Engine/Plugins/AudioInsights/Source/AudioInsights/Private/Views/AudioEventLogDashboardViewFactory.cpp
+++ b/Engine/Plugins/AudioInsights/Source/AudioInsights/Private/Views/AudioEventLogDashboardViewFactory.cpp
@@ -5,6 +5,7 @@
 #include "AudioInsightsModule.h"
 #include "AudioInsightsStyle.h"
 #include "AudioMixerTrace.h"
+#include "Features/IModularFeatures.h"
 #include "Messages/AudioEventLogTraceMessages.h"
 #include "Providers/AudioEventLogTraceProvider.h"
 
@@ -156,6 +157,7 @@ namespace UE::Audio::Insights
 		using namespace AudioEventLogPrivate;
 
 		FTraceModule& AudioInsightsTraceModule = static_cast<FTraceModule&>(FAudioInsightsModule::GetChecked().GetTraceModule());
+		AudioInsightsTraceModule.OnAnalysisStarting.AddRaw(this, &FAudioEventLogDashboardViewFactory::OnAnalysisStarting);
 
 		AudioEventLogTraceProvider = MakeShared<FAudioEventLogTraceProvider>();
 
@@ -183,6 +185,12 @@ namespace UE::Audio::Insights
 
 	FAudioEventLogDashboardViewFactory::~FAudioEventLogDashboardViewFactory()
 	{
+		if (FModuleManager::Get().IsModuleLoaded("AudioInsights") && IModularFeatures::Get().IsModularFeatureAvailable(TraceServices::ModuleFeatureName))
+		{
+			FTraceModule& TraceModule = static_cast<FTraceModule&>(FAudioInsightsModule::GetChecked().GetTraceModule());
+			TraceModule.OnAnalysisStarting.RemoveAll(this);
+		}
+
 #if WITH_EDITOR
 		FAudioEventLogSettings::OnReadSettings.RemoveAll(this);
 		FAudioEventLogSettings::OnWriteSettings.RemoveAll(this);
@@ -389,7 +397,7 @@ namespace UE::Audio::Insights
 					MessageIDColumnID,
 					{
 						LOCTEXT("AudioEventLog_MessageIDColumnDisplayName", "ID"),
-						[](const IDashboardDataViewEntry& InData) { return FText::AsNumber(AudioEventLogPrivate::CastEntry(InData).MessageID); },
+						[](const IDashboardDataViewEntry& InData) { return FText::AsNumber(CastEntry(InData).MessageID); },
 						nullptr /* GetIconName */,
 						!VisibleColumns.GetIsVisible(MessageIDColumnID) /* bDefaultHidden */,
 						0.1f /* FillWidth */
@@ -399,7 +407,7 @@ namespace UE::Audio::Insights
 					TimestampColumnID,
 					{
 						LOCTEXT("AudioEventLog_TimestampColumnDisplayName", "Timestamp"),
-						[](const IDashboardDataViewEntry& InData) { return FText::AsNumber(AudioEventLogPrivate::CastEntry(InData).Timestamp, FSlateStyle::Get().GetTimeFormat()); },
+						[this](const IDashboardDataViewEntry& InData) { return FText::AsNumber(GetTimestampRelativeToAnalysisStart(CastEntry(InData).Timestamp), FSlateStyle::Get().GetTimeFormat()); },
 						nullptr /* GetIconName */,
 						!VisibleColumns.GetIsVisible(TimestampColumnID) /* bDefaultHidden */,
 						0.15f /* FillWidth */
@@ -428,7 +436,7 @@ namespace UE::Audio::Insights
 					EventColumnID,
 					{
 						LOCTEXT("AudioEventLog_EventColumnDisplayName", "Event"),
-						[](const IDashboardDataViewEntry& InData) { return FText::FromString(AudioEventLogPrivate::CastEntry(InData).EventName); },
+						[](const IDashboardDataViewEntry& InData) { return FText::FromString(CastEntry(InData).EventName); },
 						nullptr /* GetIconName */,
 						!VisibleColumns.GetIsVisible(EventColumnID) /* bDefaultHidden */,
 						0.2f  /* FillWidth */
@@ -438,7 +446,7 @@ namespace UE::Audio::Insights
 					AssetNameColumnID,
 					{
 						LOCTEXT("AudioEventLog_AssetColumnDisplayName", "Asset"),
-						[](const IDashboardDataViewEntry& InData) { return AudioEventLogPrivate::CastEntry(InData).GetDisplayName(); },
+						[](const IDashboardDataViewEntry& InData) { return CastEntry(InData).GetDisplayName(); },
 						[this](const IDashboardDataViewEntry& InData) -> FName
 						{
 							const FAudioEventLogDashboardEntry& Entry = CastEntry(InData);
@@ -470,7 +478,7 @@ namespace UE::Audio::Insights
 					ActorColumnID,
 					{
 						LOCTEXT("AudioEventLog_ActorColumnDisplayName", "Actor"),
-						[](const IDashboardDataViewEntry& InData) { return FText::FromString(AudioEventLogPrivate::CastEntry(InData).ActorLabel); },
+						[](const IDashboardDataViewEntry& InData) { return FText::FromString(CastEntry(InData).ActorLabel); },
 						[this](const IDashboardDataViewEntry& InData) -> FName
 						{
 							return CastEntry(InData).ActorIconName;
@@ -483,7 +491,7 @@ namespace UE::Audio::Insights
 					CategoryColumnID,
 					{
 						LOCTEXT("AudioEventLog_CategoryColumnDisplayName", "Category"),
-						[](const IDashboardDataViewEntry& InData) { return FText::FromString(AudioEventLogPrivate::CastEntry(InData).CategoryName); },
+						[](const IDashboardDataViewEntry& InData) { return FText::FromString(CastEntry(InData).CategoryName); },
 						nullptr /* GetIconName */,
 						!VisibleColumns.GetIsVisible(CategoryColumnID) /* bDefaultHidden */,
 						0.3f  /* FillWidth */
@@ -837,6 +845,11 @@ namespace UE::Audio::Insights
 		}
 	}
 
+	double FAudioEventLogDashboardViewFactory::GetTimestampRelativeToAnalysisStart(const double Timestamp) const
+	{
+		return Timestamp - BeginTimestamp;
+	}
+
 	void FAudioEventLogDashboardViewFactory::OnSelectionChanged(TSharedPtr<IDashboardDataViewEntry> SelectedItem, ESelectInfo::Type SelectInfo)
 	{
 		using namespace AudioEventLogPrivate;
@@ -944,6 +957,15 @@ namespace UE::Audio::Insights
 		return &FSlateStyle::Get().GetWidgetStyle<FTableRowStyle>("TreeDashboard.TableViewRow");
 	}
 
+	void FAudioEventLogDashboardViewFactory::OnAnalysisStarting(const double Timestamp)
+	{
+#if WITH_EDITOR
+		BeginTimestamp = Timestamp - GStartTime;
+#else
+		BeginTimestamp = 0.0;
+#endif // WITH_EDITOR
+	}
+
 	void FAudioEventLogDashboardViewFactory::UpdateCustomEventLogFilters(const IDashboardDataViewEntry& Entry)
 	{
 		using namespace AudioEventLogPrivate;
diff --git a/Engine/Plugins/AudioInsights/Source/AudioInsights/Public/Views/AudioEventLogDashboardViewFactory.h b/Engine/Plugins/AudioInsights/Source/AudioInsights/Public/Views/AudioEventLogDashboardViewFactory.h
--- a/Engine/Plugins/AudioInsights/Source/AudioInsights/Public/Views/AudioEventLogDashboardViewFactory.h
+++ b/Engine/Plugins/AudioInsights/Source/AudioInsights/Public/Views/AudioEventLogDashboardViewFactory.h
@@ -94,6 +94,8 @@ namespace UE::Audio::Insights
 #endif // WITH_EDITOR
 
 	private:
+		void OnAnalysisStarting(const double Timestamp);
+
 		void UpdateCustomEventLogFilters(const IDashboardDataViewEntry& Entry);
 		TSharedRef<FAudioEventLogFilter> CreateNewEventFilterType(const TSharedPtr<FFilterCategory>& FilterCategory, const FString& EventType);
 		bool PassesEventTypeFilter(const FEventLogFilterID EventLogID) const;
@@ -126,6 +128,8 @@ namespace UE::Audio::Insights
 
 		void OnUpdateAutoScroll();
 
+		double GetTimestampRelativeToAnalysisStart(const double Timestamp) const;
+
 		TSharedPtr<FAudioEventLogTraceProvider> AudioEventLogTraceProvider;
 
 		TSharedPtr<SWidget> DashboardWidget;
@@ -144,6 +148,7 @@ namespace UE::Audio::Insights
 		const TMap<FString, TSet<FString>> InitEventTypeFilters;
 		TMap<FString, FAudioEventLogCustomEvents> CachedCustomEventSettings;
 
+		double BeginTimestamp = 0.0;
 		bool bAutoScroll = true;
 
 #if WITH_EDITOR
