[Audio Insights] Fix for flickering HasActivity indicator and meters on submixes and buses
* Asset-based additions to both Submix and Audio Bus dashboards were copying a pointer to the same dashboard entry across different audio devices. This means submixes and buses that were being updated across different devices at the same time were all feeding into the same dashboard entries.
* We now do deep copies when adding these new dashboard entries to ensure we are not sharing instances.
* Copy assignment overload added to dashboard entries for consistency.

diff --git a/Engine/Plugins/AudioInsights/Source/AudioInsights/Private/Providers/AudioBusTraceProvider.cpp b/Engine/Plugins/AudioInsights/Source/AudioInsights/Private/Providers/AudioBusTraceProvider.cpp
--- a/Engine/Plugins/AudioInsights/Source/AudioInsights/Private/Providers/AudioBusTraceProvider.cpp
+++ b/Engine/Plugins/AudioInsights/Source/AudioInsights/Private/Providers/AudioBusTraceProvider.cpp
@@ -311,7 +311,8 @@ namespace UE::Audio::Insights
 						{
 							if (!Entry.IsValid())
 							{
-								Entry = InAudioBusDashboardEntry;
+								// We need to take a deep copy to avoid sharing the same entry instance across different audio devices
+								Entry = MakeShared<FAudioBusDashboardEntry>(*InAudioBusDashboardEntry);
 							}
 						});
 					}
@@ -368,7 +369,7 @@ namespace UE::Audio::Insights
 					{
 						if (!Entry.IsValid())
 						{
-							Entry = AudioBusDashboardEntry;
+							Entry = MakeShared<FAudioBusDashboardEntry>(*AudioBusDashboardEntry);;
 						}
 					});
 				}
diff --git a/Engine/Plugins/AudioInsights/Source/AudioInsights/Private/Providers/SubmixTraceProvider.cpp b/Engine/Plugins/AudioInsights/Source/AudioInsights/Private/Providers/SubmixTraceProvider.cpp
--- a/Engine/Plugins/AudioInsights/Source/AudioInsights/Private/Providers/SubmixTraceProvider.cpp
+++ b/Engine/Plugins/AudioInsights/Source/AudioInsights/Private/Providers/SubmixTraceProvider.cpp
@@ -293,7 +293,8 @@ namespace UE::Audio::Insights
 						{
 							if (!Entry.IsValid())
 							{
-								Entry = InSubmixDashboardEntry;
+								// We need to take a deep copy to avoid sharing the same entry instance across different audio devices
+								Entry = MakeShared<FSubmixDashboardEntry>(*InSubmixDashboardEntry);
 							}
 						});
 					}
@@ -350,7 +351,7 @@ namespace UE::Audio::Insights
 					{
 						if (!Entry.IsValid())
 						{
-							Entry = SubmixDashboardEntry;
+							Entry = MakeShared<FSubmixDashboardEntry>(*SubmixDashboardEntry);;
 						}
 					});
 				}
diff --git a/Engine/Plugins/AudioInsights/Source/AudioInsights/Public/Messages/AudioBusTraceMessages.h b/Engine/Plugins/AudioInsights/Source/AudioInsights/Public/Messages/AudioBusTraceMessages.h
--- a/Engine/Plugins/AudioInsights/Source/AudioInsights/Public/Messages/AudioBusTraceMessages.h
+++ b/Engine/Plugins/AudioInsights/Source/AudioInsights/Public/Messages/AudioBusTraceMessages.h
@@ -121,8 +121,36 @@ namespace UE::Audio::Insights
 
 		}
 
+		// Copy constructor - make sure we take a deep copy of FAudioMeterInfo
+		FAudioBusDashboardEntry(FAudioBusDashboardEntry& Other)
+			: AudioBusId(Other.AudioBusId)
+			, Timestamp(Other.Timestamp)
+			, AudioMeterInfo(MakeShared<FAudioMeterInfo>(*Other.AudioMeterInfo))
+			, Name(Other.Name)
+			, bHasActivity(Other.bHasActivity)
+			, AudioBusType(Other.AudioBusType)
+		{
+		}
+
 		virtual ~FAudioBusDashboardEntry() = default;
 
+		FAudioBusDashboardEntry& operator=(const FAudioBusDashboardEntry& Other)
+		{
+			if (this == &Other)
+			{
+				return *this;
+			}
+
+			AudioBusId = Other.AudioBusId;
+			Timestamp = Other.Timestamp;
+			AudioMeterInfo = MakeShared<FAudioMeterInfo>(*Other.AudioMeterInfo);
+			Name = Other.Name;
+			bHasActivity = Other.bHasActivity;
+			AudioBusType = Other.AudioBusType;
+
+			return *this;
+		}
+
 		virtual FText GetDisplayName() const override { return FText::FromString(AudioBusType == EAudioBusType::AssetBased ? FSoftObjectPath(Name).GetAssetName() : Name); }
 		virtual const UObject* GetObject() const override { return FSoftObjectPath(Name).ResolveObject(); }
 		virtual UObject* GetObject() override { return FSoftObjectPath(Name).ResolveObject(); }
diff --git a/Engine/Plugins/AudioInsights/Source/AudioInsights/Public/Messages/SubmixTraceMessages.h b/Engine/Plugins/AudioInsights/Source/AudioInsights/Public/Messages/SubmixTraceMessages.h
--- a/Engine/Plugins/AudioInsights/Source/AudioInsights/Public/Messages/SubmixTraceMessages.h
+++ b/Engine/Plugins/AudioInsights/Source/AudioInsights/Public/Messages/SubmixTraceMessages.h
@@ -121,8 +121,36 @@ namespace UE::Audio::Insights
 
 		}
 
+		// Copy constructor - make sure we take a deep copy of FAudioMeterInfo
+		FSubmixDashboardEntry(FSubmixDashboardEntry& Other)
+			: SubmixId(Other.SubmixId)
+			, Timestamp(Other.Timestamp)
+			, AudioMeterInfo(MakeShared<FAudioMeterInfo>(*Other.AudioMeterInfo))
+			, Name(Other.Name)
+			, bHasActivity(Other.bHasActivity)
+			, bIsMainSubmix(Other.bIsMainSubmix)
+		{
+		}
+
 		virtual ~FSubmixDashboardEntry() = default;
 
+		FSubmixDashboardEntry& operator=(const FSubmixDashboardEntry& Other)
+		{
+			if (this == &Other)
+			{
+				return *this;
+			}
+
+			SubmixId = Other.SubmixId;
+			Timestamp = Other.Timestamp;
+			AudioMeterInfo = MakeShared<FAudioMeterInfo>(*Other.AudioMeterInfo);
+			Name = Other.Name;
+			bHasActivity = Other.bHasActivity;
+			bIsMainSubmix = Other.bIsMainSubmix;
+
+			return *this;
+		}
+
 		virtual FText GetDisplayName() const override { return FText::FromString(FSoftObjectPath(Name).GetAssetName()); }
 		virtual const UObject* GetObject() const override { return FSoftObjectPath(Name).ResolveObject(); }
 		virtual UObject* GetObject() override { return FSoftObjectPath(Name).ResolveObject(); }
