[Audio Insights] Stop search bar stretching vertically in Event Log when many filters are applied - hide filter bar buttons in the UI

diff --git a/Engine/Plugins/AudioInsights/Source/AudioInsights/Private/Views/AudioEventLogDashboardViewFactory.cpp b/Engine/Plugins/AudioInsights/Source/AudioInsights/Private/Views/AudioEventLogDashboardViewFactory.cpp
--- a/Engine/Plugins/AudioInsights/Source/AudioInsights/Private/Views/AudioEventLogDashboardViewFactory.cpp
+++ b/Engine/Plugins/AudioInsights/Source/AudioInsights/Private/Views/AudioEventLogDashboardViewFactory.cpp
@@ -574,6 +574,7 @@ namespace UE::Audio::Insights
 			SAssignNew(FilterBar, SAudioFilterBar<FEventLogFilterID>)
 			.CustomFilters(Filters)
 			.UseSectionsForCategories(false)
+			.Visibility(EVisibility::Collapsed) // Filter bar buttons intentionally hidden in UI
 			.OnFilterChanged_Lambda([this, Filters]()
 			{
 				UpdateFilterReason = EProcessReason::FilterUpdated;
diff --git a/Engine/Plugins/AudioInsights/Source/AudioInsights/Private/Views/TableDashboardViewFactory.cpp b/Engine/Plugins/AudioInsights/Source/AudioInsights/Private/Views/TableDashboardViewFactory.cpp
--- a/Engine/Plugins/AudioInsights/Source/AudioInsights/Private/Views/TableDashboardViewFactory.cpp
+++ b/Engine/Plugins/AudioInsights/Source/AudioInsights/Private/Views/TableDashboardViewFactory.cpp
@@ -270,15 +270,23 @@ namespace UE::Audio::Insights
 				+ SHorizontalBox::Slot()
 				.AutoWidth()
 				[
-					SAssignNew(SearchBoxWidget, SSearchBox)
-					.SelectAllTextWhenFocused(true)
-					.HintText(LOCTEXT("TableDashboardView_SearchBoxHintText", "Filter"))
-					.MinDesiredWidth(100)
-					.OnTextChanged(this, &FTraceTableDashboardViewFactory::SetSearchBoxFilterText)
+					SNew(SVerticalBox)
+
+					+ SVerticalBox::Slot()
+					.MaxHeight(FSlateStyle::Get().GetSearchBoxMaxHeight())
+					.Padding(0.0f, 4.0f, 0.0f, 4.0f)
+					[
+						SAssignNew(SearchBoxWidget, SSearchBox)
+						.SelectAllTextWhenFocused(true)
+						.HintText(LOCTEXT("TableDashboardView_SearchBoxHintText", "Filter"))
+						.MinDesiredWidth(100)
+						.OnTextChanged(this, &FTraceTableDashboardViewFactory::SetSearchBoxFilterText)
+					]
 				]
 
 				// Optional : Active filters area
 				+ SHorizontalBox::Slot()
+				.AutoWidth()
 				[
 					CreateFilterBarWidget()
 				]
diff --git a/Engine/Plugins/AudioInsights/Source/AudioInsights/Private/Views/TreeDashboardViewFactory.cpp b/Engine/Plugins/AudioInsights/Source/AudioInsights/Private/Views/TreeDashboardViewFactory.cpp
--- a/Engine/Plugins/AudioInsights/Source/AudioInsights/Private/Views/TreeDashboardViewFactory.cpp
+++ b/Engine/Plugins/AudioInsights/Source/AudioInsights/Private/Views/TreeDashboardViewFactory.cpp
@@ -103,7 +103,7 @@ namespace UE::Audio::Insights
 				[
 					SNew(SVerticalBox)
 					+ SVerticalBox::Slot()
-					.MaxHeight(22.0f)
+					.MaxHeight(FSlateStyle::Get().GetSearchBoxMaxHeight())
 					.Padding(0.0f, 4.0f, 0.0f, 6.0f)
 					[
 						SAssignNew(SearchBoxWidget, SSearchBox)
diff --git a/Engine/Plugins/AudioInsights/Source/AudioInsights/Public/AudioInsightsStyle.h b/Engine/Plugins/AudioInsights/Source/AudioInsights/Public/AudioInsightsStyle.h
--- a/Engine/Plugins/AudioInsights/Source/AudioInsights/Public/AudioInsightsStyle.h
+++ b/Engine/Plugins/AudioInsights/Source/AudioInsights/Public/AudioInsightsStyle.h
@@ -126,6 +126,11 @@ namespace UE::Audio::Insights
 			return FText::Format(LOCTEXT("MemoryInMegabytesFormat", "{0}mb"), FText::AsNumber(InMemoryMB, GetMemoryFormat()));
 		}
 
+		float GetSearchBoxMaxHeight() const
+		{
+			return 22.0f;
+		}
+
 		FSlateIcon CreateIcon(FName InName)  const
 		{
 			return { GetStyleName(), InName };
