[Audio Insights] Fix for assert triggering in VirtualLoopsTraceProvider when moving time marker inside standalone - make sure we don't create array of uninitialized messages and then try to use them

diff --git a/Engine/Plugins/AudioInsights/Source/AudioInsights/Private/Providers/SoundTraceProvider.cpp b/Engine/Plugins/AudioInsights/Source/AudioInsights/Private/Providers/SoundTraceProvider.cpp
--- a/Engine/Plugins/AudioInsights/Source/AudioInsights/Private/Providers/SoundTraceProvider.cpp
+++ b/Engine/Plugins/AudioInsights/Source/AudioInsights/Private/Providers/SoundTraceProvider.cpp
@@ -1568,19 +1568,19 @@ namespace UE::Audio::Insights
 			// Collect messages (volume, pitch, etc.) from active sounds (based on active sounds's PlayOrder)
 			struct CachedEntryInfo
 			{
-				FSoundPriorityMessage PriorityMessage;
-				FSoundDistanceMessage DistanceMessage;
-				FSoundDistanceAttenuationMessage DistanceAttenuationMessage;
-				FSoundLPFFreqMessage LPFFreqMessage;
-				FSoundHPFFreqMessage HPFFreqMessage;
-				FSoundEnvelopeMessage AmplitudeMessage;
-				FSoundVolumeMessage VolumeMessage;
-				FSoundPitchMessage PitchMessage;
-				FSoundRelativeRenderCostMessage RelativeRenderCostMessage;
+				TOptional<FSoundPriorityMessage> PriorityMessage;
+				TOptional<FSoundDistanceMessage> DistanceMessage;
+				TOptional<FSoundDistanceAttenuationMessage> DistanceAttenuationMessage;
+				TOptional<FSoundLPFFreqMessage> LPFFreqMessage;
+				TOptional<FSoundHPFFreqMessage> HPFFreqMessage;
+				TOptional<FSoundEnvelopeMessage> AmplitudeMessage;
+				TOptional<FSoundVolumeMessage> VolumeMessage;
+				TOptional<FSoundPitchMessage> PitchMessage;
+				TOptional<FSoundRelativeRenderCostMessage> RelativeRenderCostMessage;
 			};
 
 			TArray<CachedEntryInfo> CachedEntryInfos;
-			CachedEntryInfos.SetNumUninitialized(ActiveSoundPlayOrderArray.Num());
+			CachedEntryInfos.SetNum(ActiveSoundPlayOrderArray.Num());
 
 			// Using ParallelFor to speed-up the cached messages retrieval, using a traditional for loop is unacceptably slower, specially in large traces.
 			ParallelFor(ActiveSoundPlayOrderArray.Num(),
@@ -1655,15 +1655,20 @@ namespace UE::Audio::Insights
 			// Update the device entries with the collected info
 			for (const CachedEntryInfo& CachedEntryInfo : CachedEntryInfos)
 			{
+				if (!CachedEntryInfo.PriorityMessage.IsSet())
+				{
+					continue;
+				}
+
 				TSharedPtr<FSoundDashboardEntry> ActiveSoundEntry = nullptr;
-				GetActiveSoundEntryFromIDs(CachedEntryInfo.PriorityMessage.ActiveSoundPlayOrder, ActiveSoundEntry);
+				GetActiveSoundEntryFromIDs(CachedEntryInfo.PriorityMessage.GetValue().ActiveSoundPlayOrder, ActiveSoundEntry);
 				
 				if (!ActiveSoundEntry.IsValid())
 				{
 					continue;
 				}
 
-				FSoundDashboardEntry* SoundEntry = FSoundTraceProviderPrivate::FindSoundEntryForActiveSound(CachedEntryInfo.PriorityMessage.PlayOrder, &ActiveSoundEntry);
+				FSoundDashboardEntry* SoundEntry = FSoundTraceProviderPrivate::FindSoundEntryForActiveSound(CachedEntryInfo.PriorityMessage.GetValue().PlayOrder, &ActiveSoundEntry);
 				if (SoundEntry == nullptr)
 				{
 					continue;
@@ -1673,17 +1678,49 @@ namespace UE::Audio::Insights
 				// Changing our data buffers to single points helps avoid issues with ordering data while moving the timemarker back and forth
 				SoundEntry->SwitchBuffersToSingleValues();
 
-				SoundEntry->PriorityDataPoints.Push({ CachedEntryInfo.PriorityMessage.Timestamp, CachedEntryInfo.PriorityMessage.Priority });
-				SoundEntry->DistanceDataPoints.Push({ CachedEntryInfo.DistanceMessage.Timestamp, CachedEntryInfo.DistanceMessage.Distance });
-				SoundEntry->DistanceAttenuationDataPoints.Push({ CachedEntryInfo.DistanceAttenuationMessage.Timestamp, CachedEntryInfo.DistanceAttenuationMessage.DistanceAttenuation });
-				SoundEntry->LPFFreqDataPoints.Push({ CachedEntryInfo.LPFFreqMessage.Timestamp, CachedEntryInfo.LPFFreqMessage.LPFFrequency });
-				SoundEntry->HPFFreqDataPoints.Push({ CachedEntryInfo.HPFFreqMessage.Timestamp, CachedEntryInfo.HPFFreqMessage.HPFFrequency });
-				SoundEntry->AmplitudeDataPoints.Push({ CachedEntryInfo.AmplitudeMessage.Timestamp, CachedEntryInfo.AmplitudeMessage.Envelope });
-				SoundEntry->VolumeDataPoints.Push({ CachedEntryInfo.VolumeMessage.Timestamp, CachedEntryInfo.VolumeMessage.Volume });
-				SoundEntry->PitchDataPoints.Push({ CachedEntryInfo.PitchMessage.Timestamp, CachedEntryInfo.PitchMessage.Pitch });
-				SoundEntry->RelativeRenderCostDataPoints.Push({ CachedEntryInfo.RelativeRenderCostMessage.Timestamp, CachedEntryInfo.RelativeRenderCostMessage.RelativeRenderCost });
-
-				SoundEntry->Timestamp = CachedEntryInfo.PriorityMessage.Timestamp;
+				SoundEntry->PriorityDataPoints.Push({ CachedEntryInfo.PriorityMessage.GetValue().Timestamp, CachedEntryInfo.PriorityMessage.GetValue().Priority });
+
+				if (CachedEntryInfo.DistanceMessage.IsSet())
+				{
+					SoundEntry->DistanceDataPoints.Push({ CachedEntryInfo.DistanceMessage.GetValue().Timestamp, CachedEntryInfo.DistanceMessage.GetValue().Distance });
+				}
+
+				if (CachedEntryInfo.DistanceAttenuationMessage.IsSet())
+				{
+					SoundEntry->DistanceAttenuationDataPoints.Push({ CachedEntryInfo.DistanceAttenuationMessage.GetValue().Timestamp, CachedEntryInfo.DistanceAttenuationMessage.GetValue().DistanceAttenuation });
+				}
+
+				if (CachedEntryInfo.LPFFreqMessage.IsSet())
+				{
+					SoundEntry->LPFFreqDataPoints.Push({ CachedEntryInfo.LPFFreqMessage.GetValue().Timestamp, CachedEntryInfo.LPFFreqMessage.GetValue().LPFFrequency });
+				}
+
+				if (CachedEntryInfo.HPFFreqMessage.IsSet())
+				{
+					SoundEntry->HPFFreqDataPoints.Push({ CachedEntryInfo.HPFFreqMessage.GetValue().Timestamp, CachedEntryInfo.HPFFreqMessage.GetValue().HPFFrequency });
+				}
+
+				if (CachedEntryInfo.AmplitudeMessage.IsSet())
+				{
+					SoundEntry->AmplitudeDataPoints.Push({ CachedEntryInfo.AmplitudeMessage.GetValue().Timestamp, CachedEntryInfo.AmplitudeMessage.GetValue().Envelope });
+				}
+
+				if (CachedEntryInfo.VolumeMessage.IsSet())
+				{
+					SoundEntry->VolumeDataPoints.Push({ CachedEntryInfo.VolumeMessage.GetValue().Timestamp, CachedEntryInfo.VolumeMessage.GetValue().Volume });
+				}
+
+				if (CachedEntryInfo.PitchMessage.IsSet())
+				{
+					SoundEntry->PitchDataPoints.Push({ CachedEntryInfo.PitchMessage.GetValue().Timestamp, CachedEntryInfo.PitchMessage.GetValue().Pitch });
+				}
+
+				if (CachedEntryInfo.RelativeRenderCostMessage.IsSet())
+				{
+					SoundEntry->RelativeRenderCostDataPoints.Push({ CachedEntryInfo.RelativeRenderCostMessage.GetValue().Timestamp, CachedEntryInfo.RelativeRenderCostMessage.GetValue().RelativeRenderCost });
+				}
+
+				SoundEntry->Timestamp = CachedEntryInfo.PriorityMessage.GetValue().Timestamp;
 			}
 		}
 	}
diff --git a/Engine/Plugins/AudioInsights/Source/AudioInsights/Private/Providers/VirtualLoopTraceProvider.cpp b/Engine/Plugins/AudioInsights/Source/AudioInsights/Private/Providers/VirtualLoopTraceProvider.cpp
--- a/Engine/Plugins/AudioInsights/Source/AudioInsights/Private/Providers/VirtualLoopTraceProvider.cpp
+++ b/Engine/Plugins/AudioInsights/Source/AudioInsights/Private/Providers/VirtualLoopTraceProvider.cpp
@@ -146,7 +146,7 @@ namespace UE::Audio::Insights
 		DeviceData->GenerateKeyArray(PlayOrderArray);
 
 		TArray<CachedEntryInfo> CachedEntryInfos;
-		CachedEntryInfos.SetNumUninitialized(PlayOrderArray.Num());
+		CachedEntryInfos.SetNum(PlayOrderArray.Num());
 
 #if WITH_EDITOR
 		const FAudioInsightsCacheManager& CacheManager = FAudioInsightsModule::GetChecked().GetCacheManager();
