[Audio Insights] Refresh Audio Insights dashboards after undo

diff --git a/Engine/Plugins/AudioInsights/Source/AudioInsightsEditor/Private/AudioInsightsEditorSettings.cpp b/Engine/Plugins/AudioInsights/Source/AudioInsightsEditor/Private/AudioInsightsEditorSettings.cpp
--- a/Engine/Plugins/AudioInsights/Source/AudioInsightsEditor/Private/AudioInsightsEditorSettings.cpp
+++ b/Engine/Plugins/AudioInsights/Source/AudioInsightsEditor/Private/AudioInsightsEditorSettings.cpp
@@ -112,6 +112,14 @@ void UAudioInsightsEditorSettings::PostEditChangeChainProperty(FPropertyChangedC
 	}
 }
 
+void UAudioInsightsEditorSettings::PostEditUndo()
+{
+	Super::PostEditUndo();
+
+	FAudioEventLogSettings::OnReadSettings.Broadcast(EventLogSettings);
+	FSoundDashboardSettings::OnReadSettings.Broadcast(SoundDashboardSettings);
+}
+
 bool UAudioInsightsEditorSettings::VerifyEventLogInput(const TDoubleLinkedList<FProperty*>::TDoubleLinkedListNode& ActiveMemberNode, FPropertyChangedChainEvent& PropertyChangedEvent)
 {
 	check(ActiveMemberNode.GetValue() != nullptr);
diff --git a/Engine/Plugins/AudioInsights/Source/AudioInsightsEditor/Public/AudioInsightsEditorSettings.h b/Engine/Plugins/AudioInsights/Source/AudioInsightsEditor/Public/AudioInsightsEditorSettings.h
--- a/Engine/Plugins/AudioInsights/Source/AudioInsightsEditor/Public/AudioInsightsEditorSettings.h
+++ b/Engine/Plugins/AudioInsights/Source/AudioInsightsEditor/Public/AudioInsightsEditorSettings.h
@@ -27,6 +27,7 @@ public:
 	UE_API virtual void PostInitProperties() override;
 	UE_API virtual void BeginDestroy() override;
 	UE_API virtual void PostEditChangeChainProperty(struct FPropertyChangedChainEvent& PropertyChangedEvent) override;
+	UE_API virtual void PostEditUndo() override;
 
 	/** Whether to automatically set the first PIE client in Audio Insights World Filter. */
 	UPROPERTY(Config, EditAnywhere, Category="World Filter")
