[Audio Insights] Visualize the limits of the cache in the Event Log
* Show icon when an event is in the last chunk in the cache and has been marked for deletion
* Changed wording of option "When last in cache" to "When marked for deletion"
* When the cache pauses during "When marked for deletion" mode, do not automatically restart caching when the user inspects a different row, instead wait for the user to clear inspected rows before starting again

diff --git a/Engine/Plugins/AudioInsights/Source/AudioInsights/Private/AudioEventLogEditorCommands.cpp b/Engine/Plugins/AudioInsights/Source/AudioInsights/Private/AudioEventLogEditorCommands.cpp
--- a/Engine/Plugins/AudioInsights/Source/AudioInsights/Private/AudioEventLogEditorCommands.cpp
+++ b/Engine/Plugins/AudioInsights/Source/AudioInsights/Private/AudioEventLogEditorCommands.cpp
@@ -20,7 +20,7 @@ namespace UE::Audio::Insights
 		UI_COMMAND(Edit, "Edit", "Opens the selected sound for edit.", EUserInterfaceActionType::Button, FInputChord(EKeys::E, EModifierKey::Control));
 		UI_COMMAND(ResetInspectTimestamp, "Reset Inspect Timestamp", "Resumes gathering the latest audio data and auto-scrolling in the Event Log if previously paused.", EUserInterfaceActionType::Button, FInputChord(EKeys::Escape));
 
-		UI_COMMAND(AutoStopCachingWhenLastInCache, "When last in cache", "Automatically stops caching new data when inspecting an event that is about to be deleted from the cache.", EUserInterfaceActionType::RadioButton, FInputChord(EKeys::F1, EModifierKey::Control));
+		UI_COMMAND(AutoStopCachingWhenLastInCache, "When marked for deletion", "Automatically stops caching new data when inspecting an event that has been marked for deletion from the cache.", EUserInterfaceActionType::RadioButton, FInputChord(EKeys::F1, EModifierKey::Control));
 		UI_COMMAND(AutoStopCachingOnInspect, "On inspect", "Automatically stops caching new data when inspecting an event.", EUserInterfaceActionType::RadioButton, FInputChord(EKeys::F2, EModifierKey::Control));
 		UI_COMMAND(AutoStopCachingDisabled, "Never", "Disable automatically stopping the cache. Inspected events may be deleted from the cache during inspection.", EUserInterfaceActionType::RadioButton, FInputChord(EKeys::F3, EModifierKey::Control));
 	}
diff --git a/Engine/Plugins/AudioInsights/Source/AudioInsights/Private/Cache/AudioInsightsCacheManager.cpp b/Engine/Plugins/AudioInsights/Source/AudioInsights/Private/Cache/AudioInsightsCacheManager.cpp
--- a/Engine/Plugins/AudioInsights/Source/AudioInsights/Private/Cache/AudioInsightsCacheManager.cpp
+++ b/Engine/Plugins/AudioInsights/Source/AudioInsights/Private/Cache/AudioInsightsCacheManager.cpp
@@ -117,6 +117,28 @@ namespace UE::Audio::Insights
 		return TOptional<uint32>();
 	}
 
+	const FAudioCachedMessageChunk* const FAudioInsightsCacheManager::GetChunk(const uint32 NumChunksFromStart) const
+	{
+		ensure(NumChunksFromStart < Cache.Capacity());
+		if (NumChunksFromStart >= Cache.Capacity())
+		{
+			return nullptr;
+		}
+
+		uint32 ReadIndex = StartIndex;
+		for (uint32 Index = 0u; Index < NumChunksFromStart; ++Index)
+		{
+			ReadIndex = Cache.GetNextIndex(ReadIndex);
+		}
+
+		if (!Cache[ReadIndex]->HasAnyData())
+		{
+			return nullptr;
+		}
+
+		return Cache[ReadIndex].Get();
+	}
+
 	void FAudioInsightsCacheManager::CreateCache()
 	{
 		for (uint32 Index = 0u; Index < Cache.Capacity(); ++Index)
diff --git a/Engine/Plugins/AudioInsights/Source/AudioInsights/Private/Providers/AudioEventLogTraceProvider.cpp b/Engine/Plugins/AudioInsights/Source/AudioInsights/Private/Providers/AudioEventLogTraceProvider.cpp
--- a/Engine/Plugins/AudioInsights/Source/AudioInsights/Private/Providers/AudioEventLogTraceProvider.cpp
+++ b/Engine/Plugins/AudioInsights/Source/AudioInsights/Private/Providers/AudioEventLogTraceProvider.cpp
@@ -170,8 +170,16 @@ namespace UE::Audio::Insights
 			return false;
 		}
 
+		const double AboutToBeDeletedTimestamp = GetAboutToBeDeletedFromCacheTimeThreshold();
+
 		FAudioInsightsCacheManager& CacheManager = IAudioInsightsModule::GetChecked().GetCacheManager();
-		const uint32 HalfNumChunks = CacheManager.GetNumChunks() > 1 ? CacheManager.GetNumChunks() / 2 : 1;
+
+		const uint32 HalfNumChunks = CacheManager.GetNumChunks() / 2;
+		ensure(HalfNumChunks > 0u);
+		if (HalfNumChunks == 0u)
+		{
+			return false;
+		}
 
 		bool bAnyEntriesRemoved = false;
 		for (auto& [DeviceID, MessageIDToDashboardEntryMap] : DeviceDataMap)
@@ -185,6 +193,11 @@ namespace UE::Audio::Insights
 				{
 					MessagesToRemove.Add(MessageID);
 				}
+				else if (DashboardEntry->Timestamp <= AboutToBeDeletedTimestamp)
+				{
+					DashboardEntry->CachedState = EAudioEventCacheState::NextToBeDeleted;
+					DashboardEntry->bCacheStatusIsDirty = true; // Mark as dirty so UI will update deleted rows and cache status icon at the same time
+				}
 				else
 				{
 					break;
@@ -200,7 +213,6 @@ namespace UE::Audio::Insights
 		}
 
 		bCacheChunksUpdated = false;
-
 		return bAnyEntriesRemoved;
 #else
 		return false;
@@ -262,6 +274,34 @@ namespace UE::Audio::Insights
 		CacheStartTimestamp = NewCacheStartTimestamp;
 		bCacheChunksUpdated = true;
 	}
+
+	double FAudioEventLogTraceProvider::GetAboutToBeDeletedFromCacheTimeThreshold() const
+	{
+		FAudioInsightsCacheManager& CacheManager = IAudioInsightsModule::GetChecked().GetCacheManager();
+
+		constexpr int32 NumChunksDeletedAtOnce = 1;
+
+		const int32 NumUsedChunks = CacheManager.GetNumUsedChunks();
+		const int32 MaxNumChunks = CacheManager.GetNumChunks();
+
+		ensure(MaxNumChunks > NumChunksDeletedAtOnce);
+		if (MaxNumChunks <= NumChunksDeletedAtOnce)
+		{
+			return TNumericLimits<double>::Lowest();
+		}
+
+		// We don't mark anything as leaving the cache if the cache isn't close to overwriting chunks yet
+		if (NumUsedChunks < (MaxNumChunks - NumChunksDeletedAtOnce))
+		{
+			return TNumericLimits<double>::Lowest();
+		}
+
+		const int32 BeingDeletedChunkIndex = NumUsedChunks - MaxNumChunks + NumChunksDeletedAtOnce - 1;
+
+		const FAudioCachedMessageChunk* const NextToBeDeletedChunk = BeingDeletedChunkIndex >= 0 ? CacheManager.GetChunk(BeingDeletedChunkIndex) : nullptr;
+
+		return NextToBeDeletedChunk ? NextToBeDeletedChunk->GetChunkTimeRangeEnd() : TNumericLimits<double>::Lowest();
+	}
 #endif // WITH_EDITOR
 
 	void FAudioEventLogTraceProvider::ClearActiveAudioDeviceEntries()
diff --git a/Engine/Plugins/AudioInsights/Source/AudioInsights/Private/Providers/AudioEventLogTraceProvider.h b/Engine/Plugins/AudioInsights/Source/AudioInsights/Private/Providers/AudioEventLogTraceProvider.h
--- a/Engine/Plugins/AudioInsights/Source/AudioInsights/Private/Providers/AudioEventLogTraceProvider.h
+++ b/Engine/Plugins/AudioInsights/Source/AudioInsights/Private/Providers/AudioEventLogTraceProvider.h
@@ -31,6 +31,7 @@ namespace UE::Audio::Insights
 
 #if WITH_EDITOR
 		void OnCacheChunkOverwritten(const double NewCacheStartTimestamp);
+		double GetAboutToBeDeletedFromCacheTimeThreshold() const;
 #endif // WITH_EDITOR
 
 		FAudioEventLogMessages TraceMessages;
diff --git a/Engine/Plugins/AudioInsights/Source/AudioInsights/Private/Views/AudioEventLogDashboardViewFactory.cpp b/Engine/Plugins/AudioInsights/Source/AudioInsights/Private/Views/AudioEventLogDashboardViewFactory.cpp
--- a/Engine/Plugins/AudioInsights/Source/AudioInsights/Private/Views/AudioEventLogDashboardViewFactory.cpp
+++ b/Engine/Plugins/AudioInsights/Source/AudioInsights/Private/Views/AudioEventLogDashboardViewFactory.cpp
@@ -26,6 +26,11 @@ namespace UE::Audio::Insights
 			return static_cast<const FAudioEventLogDashboardEntry&>(InData);
 		};
 
+		FAudioEventLogDashboardEntry& CastEntry(IDashboardDataViewEntry& InData)
+		{
+			return static_cast<FAudioEventLogDashboardEntry&>(InData);
+		};
+
 		bool CanClearLog()
 		{
 #if WITH_EDITOR
@@ -61,7 +66,7 @@ namespace UE::Audio::Insights
 		const FString CategoryMessages = LOCTEXT("EventLogCategory_Messages", "Messages").ToString();
 		const FString CategoryCustom = LOCTEXT("EventLogCategory_Custom", "Custom").ToString();
 
-
+		const FName CacheStatusColumnID = "CacheStatus";
 		const FName MessageIDColumnID = "MessageID";
 		const FName TimestampColumnID = "Timestamp";
 		const FName PlayOrderColumnID = "PlayOrder";
@@ -345,6 +350,8 @@ namespace UE::Audio::Insights
 
 	void FAudioEventLogDashboardViewFactory::ProcessEntries(FTraceTableDashboardViewFactory::EProcessReason Reason)
 	{
+		using namespace AudioEventLogPrivate;
+
 		const FString FilterString = GetSearchFilterText().ToString();
 		FTraceTableDashboardViewFactory::FilterEntries<FAudioEventLogTraceProvider>([this, &FilterString](const IDashboardDataViewEntry& Entry)
 		{
@@ -372,9 +379,27 @@ namespace UE::Audio::Insights
 		});
 
 #if WITH_EDITOR
-		if (Reason == FTraceTableDashboardViewFactory::EProcessReason::FilterUpdated)
+		if (Reason == EProcessReason::EntriesUpdated && FilteredEntriesListView.IsValid())
 		{
-			FAudioEventLogSettings::OnRequestWriteSettings.Broadcast();
+			// Reset bCacheStatusIsDirty flag to push icon update to the UI
+			for (const TSharedPtr<IDashboardDataViewEntry>& Entry : FilteredEntriesListView->GetItems())
+			{
+				if (!Entry.IsValid())
+				{
+					continue;
+				}
+
+				FAudioEventLogDashboardEntry& EventLogEntry = CastEntry(*Entry);
+				if (EventLogEntry.CachedState == EAudioEventCacheState::NextToBeDeleted)
+				{
+					EventLogEntry.bCacheStatusIsDirty = false;
+				}
+				else
+				{
+					// Only the earliest items in the list will be marked as NextToBeDeleted - no need to process anymore items
+					break;
+				}
+			}
 		}
 #endif // WITH_EDITOR
 	}
@@ -393,6 +418,50 @@ namespace UE::Audio::Insights
 
 			return TMap<FName, FTraceTableDashboardViewFactory::FColumnData>
 			{
+#if WITH_EDITOR
+				{
+					CacheStatusColumnID,
+					{
+						FText::GetEmpty(),
+						[](const IDashboardDataViewEntry& InData) { return FText::GetEmpty(); }, /* GetDisplayValue */
+						[](const IDashboardDataViewEntry& InData) -> FName
+						{
+							return FName("AudioInsights.Icon.EventLog.MarkedForDelete");
+						},
+						!VisibleColumns.GetIsVisible(CacheStatusColumnID) /* bDefaultHidden */,
+						0.075f /* FillWidth */,
+						HAlign_Center /* Alignment */,
+						[this](const IDashboardDataViewEntry& InData) -> FLinearColor /* GetIconColor */
+						{
+							const FAudioEventLogDashboardEntry& Entry = CastEntry(InData);
+
+							switch (Entry.CachedState)
+							{
+								case EAudioEventCacheState::NextToBeDeleted:
+									return Entry.bCacheStatusIsDirty ? FLinearColor::Transparent : FLinearColor::White;
+								default:
+									break;
+							}
+
+							return FLinearColor::Transparent;
+						},
+						[](const IDashboardDataViewEntry& InData) -> FText /* GetIconTooltip */
+						{
+							const FAudioEventLogDashboardEntry& Entry = CastEntry(InData);
+
+							switch (Entry.CachedState)
+							{
+								case EAudioEventCacheState::NextToBeDeleted:
+									return LOCTEXT("AudioEventLog_NextToBeDeleted", "This event has been marked for deletion the next time the cache is full.");
+								default:
+									break;
+							}
+
+							return FText::GetEmpty();
+						}
+					}
+				},
+#endif // WITH_EDITOR
 				{
 					MessageIDColumnID,
 					{
@@ -447,7 +516,7 @@ namespace UE::Audio::Insights
 					{
 						LOCTEXT("AudioEventLog_AssetColumnDisplayName", "Asset"),
 						[](const IDashboardDataViewEntry& InData) { return CastEntry(InData).GetDisplayName(); },
-						[this](const IDashboardDataViewEntry& InData) -> FName
+						[](const IDashboardDataViewEntry& InData) -> FName
 						{
 							const FAudioEventLogDashboardEntry& Entry = CastEntry(InData);
 
@@ -479,7 +548,7 @@ namespace UE::Audio::Insights
 					{
 						LOCTEXT("AudioEventLog_ActorColumnDisplayName", "Actor"),
 						[](const IDashboardDataViewEntry& InData) { return FText::FromString(CastEntry(InData).ActorLabel); },
-						[this](const IDashboardDataViewEntry& InData) -> FName
+						[](const IDashboardDataViewEntry& InData) -> FName
 						{
 							return CastEntry(InData).ActorIconName;
 						},
@@ -904,7 +973,7 @@ namespace UE::Audio::Insights
 					AudioInsightsModule.GetTimingViewExtender().StopCachingAndProcessingNewMessages();
 					bAutoScroll = false;
 				}
-				else
+				else if (AudioInsightsModule.GetTimingViewExtender().GetMessageCacheAndProcessingStatus() != ECacheAndProcess::None)
 				{
 					AudioInsightsModule.GetTimingViewExtender().StopProcessingNewMessages();
 				}
diff --git a/Engine/Plugins/AudioInsights/Source/AudioInsights/Private/Views/SoundDashboardViewFactory.cpp b/Engine/Plugins/AudioInsights/Source/AudioInsights/Private/Views/SoundDashboardViewFactory.cpp
--- a/Engine/Plugins/AudioInsights/Source/AudioInsights/Private/Views/SoundDashboardViewFactory.cpp
+++ b/Engine/Plugins/AudioInsights/Source/AudioInsights/Private/Views/SoundDashboardViewFactory.cpp
@@ -1945,7 +1945,7 @@ namespace UE::Audio::Insights
 			{
 				MuteColumnName,
 				{
-					.GetIconName = [this](const IDashboardDataTreeViewEntry& InData)
+					.GetIconName = [](const IDashboardDataTreeViewEntry& InData)
 					{
 						return "AudioInsights.Icon.SoundDashboard.Mute";
 					}
@@ -1954,7 +1954,7 @@ namespace UE::Audio::Insights
 			{
 				SoloColumnName,
 				{
-					.GetIconName = [this](const IDashboardDataTreeViewEntry& InData)
+					.GetIconName = [](const IDashboardDataTreeViewEntry& InData)
 					{
 						return "AudioInsights.Icon.SoundDashboard.Solo";
 					}
@@ -1979,7 +1979,7 @@ namespace UE::Audio::Insights
 
 						return SoundDashboardEntry.GetDisplayName();
 					},
-					.GetIconName = [this](const IDashboardDataTreeViewEntry& InData) -> FName
+					.GetIconName = [](const IDashboardDataTreeViewEntry& InData) -> FName
 					{
 						const FSoundDashboardEntry& SoundDashboardEntry = CastEntry(InData);
 
diff --git a/Engine/Plugins/AudioInsights/Source/AudioInsights/Private/Views/TableDashboardViewFactory.cpp b/Engine/Plugins/AudioInsights/Source/AudioInsights/Private/Views/TableDashboardViewFactory.cpp
--- a/Engine/Plugins/AudioInsights/Source/AudioInsights/Private/Views/TableDashboardViewFactory.cpp
+++ b/Engine/Plugins/AudioInsights/Source/AudioInsights/Private/Views/TableDashboardViewFactory.cpp
@@ -54,12 +54,32 @@ namespace UE::Audio::Insights
 		// Icon
 		+ SHorizontalBox::Slot()
 		.Padding(0.0f, 2.0f)
+		.HAlign(ColumnData.Alignment)
 		.AutoWidth()
 		[
-			IconName != NAME_None
+			ColumnData.GetIconName.IsSet()
 
 			? SNew(SImage)
-			.Image(FSlateStyle::Get().GetBrush(IconName))
+			.Image_Lambda([this, InRowData, Column]()
+			{
+				const FColumnData& ColumnData = GetColumns()[Column];
+				return FSlateStyle::Get().GetBrush(ColumnData.GetIconName(InRowData.Get()));
+			})
+			.ColorAndOpacity_Lambda([this, InRowData, Column]()
+			{
+				const FColumnData& ColumnData = GetColumns()[Column];
+				return ColumnData.GetIconColor.IsSet() ? ColumnData.GetIconColor(InRowData.Get()) : FLinearColor::White;
+			})
+			.ToolTipText_Lambda([this, InRowData, Column]()
+			{
+				const FColumnData& ColumnData = GetColumns()[Column];
+				return ColumnData.GetIconTooltip.IsSet() ? ColumnData.GetIconTooltip(InRowData.Get()) : FText::GetEmpty();
+			})
+			.Visibility_Lambda([this, InRowData, Column]()
+			{
+				const FColumnData& ColumnData = GetColumns()[Column];
+				return ColumnData.GetIconName(InRowData.Get()) == NAME_None ? EVisibility::Collapsed : EVisibility::Visible;
+			})
 
 			: SNullWidget::NullWidget
 		]
@@ -338,11 +358,20 @@ namespace UE::Audio::Insights
 	{
 		for (TSharedPtr<FTraceProviderBase> Provider : Providers)
 		{
-			const FName ProviderName = Provider->GetName();
-			if (const uint64* CurrentUpdateId = UpdateIds.Find(ProviderName))
+			if (!Provider.IsValid())
 			{
-				const uint64 LastUpdateId = Provider->GetLastUpdateId();
-				if (*CurrentUpdateId != LastUpdateId)
+				continue;
+			}
+
+			if (Provider->ShouldForceUpdate())
+			{
+				Provider->ResetShouldForceUpdate();
+
+				UpdateFilterReason = EProcessReason::EntriesUpdated;
+			}
+			else if (const uint64* CurrentUpdateId = UpdateIds.Find(Provider->GetName()))
+			{
+				if (*CurrentUpdateId != Provider->GetLastUpdateId())
 				{
 					UpdateFilterReason = EProcessReason::EntriesUpdated;
 				}
@@ -436,17 +465,37 @@ namespace UE::Audio::Insights
 
 		return SNew(SHorizontalBox)
 
-		// Icon
-		+ SHorizontalBox::Slot()
-		.Padding(0.0f, 2.0f)
-		.AutoWidth()
-		[
-			IconName != NAME_None
+			// Icon
+			+ SHorizontalBox::Slot()
+			.Padding(0.0f, 2.0f)
+			.HAlign(ColumnData.Alignment)
+			.AutoWidth()
+			[
+				ColumnData.GetIconName.IsSet()
 
-			? SNew(SImage)
-			.Image(FSlateStyle::Get().GetBrush(IconName))
+				? SNew(SImage)
+				.Image_Lambda([this, InRowData, Column]()
+				{
+					const FColumnData& ColumnData = GetColumns()[Column];
+					return FSlateStyle::Get().GetBrush(ColumnData.GetIconName(InRowData.Get()));
+				})
+				.ColorAndOpacity_Lambda([this, InRowData, Column]()
+				{
+					const FColumnData& ColumnData = GetColumns()[Column];
+					return ColumnData.GetIconColor.IsSet() ? ColumnData.GetIconColor(InRowData.Get()) : FLinearColor::White;
+				})
+				.ToolTipText_Lambda([this, InRowData, Column]()
+				{
+					const FColumnData& ColumnData = GetColumns()[Column];
+					return ColumnData.GetIconTooltip.IsSet() ? ColumnData.GetIconTooltip(InRowData.Get()) : FText::GetEmpty();
+				})
+				.Visibility_Lambda([this, InRowData, Column]()
+				{
+					const FColumnData& ColumnData = GetColumns()[Column];
+					return ColumnData.GetIconName(InRowData.Get()) == NAME_None ? EVisibility::Collapsed : EVisibility::Visible;
+				})
 
-			: SNullWidget::NullWidget
+				: SNullWidget::NullWidget
 		]
 
 		+ SHorizontalBox::Slot()
diff --git a/Engine/Plugins/AudioInsights/Source/AudioInsights/Public/AudioInsightsStyle.h b/Engine/Plugins/AudioInsights/Source/AudioInsights/Public/AudioInsightsStyle.h
--- a/Engine/Plugins/AudioInsights/Source/AudioInsights/Public/AudioInsightsStyle.h
+++ b/Engine/Plugins/AudioInsights/Source/AudioInsights/Public/AudioInsightsStyle.h
@@ -231,6 +231,9 @@ namespace UE::Audio::Insights
 			Set("AudioInsights.Icon.SoundDashboard.Transparent",      new IMAGE_BRUSH_SVG(TEXT("Icons/SoundDashboard/transparent"),      Icon16));
 			Set("AudioInsights.Icon.SoundDashboard.Visible",		  new IMAGE_BRUSH_SVG(TEXT("Icons/SoundDashboard/visible"),			 Icon16));
 			Set("AudioInsights.Icon.SoundDashboard.Invisible",		  new IMAGE_BRUSH_SVG(TEXT("Icons/SoundDashboard/invisible"),		 Icon16));
+
+			/* Event Log */
+			Set("AudioInsights.Icon.EventLog.MarkedForDelete",		  new IMAGE_BRUSH_SVG(TEXT("Icons/EventLog/cache_marked_for_delete"), Icon16));
 			
 			// Category colors
 			Set("SoundDashboard.MetaSoundColor", FLinearColor(0.008f, 0.76f, 0.078f));
diff --git a/Engine/Plugins/AudioInsights/Source/AudioInsights/Public/Cache/AudioInsightsCacheManager.h b/Engine/Plugins/AudioInsights/Source/AudioInsights/Public/Cache/AudioInsightsCacheManager.h
--- a/Engine/Plugins/AudioInsights/Source/AudioInsights/Public/Cache/AudioInsightsCacheManager.h
+++ b/Engine/Plugins/AudioInsights/Source/AudioInsights/Public/Cache/AudioInsightsCacheManager.h
@@ -42,6 +42,10 @@ namespace UE::Audio::Insights
 		// If the timestamp does not exist within the current cache, the returned result will be unset
 		UE_API TOptional<uint32> TryGetNumChunksFromStartForTimestamp(const double Timestamp) const;
 
+		// Returns a chunk by index relative to the start of the cache
+		// e.g. NumChunksFromStart = 0u will return the chunk with the earliest messages in the cache
+		UE_API const FAudioCachedMessageChunk* const GetChunk(const uint32 NumChunksFromStart) const;
+
 		// Find the closest message to a timestamp in the cache for a specific message type
 		// Will look backwards from the chosen timestamp until a message is found
 		// @param ID: When overloaded, the search algorithm will only return messages that have a matching ID.
diff --git a/Engine/Plugins/AudioInsights/Source/AudioInsights/Public/Messages/AudioEventLogTraceMessages.h b/Engine/Plugins/AudioInsights/Source/AudioInsights/Public/Messages/AudioEventLogTraceMessages.h
--- a/Engine/Plugins/AudioInsights/Source/AudioInsights/Public/Messages/AudioEventLogTraceMessages.h
+++ b/Engine/Plugins/AudioInsights/Source/AudioInsights/Public/Messages/AudioEventLogTraceMessages.h
@@ -25,6 +25,14 @@ namespace UE::Audio::Insights
 		SoundCueTemplate
 	};
 
+#if WITH_EDITOR
+	enum class EAudioEventCacheState : uint8
+	{
+		Cached,
+		NextToBeDeleted
+	};
+#endif // WITH_EDITOR
+
 	struct FAudioEventLogMessage : public IAudioCachedMessage
 	{
 		FAudioEventLogMessage() = default;
@@ -66,6 +74,11 @@ namespace UE::Audio::Insights
 		FString CategoryName;
 
 		EAudioEventLogSoundCategory Category = EAudioEventLogSoundCategory::None;
+
+#if WITH_EDITOR
+		EAudioEventCacheState CachedState = EAudioEventCacheState::Cached;
+		bool bCacheStatusIsDirty = false;
+#endif // WITH_EDITOR
 	};
 
 	class FAudioEventLogMessages
diff --git a/Engine/Plugins/AudioInsights/Source/AudioInsights/Public/Settings/AudioEventLogSettings.h b/Engine/Plugins/AudioInsights/Source/AudioInsights/Public/Settings/AudioEventLogSettings.h
--- a/Engine/Plugins/AudioInsights/Source/AudioInsights/Public/Settings/AudioEventLogSettings.h
+++ b/Engine/Plugins/AudioInsights/Source/AudioInsights/Public/Settings/AudioEventLogSettings.h
@@ -36,6 +36,9 @@ struct FAudioEventLogVisibleColumns : public FVisibleColumnsSettings
 {
 	GENERATED_BODY()
 
+	UPROPERTY(EditAnywhere, config, Category = EventLog)
+	bool bCacheStatus = true;
+
 	UPROPERTY(EditAnywhere, config, Category = EventLog)
 	bool bMessageID = false;
 
diff --git a/Engine/Plugins/AudioInsights/Source/AudioInsights/Public/Views/TableDashboardViewFactory.h b/Engine/Plugins/AudioInsights/Source/AudioInsights/Public/Views/TableDashboardViewFactory.h
--- a/Engine/Plugins/AudioInsights/Source/AudioInsights/Public/Views/TableDashboardViewFactory.h
+++ b/Engine/Plugins/AudioInsights/Source/AudioInsights/Public/Views/TableDashboardViewFactory.h
@@ -158,6 +158,8 @@ namespace UE::Audio::Insights
 			bool bDefaultHidden = false;
 			const float FillWidth = 1.0f;
 			const EHorizontalAlignment Alignment = HAlign_Left;
+			const TFunction<FLinearColor(const IDashboardDataViewEntry&)> GetIconColor = nullptr;
+			const TFunction<FText(const IDashboardDataViewEntry&)> GetIconTooltip = nullptr;
 		};
 
 		virtual const TMap<FName, FColumnData>& GetColumns() const = 0;
