[Audio Insights] AssetProvider now handles AssetPaths instead of dashboard entries to not mix responsibilities + bug fixes for audio meters getting stuck unless we scroll to the checked item

diff --git a/Engine/Plugins/AudioInsights/Source/AudioInsights/Private/Providers/AudioBusAssetProvider.cpp b/Engine/Plugins/AudioInsights/Source/AudioInsights/Private/Providers/AudioBusAssetProvider.cpp
deleted file mode 100644
diff --git a/Engine/Plugins/AudioInsights/Source/AudioInsights/Private/Providers/AudioBusAssetProvider.h b/Engine/Plugins/AudioInsights/Source/AudioInsights/Private/Providers/AudioBusAssetProvider.h
deleted file mode 100644
diff --git a/Engine/Plugins/AudioInsights/Source/AudioInsights/Private/Providers/AudioBusTraceProvider.cpp b/Engine/Plugins/AudioInsights/Source/AudioInsights/Private/Providers/AudioBusTraceProvider.cpp
--- a/Engine/Plugins/AudioInsights/Source/AudioInsights/Private/Providers/AudioBusTraceProvider.cpp
+++ b/Engine/Plugins/AudioInsights/Source/AudioInsights/Private/Providers/AudioBusTraceProvider.cpp
@@ -34,6 +34,23 @@ namespace UE::Audio::Insights
 
 			return nullptr;
 		}
+#else
+		TSharedPtr<FAudioBusDashboardEntry> CreateEntryFromAsset(const FString& InAssetPath)
+		{
+			TSharedPtr<FAudioBusDashboardEntry> Entry = MakeShared<FAudioBusDashboardEntry>();
+			
+			const TObjectPtr<UAudioBus> AudioBus = Cast<UAudioBus>(FSoftObjectPath(InAssetPath).ResolveObject());
+	
+			Entry->Timestamp = FPlatformTime::Seconds() - GStartTime;
+			Entry->AudioBusId = AudioBus ? AudioBus->GetUniqueID() : INDEX_NONE;
+
+			Entry->Name = InAssetPath;
+
+			Entry->AudioMeterInfo->NumChannels = AudioBus ? AudioBus->GetNumChannels() : 0;
+			Entry->AudioBusType = EAudioBusType::AssetBased;
+
+			return Entry;
+		}
 #endif // !WITH_EDITOR
 	}
 
@@ -284,6 +301,14 @@ namespace UE::Audio::Insights
 			}
 		});
 
+#if WITH_EDITOR
+		if (bAssetsUpdated)
+		{
+			OnAudioBusListUpdated.ExecuteIfBound();
+			bAssetsUpdated = false;
+		}
+#endif // WITH_EDITOR
+
 		return true;
 	}
 
@@ -293,37 +318,45 @@ namespace UE::Audio::Insights
 		AudioBusAssetProvider.RequestEntriesUpdate();
 	}
 
-	void FAudioBusTraceProvider::HandleOnAudioBusAssetAdded(const TSharedPtr<FAudioBusDashboardEntry>& InAudioBusDashboardEntry)
+	void FAudioBusTraceProvider::HandleOnAudioBusAssetAdded(const FString& InAssetPath)
 	{
+		using namespace FAudioBusTraceProviderPrivate;
+
 		if (const FAudioDeviceManager* DeviceManager = FAudioDeviceManager::Get())
 		{
 			const TArray<FAudioDevice*> AudioDevices = DeviceManager->GetAudioDevices();
 
 			for (const FAudioDevice* AudioDevice : AudioDevices)
 			{
-				if (AudioDevice)
+				if (AudioDevice == nullptr)
 				{
-					const TSharedPtr<FAudioBusDashboardEntry>* FoundAudioBusDashboardEntry = FindDeviceEntry(AudioDevice->DeviceID, InAudioBusDashboardEntry->AudioBusId);
+					continue;
+				}
 
-					if (FoundAudioBusDashboardEntry == nullptr)
+				const TObjectPtr<UAudioBus> AudioBus = Cast<UAudioBus>(FSoftObjectPath(InAssetPath).ResolveObject());
+				const uint32 AudioBusId = AudioBus ? AudioBus->GetUniqueID() : INDEX_NONE;
+
+				const TSharedPtr<FAudioBusDashboardEntry>* FoundAudioBusDashboardEntry = FindDeviceEntry(AudioDevice->DeviceID, AudioBusId);
+
+				if (FoundAudioBusDashboardEntry == nullptr)
+				{
+					UpdateDeviceEntry(AudioDevice->DeviceID, AudioBusId, [&InAssetPath](TSharedPtr<FAudioBusDashboardEntry>& Entry)
 					{
-						UpdateDeviceEntry(AudioDevice->DeviceID, InAudioBusDashboardEntry->AudioBusId, [&InAudioBusDashboardEntry](TSharedPtr<FAudioBusDashboardEntry>& Entry)
+						if (!Entry.IsValid())
 						{
-							if (!Entry.IsValid())
-							{
-								// We need to take a deep copy to avoid sharing the same entry instance across different audio devices
-								Entry = MakeShared<FAudioBusDashboardEntry>(*InAudioBusDashboardEntry);
-							}
-						});
-					}
-
-					OnAudioBusAdded.ExecuteIfBound(InAudioBusDashboardEntry->AudioBusId);
+							Entry = CreateEntryFromAsset(InAssetPath);
+						}
+					});
 				}
+
+				++LastUpdateId;
+				OnAudioBusAdded.ExecuteIfBound(AudioBusId);
+				bAssetsUpdated = true;
 			}
 		}
 	}
 
-	void FAudioBusTraceProvider::HandleOnAudioBusAssetRemoved(const uint32 InAudioBusId)
+	void FAudioBusTraceProvider::HandleOnAudioBusAssetRemoved(const FString& InAssetPath)
 	{
 		if (const FAudioDeviceManager* DeviceManager = FAudioDeviceManager::Get())
 		{
@@ -333,16 +366,23 @@ namespace UE::Audio::Insights
 			{
 				if (AudioDevice)
 				{
-					RemoveDeviceEntry(AudioDevice->DeviceID, InAudioBusId);
+					const TObjectPtr<UAudioBus> AudioBus = Cast<UAudioBus>(FSoftObjectPath(InAssetPath).ResolveObject());
+					const uint32 AudioBusId = AudioBus ? AudioBus->GetUniqueID() : INDEX_NONE;
 
-					OnAudioBusRemoved.ExecuteIfBound(InAudioBusId);
+					RemoveDeviceEntry(AudioDevice->DeviceID, AudioBusId);
+
+					++LastUpdateId;
+					OnAudioBusRemoved.ExecuteIfBound(AudioBusId);
+					bAssetsUpdated = true;
 				}
 			}
 		}
 	}
 
-	void FAudioBusTraceProvider::HandleOnAudioBusAssetListUpdated(const TArray<TSharedPtr<FAudioBusDashboardEntry>>& InAudioBusAssetEntries)
+	void FAudioBusTraceProvider::HandleOnAudioBusAssetListUpdated(const TArray<FString>& InAssetPaths)
 	{
+		using namespace FAudioBusTraceProviderPrivate;
+
 		const FAudioDeviceManager* DeviceManager = FAudioDeviceManager::Get();
 
 		if (DeviceManager == nullptr)
@@ -350,8 +390,6 @@ namespace UE::Audio::Insights
 			return;
 		}
 
-		DeviceDataMap.Reset();
-
 		for (const FAudioDevice* AudioDevice : DeviceManager->GetAudioDevices())
 		{
 			if (AudioDevice == nullptr)
@@ -359,24 +397,29 @@ namespace UE::Audio::Insights
 				continue;
 			}
 
-			for (const auto& AudioBusDashboardEntry : InAudioBusAssetEntries)
+			for (const FString& AssetPath : InAssetPaths)
 			{
-				const TSharedPtr<FAudioBusDashboardEntry>* FoundAudioBusDashboardEntry = FindDeviceEntry(AudioDevice->DeviceID, AudioBusDashboardEntry->AudioBusId);
+				const TObjectPtr<UAudioBus> AudioBus = Cast<UAudioBus>(FSoftObjectPath(AssetPath).ResolveObject());
+				const uint32 AudioBusId = AudioBus ? AudioBus->GetUniqueID() : INDEX_NONE;
+
+				const TSharedPtr<FAudioBusDashboardEntry>* FoundAudioBusDashboardEntry = FindDeviceEntry(AudioDevice->DeviceID, AudioBusId);
 
 				if (FoundAudioBusDashboardEntry == nullptr)
 				{
-					UpdateDeviceEntry(AudioDevice->DeviceID, AudioBusDashboardEntry->AudioBusId, [&AudioBusDashboardEntry](TSharedPtr<FAudioBusDashboardEntry>& Entry)
+					UpdateDeviceEntry(AudioDevice->DeviceID, AudioBusId, [&AssetPath](TSharedPtr<FAudioBusDashboardEntry>& Entry)
 					{
 						if (!Entry.IsValid())
 						{
-							Entry = MakeShared<FAudioBusDashboardEntry>(*AudioBusDashboardEntry);;
+							Entry = CreateEntryFromAsset(AssetPath);
 						}
 					});
 				}
 			}
 		}
 
+		++LastUpdateId;
 		OnAudioBusListUpdated.ExecuteIfBound();
+		bAssetsUpdated = true;
 	}
 #else
 	void FAudioBusTraceProvider::InitSessionCachedMessages(TraceServices::IAnalysisSession& InSession)
@@ -562,6 +605,7 @@ namespace UE::Audio::Insights
 	void FAudioBusTraceProvider::OnTimeControlMethodReset()
 	{
 		AudioBusAssetProvider.RequestEntriesUpdate();
+		OnAudioBusListUpdated.ExecuteIfBound();
 	}
 #endif
 } // namespace UE::Audio::Insights
diff --git a/Engine/Plugins/AudioInsights/Source/AudioInsights/Private/Providers/AudioBusTraceProvider.h b/Engine/Plugins/AudioInsights/Source/AudioInsights/Private/Providers/AudioBusTraceProvider.h
--- a/Engine/Plugins/AudioInsights/Source/AudioInsights/Private/Providers/AudioBusTraceProvider.h
+++ b/Engine/Plugins/AudioInsights/Source/AudioInsights/Private/Providers/AudioBusTraceProvider.h
@@ -8,7 +8,7 @@
 
 #if WITH_EDITOR
 #include "Delegates/Delegate.h"
-#include "Providers/AudioBusAssetProvider.h"
+#include "Providers/AssetProvider.h"
 #endif // WITH_EDITOR
 
 namespace UE::Audio::Insights
@@ -38,22 +38,23 @@ namespace UE::Audio::Insights
 	private:
 		virtual bool ProcessMessages() override;
 
-		virtual void OnTimingViewTimeMarkerChanged(double TimeMarker) override;
-
 #if WITH_EDITOR
 		virtual void OnTimeControlMethodReset() override;
 
-		void HandleOnAudioBusAssetAdded(const TSharedPtr<FAudioBusDashboardEntry>& InAudioBusDashboardEntry);
-		void HandleOnAudioBusAssetRemoved(const uint32 InAudioBusId);
-		void HandleOnAudioBusAssetListUpdated(const TArray<TSharedPtr<FAudioBusDashboardEntry>>& InAudioBusAssetEntries);
+		void HandleOnAudioBusAssetAdded(const FString& InAssetPath);
+		void HandleOnAudioBusAssetRemoved(const FString& InAssetPath);
+		void HandleOnAudioBusAssetListUpdated(const TArray<FString>& InAssetPaths);
 #else
 		TUniquePtr<FAudioBusSessionCachedMessages> SessionCachedMessages;
 #endif // WITH_EDITOR
 
+		virtual void OnTimingViewTimeMarkerChanged(double TimeMarker) override;
+
 		FAudioBusMessages TraceMessages;
 
 #if WITH_EDITOR
-		FAudioBusAssetProvider AudioBusAssetProvider;
+		TAssetProvider<UAudioBus> AudioBusAssetProvider;
+		bool bAssetsUpdated = false;
 #endif // WITH_EDITOR
 	};
 } // namespace UE::Audio::Insights
diff --git a/Engine/Plugins/AudioInsights/Source/AudioInsights/Private/Providers/SubmixAssetProvider.cpp b/Engine/Plugins/AudioInsights/Source/AudioInsights/Private/Providers/SubmixAssetProvider.cpp
deleted file mode 100644
diff --git a/Engine/Plugins/AudioInsights/Source/AudioInsights/Private/Providers/SubmixAssetProvider.h b/Engine/Plugins/AudioInsights/Source/AudioInsights/Private/Providers/SubmixAssetProvider.h
deleted file mode 100644
diff --git a/Engine/Plugins/AudioInsights/Source/AudioInsights/Private/Providers/SubmixTraceProvider.cpp b/Engine/Plugins/AudioInsights/Source/AudioInsights/Private/Providers/SubmixTraceProvider.cpp
--- a/Engine/Plugins/AudioInsights/Source/AudioInsights/Private/Providers/SubmixTraceProvider.cpp
+++ b/Engine/Plugins/AudioInsights/Source/AudioInsights/Private/Providers/SubmixTraceProvider.cpp
@@ -34,6 +34,29 @@ namespace UE::Audio::Insights
 
 			return nullptr;
 		}
+#else
+		TSharedPtr<FSubmixDashboardEntry> CreateEntryFromAsset(const FString& InAssetPath)
+		{
+			TSharedPtr<FSubmixDashboardEntry> Entry = MakeShared<FSubmixDashboardEntry>();
+
+			Entry->Timestamp = FPlatformTime::Seconds() - GStartTime;
+			Entry->SubmixId = GetTypeHash(InAssetPath);
+
+			Entry->Name = InAssetPath;
+
+			if (const FAudioDeviceManager* AudioDeviceManager = FAudioDeviceManager::Get())
+			{
+				const IAudioInsightsModule& AudioInsightsModule = IAudioInsightsModule::GetEditorChecked();
+				const ::Audio::FDeviceId AudioDeviceId = AudioInsightsModule.GetDeviceId();
+
+				if (const ::Audio::FMixerDevice* MixerDevice = static_cast<const ::Audio::FMixerDevice*>(AudioDeviceManager->GetAudioDeviceRaw(AudioDeviceId)))
+				{
+					Entry->AudioMeterInfo->NumChannels = MixerDevice->GetNumDeviceChannels();
+				}
+			}
+
+			return Entry;
+		}
 #endif // !WITH_EDITOR
 	}
 
@@ -266,6 +289,14 @@ namespace UE::Audio::Insights
 			}
 		});
 
+#if WITH_EDITOR
+		if (bAssetsUpdated)
+		{
+			OnSubmixListUpdated.ExecuteIfBound();
+			bAssetsUpdated = false;
+		}
+#endif // WITH_EDITOR
+
 		return true;
 	}
 
@@ -275,37 +306,44 @@ namespace UE::Audio::Insights
 		SubmixAssetProvider.RequestEntriesUpdate();
 	}
 
-	void FSubmixTraceProvider::HandleOnSubmixAssetAdded(const TSharedPtr<FSubmixDashboardEntry>& InSubmixDashboardEntry)
+	void FSubmixTraceProvider::HandleOnSubmixAssetAdded(const FString& InAssetPath)
 	{
+		using namespace FSubmixTraceProviderPrivate;
+
 		if (const FAudioDeviceManager* DeviceManager = FAudioDeviceManager::Get())
 		{
 			const TArray<FAudioDevice*> AudioDevices = DeviceManager->GetAudioDevices();
 
 			for (const FAudioDevice* AudioDevice : AudioDevices)
 			{
-				if (AudioDevice)
+				if (AudioDevice == nullptr)
 				{
-					const TSharedPtr<FSubmixDashboardEntry>* FoundSubmixDashboardEntry = FindDeviceEntry(AudioDevice->DeviceID, InSubmixDashboardEntry->SubmixId);
+					continue;
+				}
+
+				const uint32 SubmixId = GetTypeHash(InAssetPath);
 
-					if (FoundSubmixDashboardEntry == nullptr)
+				const TSharedPtr<FSubmixDashboardEntry>* FoundSubmixDashboardEntry = FindDeviceEntry(AudioDevice->DeviceID, SubmixId);
+
+				if (FoundSubmixDashboardEntry == nullptr)
+				{
+					UpdateDeviceEntry(AudioDevice->DeviceID, SubmixId, [this, &InAssetPath](TSharedPtr<FSubmixDashboardEntry>& Entry)
 					{
-						UpdateDeviceEntry(AudioDevice->DeviceID, InSubmixDashboardEntry->SubmixId, [&InSubmixDashboardEntry](TSharedPtr<FSubmixDashboardEntry>& Entry)
+						if (!Entry.IsValid())
 						{
-							if (!Entry.IsValid())
-							{
-								// We need to take a deep copy to avoid sharing the same entry instance across different audio devices
-								Entry = MakeShared<FSubmixDashboardEntry>(*InSubmixDashboardEntry);
-							}
-						});
-					}
-
-					OnSubmixAssetAdded.ExecuteIfBound();
+							Entry = CreateEntryFromAsset(InAssetPath);
+						}
+					});
 				}
+
+				++LastUpdateId;
+				OnSubmixAdded.ExecuteIfBound(SubmixId);
+				bAssetsUpdated = true;
 			}
 		}
 	}
 
-	void FSubmixTraceProvider::HandleOnSubmixAssetRemoved(const uint32 InSubmixId)
+	void FSubmixTraceProvider::HandleOnSubmixAssetRemoved(const FString& InAssetPath)
 	{
 		if (const FAudioDeviceManager* DeviceManager = FAudioDeviceManager::Get())
 		{
@@ -315,16 +353,22 @@ namespace UE::Audio::Insights
 			{
 				if (AudioDevice)
 				{
-					RemoveDeviceEntry(AudioDevice->DeviceID, InSubmixId);
+					const uint32 SubmixId = GetTypeHash(InAssetPath);
+
+					RemoveDeviceEntry(AudioDevice->DeviceID, SubmixId);
 
-					OnSubmixAssetRemoved.ExecuteIfBound();
+					++LastUpdateId;
+					OnSubmixRemoved.ExecuteIfBound(SubmixId);
+					bAssetsUpdated = true;
 				}
 			}
 		}
 	}
 
-	void FSubmixTraceProvider::HandleOnSubmixAssetListUpdated(const TArray<TSharedPtr<FSubmixDashboardEntry>>& InSubmixAssetEntries)
+	void FSubmixTraceProvider::HandleOnSubmixAssetListUpdated(const TArray<FString>& InAssetPaths)
 	{
+		using namespace FSubmixTraceProviderPrivate;
+
 		const FAudioDeviceManager* DeviceManager = FAudioDeviceManager::Get();
 
 		if (DeviceManager == nullptr)
@@ -332,8 +376,6 @@ namespace UE::Audio::Insights
 			return;
 		}
 
-		DeviceDataMap.Reset();
-
 		for (const FAudioDevice* AudioDevice : DeviceManager->GetAudioDevices())
 		{
 			if (AudioDevice == nullptr)
@@ -341,24 +383,28 @@ namespace UE::Audio::Insights
 				continue;
 			}
 
-			for (const auto& SubmixDashboardEntry : InSubmixAssetEntries)
+			for (const FString& AssetPath : InAssetPaths)
 			{
-				const TSharedPtr<FSubmixDashboardEntry>* FoundSubmixDashboardEntry = FindDeviceEntry(AudioDevice->DeviceID, SubmixDashboardEntry->SubmixId);
+				const uint32 SubmixId = GetTypeHash(AssetPath);
+
+				const TSharedPtr<FSubmixDashboardEntry>* FoundSubmixDashboardEntry = FindDeviceEntry(AudioDevice->DeviceID, SubmixId);
 
 				if (FoundSubmixDashboardEntry == nullptr)
 				{
-					UpdateDeviceEntry(AudioDevice->DeviceID, SubmixDashboardEntry->SubmixId, [&SubmixDashboardEntry](TSharedPtr<FSubmixDashboardEntry>& Entry)
+					UpdateDeviceEntry(AudioDevice->DeviceID, SubmixId, [&AssetPath](TSharedPtr<FSubmixDashboardEntry>& Entry)
 					{
 						if (!Entry.IsValid())
 						{
-							Entry = MakeShared<FSubmixDashboardEntry>(*SubmixDashboardEntry);;
+							Entry = CreateEntryFromAsset(AssetPath);
 						}
 					});
 				}
 			}
 		}
 
-		OnSubmixAssetListUpdated.ExecuteIfBound();
+		++LastUpdateId;
+		OnSubmixListUpdated.ExecuteIfBound();
+		bAssetsUpdated = true;
 	}
 #else
 	void FSubmixTraceProvider::InitSessionCachedMessages(TraceServices::IAnalysisSession& InSession)
@@ -397,7 +443,6 @@ namespace UE::Audio::Insights
 				Entry->AudioMeterInfo->NumChannels = Msg.NumChannels;
 				Entry->bIsMainSubmix = Msg.bIsMainSubmix;
 			});
-
 		};
 
 #if WITH_EDITOR
@@ -516,6 +561,7 @@ namespace UE::Audio::Insights
 						Entry->bHasActivity = HasActivityMessage.bHasActivity;
 					});
 				}
+
 				if (CachedEntryInfo.EnvelopeValuesMessage.IsSet())
 				{
 					const FSubmixEnvelopeValuesMessage& EnvelopeValuesMessage = CachedEntryInfo.EnvelopeValuesMessage.GetValue();
@@ -545,6 +591,7 @@ namespace UE::Audio::Insights
 	void FSubmixTraceProvider::OnTimeControlMethodReset()
 	{
 		SubmixAssetProvider.RequestEntriesUpdate();
+		OnSubmixListUpdated.ExecuteIfBound();
 	}
 #endif // WITH_EDITOR
 } // namespace UE::Audio::Insights
diff --git a/Engine/Plugins/AudioInsights/Source/AudioInsights/Private/Providers/SubmixTraceProvider.h b/Engine/Plugins/AudioInsights/Source/AudioInsights/Private/Providers/SubmixTraceProvider.h
--- a/Engine/Plugins/AudioInsights/Source/AudioInsights/Private/Providers/SubmixTraceProvider.h
+++ b/Engine/Plugins/AudioInsights/Source/AudioInsights/Private/Providers/SubmixTraceProvider.h
@@ -8,7 +8,7 @@
 
 #if WITH_EDITOR
 #include "Delegates/Delegate.h"
-#include "Providers/SubmixAssetProvider.h"
+#include "Providers/AssetProvider.h"
 #endif // WITH_EDITOR
 
 namespace UE::Audio::Insights
@@ -25,14 +25,15 @@ namespace UE::Audio::Insights
 
 #if !WITH_EDITOR
 		virtual void InitSessionCachedMessages(TraceServices::IAnalysisSession& InSession) override;
-#else
+#endif // !WITH_EDITOR
 		void RequestEntriesUpdate();
 
-		DECLARE_DELEGATE(FOnSubmixAssetsUpdated);
-		FOnSubmixAssetsUpdated OnSubmixAssetAdded;
-		FOnSubmixAssetsUpdated OnSubmixAssetRemoved;
-		FOnSubmixAssetsUpdated OnSubmixAssetListUpdated;
-#endif // !WITH_EDITOR
+		DECLARE_DELEGATE_OneParam(FOnSubmixChanged, const uint32 /*SubmixId*/);
+		FOnSubmixChanged OnSubmixAdded;
+		FOnSubmixChanged OnSubmixRemoved;
+
+		DECLARE_DELEGATE(FOnSubmixListUpdated);
+		FOnSubmixListUpdated OnSubmixListUpdated;
 
 	private:
 		virtual bool ProcessMessages() override;
@@ -40,9 +41,9 @@ namespace UE::Audio::Insights
 #if WITH_EDITOR
 		virtual void OnTimeControlMethodReset() override;
 
-		void HandleOnSubmixAssetAdded(const TSharedPtr<FSubmixDashboardEntry>& InSubmixDashboardEntry);
-		void HandleOnSubmixAssetRemoved(const uint32 InSubmixId);
-		void HandleOnSubmixAssetListUpdated(const TArray<TSharedPtr<FSubmixDashboardEntry>>& InSubmixAssetEntries);
+		void HandleOnSubmixAssetAdded(const FString& InAssetPath);
+		void HandleOnSubmixAssetRemoved(const FString& InAssetPath);
+		void HandleOnSubmixAssetListUpdated(const TArray<FString>& InAssetPaths);
 #else
 		TUniquePtr<FSubmixSessionCachedMessages> SessionCachedMessages;
 #endif // WITH_EDITOR
@@ -52,8 +53,8 @@ namespace UE::Audio::Insights
 		FSubmixMessages TraceMessages;
 
 #if WITH_EDITOR
-		FSubmixAssetProvider SubmixAssetProvider;
+		TAssetProvider<USoundSubmix> SubmixAssetProvider;
+		bool bAssetsUpdated = false;
 #endif // WITH_EDITOR
-
 	};
 } // namespace UE::Audio::Insights
diff --git a/Engine/Plugins/AudioInsights/Source/AudioInsights/Private/Views/AudioBusDashboardViewFactory.cpp b/Engine/Plugins/AudioInsights/Source/AudioInsights/Private/Views/AudioBusDashboardViewFactory.cpp
--- a/Engine/Plugins/AudioInsights/Source/AudioInsights/Private/Views/AudioBusDashboardViewFactory.cpp
+++ b/Engine/Plugins/AudioInsights/Source/AudioInsights/Private/Views/AudioBusDashboardViewFactory.cpp
@@ -40,6 +40,8 @@ namespace UE::Audio::Insights
 			IAudioInsightsTraceModule& AudioInsightsTraceModule = IAudioInsightsModule::GetChecked().GetTraceModule();
 			AudioInsightsTraceModule.ExecuteConsoleCommand(AudioBusCommandStr);
 		}
+
+		const FSlateColor AudioBusMeterTextColor(FColor(80, 200, 255));
 	}
 
 	FAudioBusDashboardViewFactory::FAudioBusDashboardViewFactory()
@@ -127,11 +129,8 @@ namespace UE::Audio::Insights
 			const FAudioBusDashboardEntry& AudioBusDashboardEntry = CastEntry(*InRowData);
 			const uint32 AudioBusId = AudioBusDashboardEntry.AudioBusId;
 
-			const bool bIsAudioBusChecked = AudioBusCheckboxCheckedStates.FindOrAdd(AudioBusId);
-
-			if (bIsAudioBusChecked)
+			if (CheckedAudioBuses.Find(AudioBusId) != nullptr)
 			{
-				SendAudioBusEnvelopeFollowerCVar(bIsAudioBusChecked, AudioBusDashboardEntry.AudioBusId);
 				FAudioMeterProvider::OnUpdateAudioMeterInfo.ExecuteIfBound(AudioBusId, AudioBusDashboardEntry.AudioMeterInfo, AudioBusDashboardEntry.GetDisplayName());
 			}
 
@@ -143,10 +142,7 @@ namespace UE::Audio::Insights
 					SNew(SCheckBox)
 					.IsChecked_Lambda([this, AudioBusId]()
 					{
-						const bool* FoundAudioBusCheckedState = AudioBusCheckboxCheckedStates.Find(AudioBusId);
-						const bool bIsAudioBusChecked = FoundAudioBusCheckedState && *FoundAudioBusCheckedState == true;
-
-						return bIsAudioBusChecked ? ECheckBoxState::Checked : ECheckBoxState::Unchecked;
+						return CheckedAudioBuses.Find(AudioBusId) != nullptr ? ECheckBoxState::Checked : ECheckBoxState::Unchecked;
 					})
 					.OnCheckStateChanged_Lambda([this, InRowData](ECheckBoxState NewState)
 					{
@@ -154,19 +150,16 @@ namespace UE::Audio::Insights
 
 						const FAudioBusDashboardEntry& AudioBusDashboardEntry = CastEntry(*InRowData);
 
-						if (bool* FoundAudioBusCheckedState = AudioBusCheckboxCheckedStates.Find(AudioBusDashboardEntry.AudioBusId))
-						{
-							*FoundAudioBusCheckedState = bIsChecked;
-						}
-
 						SendAudioBusEnvelopeFollowerCVar(bIsChecked, AudioBusDashboardEntry.AudioBusId);
 
 						if (bIsChecked)
 						{
-							FAudioMeterProvider::OnAddAudioMeter.ExecuteIfBound(AudioBusDashboardEntry.AudioBusId, AudioBusDashboardEntry.AudioMeterInfo, AudioBusDashboardEntry.GetDisplayName(), FSlateColor(FColor(80, 200, 255)));
+							CheckedAudioBuses.Emplace(AudioBusDashboardEntry.AudioBusId);
+							FAudioMeterProvider::OnAddAudioMeter.ExecuteIfBound(AudioBusDashboardEntry.AudioBusId, AudioBusDashboardEntry.AudioMeterInfo, AudioBusDashboardEntry.GetDisplayName(), AudioBusMeterTextColor);
 						}
 						else
 						{
+							CheckedAudioBuses.Remove(AudioBusDashboardEntry.AudioBusId);
 							FAudioMeterProvider::OnRemoveAudioMeter.ExecuteIfBound(AudioBusDashboardEntry.AudioBusId);
 						}
 					})
@@ -182,7 +175,7 @@ namespace UE::Audio::Insights
 				[
 					SNew(STextBlock)
 					.Text(ValueText)
-					.ColorAndOpacity(AudioBusDashboardEntry.AudioBusType == EAudioBusType::CodeGenerated ? FSlateColor(FColor(80, 200, 255)) : FSlateColor::UseForeground())
+					.ColorAndOpacity(AudioBusDashboardEntry.AudioBusType == EAudioBusType::CodeGenerated ? AudioBusMeterTextColor : FSlateColor::UseForeground())
 					.OnDoubleClicked_Lambda([this, InRowData](const FGeometry& MyGeometry, const FPointerEvent& PointerEvent)
 					{
 #if WITH_EDITOR
@@ -281,8 +274,6 @@ namespace UE::Audio::Insights
 
 	TSharedRef<SWidget> FAudioBusDashboardViewFactory::MakeWidget(TSharedRef<SDockTab> OwnerTab, const FSpawnTabArgs& SpawnTabArgs)
 	{
-		using namespace FAudioBusDashboardViewFactoryPrivate;
-
 		if (!DashboardWidget.IsValid())
 		{
 			SAssignNew(DashboardWidget, SVerticalBox)
@@ -315,18 +306,7 @@ namespace UE::Audio::Insights
 #endif // WITH_EDITOR
 
 		// Update audio meters (if any has been checked previously)
-		for (const TSharedPtr<IDashboardDataViewEntry>& Entry : DataViewEntries)
-		{
-			const FAudioBusDashboardEntry& AudioBusDashboardEntry = CastEntry(*Entry);
-
-			if (const bool* bFoundAudioBusCheckboxCheckedState = AudioBusCheckboxCheckedStates.Find(AudioBusDashboardEntry.AudioBusId))
-			{
-				if (*bFoundAudioBusCheckboxCheckedState)
-				{
-					FAudioMeterProvider::OnUpdateAudioMeterInfo.ExecuteIfBound(AudioBusDashboardEntry.AudioBusId, AudioBusDashboardEntry.AudioMeterInfo, AudioBusDashboardEntry.GetDisplayName());
-				}
-			}
-		}
+		UpdateAudioMeters();
 
 		return DashboardWidget->AsShared();
 	}
@@ -465,6 +445,22 @@ namespace UE::Audio::Insights
 		}
 	}
 
+	void FAudioBusDashboardViewFactory::UpdateAudioMeters()
+	{
+		using namespace FAudioBusDashboardViewFactoryPrivate;
+
+		for (const TSharedPtr<IDashboardDataViewEntry>& Entry : DataViewEntries)
+		{
+			const FAudioBusDashboardEntry& AudioBusDashboardEntry = CastEntry(*Entry);
+
+			if (CheckedAudioBuses.Find(AudioBusDashboardEntry.AudioBusId) != nullptr)
+			{
+				SendAudioBusEnvelopeFollowerCVar(true, AudioBusDashboardEntry.AudioBusId);
+				FAudioMeterProvider::OnUpdateAudioMeterInfo.ExecuteIfBound(AudioBusDashboardEntry.AudioBusId, AudioBusDashboardEntry.AudioMeterInfo, AudioBusDashboardEntry.GetDisplayName());
+			}
+		}
+	}
+
 	void FAudioBusDashboardViewFactory::HandleOnAudioBusAdded(const uint32 InAudioBusId)
 	{
 		RequestListRefresh();
@@ -472,6 +468,9 @@ namespace UE::Audio::Insights
 
 	void FAudioBusDashboardViewFactory::HandleOnAudioBusRemoved(const uint32 InAudioBusId)
 	{
+		using namespace FAudioBusDashboardViewFactoryPrivate;
+
+		SendAudioBusEnvelopeFollowerCVar(false, InAudioBusId);
 		FAudioMeterProvider::OnRemoveAudioMeter.ExecuteIfBound(InAudioBusId);
 		RequestListRefresh();
 	}
@@ -482,6 +481,8 @@ namespace UE::Audio::Insights
 		{
 			FilteredEntriesListView->RequestListRefresh();
 		}
+
+		UpdateAudioMeters();
 	}
 } // namespace UE::Audio::Insights
 
diff --git a/Engine/Plugins/AudioInsights/Source/AudioInsights/Private/Views/SubmixDashboardViewFactory.cpp b/Engine/Plugins/AudioInsights/Source/AudioInsights/Private/Views/SubmixDashboardViewFactory.cpp
--- a/Engine/Plugins/AudioInsights/Source/AudioInsights/Private/Views/SubmixDashboardViewFactory.cpp
+++ b/Engine/Plugins/AudioInsights/Source/AudioInsights/Private/Views/SubmixDashboardViewFactory.cpp
@@ -23,6 +23,8 @@ namespace UE::Audio::Insights
 		{
 			return static_cast<const FSubmixDashboardEntry&>(InData);
 		};
+
+		const FSlateColor SubmixMeterTextColor(FColor(30, 240, 90));
 	}
 
 	FSubmixDashboardViewFactory::FSubmixDashboardViewFactory()
@@ -41,20 +43,16 @@ namespace UE::Audio::Insights
 		SortByColumn = "Name";
 		SortMode     = EColumnSortMode::Ascending;
 
-#if WITH_EDITOR
-		SubmixProvider->OnSubmixAssetAdded.BindRaw(this, &FSubmixDashboardViewFactory::RequestListRefresh);
-		SubmixProvider->OnSubmixAssetRemoved.BindRaw(this, &FSubmixDashboardViewFactory::RequestListRefresh);
-		SubmixProvider->OnSubmixAssetListUpdated.BindRaw(this, &FSubmixDashboardViewFactory::RequestListRefresh);
-#endif // WITH_EDITOR
+		SubmixProvider->OnSubmixAdded.BindRaw(this, &FSubmixDashboardViewFactory::HandleOnSubmixAdded);
+		SubmixProvider->OnSubmixRemoved.BindRaw(this, &FSubmixDashboardViewFactory::HandleOnSubmixRemoved);
+		SubmixProvider->OnSubmixListUpdated.BindRaw(this, &FSubmixDashboardViewFactory::RequestListRefresh);
 	}
 
 	FSubmixDashboardViewFactory::~FSubmixDashboardViewFactory()
 	{
-#if WITH_EDITOR
-		SubmixProvider->OnSubmixAssetAdded.Unbind();
-		SubmixProvider->OnSubmixAssetRemoved.Unbind();
-		SubmixProvider->OnSubmixAssetListUpdated.Unbind();
-#endif // WITH_EDITOR
+		SubmixProvider->OnSubmixAdded.Unbind();
+		SubmixProvider->OnSubmixRemoved.Unbind();
+		SubmixProvider->OnSubmixListUpdated.Unbind();
 	}
 
 	FName FSubmixDashboardViewFactory::GetName() const
@@ -114,11 +112,8 @@ namespace UE::Audio::Insights
 			const FSubmixDashboardEntry& SubmixDashboardEntry = CastEntry(*InRowData);
 			const uint32 SubmixId = SubmixDashboardEntry.SubmixId;
 
-			const bool bIsSubmixChecked = SubmixCheckboxCheckedStates.FindOrAdd(SubmixId);
-
-			if (bIsSubmixChecked)
+			if (CheckedSubmixes.Find(SubmixId) != nullptr)
 			{
-				SendSubmixEnvelopeFollowerCVar(bIsSubmixChecked, SubmixDashboardEntry);
 				FAudioMeterProvider::OnUpdateAudioMeterInfo.ExecuteIfBound(SubmixId, SubmixDashboardEntry.AudioMeterInfo, SubmixDashboardEntry.GetDisplayName());
 			}
 
@@ -130,10 +125,7 @@ namespace UE::Audio::Insights
 					SNew(SCheckBox)
 					.IsChecked_Lambda([this, SubmixId]()
 					{
-						const bool* FoundSubmixCheckedState = SubmixCheckboxCheckedStates.Find(SubmixId);
-						const bool bIsSubmixChecked = FoundSubmixCheckedState && *FoundSubmixCheckedState == true;
-
-						return bIsSubmixChecked ? ECheckBoxState::Checked : ECheckBoxState::Unchecked;
+						return CheckedSubmixes.Find(SubmixId) != nullptr ? ECheckBoxState::Checked : ECheckBoxState::Unchecked;
 					})
 					.OnCheckStateChanged_Lambda([this, InRowData](ECheckBoxState NewState)
 					{
@@ -141,19 +133,16 @@ namespace UE::Audio::Insights
 
 						const FSubmixDashboardEntry& SubmixDashboardEntry = CastEntry(*InRowData);
 
-						if (bool* FoundSubmixCheckedState = SubmixCheckboxCheckedStates.Find(SubmixDashboardEntry.SubmixId))
-						{
-							*FoundSubmixCheckedState = bIsChecked;
-						}
-
 						SendSubmixEnvelopeFollowerCVar(bIsChecked, SubmixDashboardEntry);
 
 						if (bIsChecked)
 						{
-							FAudioMeterProvider::OnAddAudioMeter.ExecuteIfBound(SubmixDashboardEntry.SubmixId, SubmixDashboardEntry.AudioMeterInfo, SubmixDashboardEntry.GetDisplayName(), FSlateColor(FColor(30, 240, 90)));
+							CheckedSubmixes.Emplace(SubmixDashboardEntry.SubmixId);
+							FAudioMeterProvider::OnAddAudioMeter.ExecuteIfBound(SubmixDashboardEntry.SubmixId, SubmixDashboardEntry.AudioMeterInfo, SubmixDashboardEntry.GetDisplayName(), SubmixMeterTextColor);
 						}
 						else
 						{
+							CheckedSubmixes.Remove(SubmixDashboardEntry.SubmixId);
 							FAudioMeterProvider::OnRemoveAudioMeter.ExecuteIfBound(SubmixDashboardEntry.SubmixId);
 						}
 					})
@@ -209,8 +198,6 @@ namespace UE::Audio::Insights
 
 	TSharedRef<SWidget> FSubmixDashboardViewFactory::MakeWidget(TSharedRef<SDockTab> OwnerTab, const FSpawnTabArgs& SpawnTabArgs)
 	{
-		using namespace FSubmixDashboardViewFactoryPrivate;
-
 		if (!DashboardWidget.IsValid())
 		{
 			DashboardWidget = FTraceTableDashboardViewFactory::MakeWidget(OwnerTab, SpawnTabArgs);
@@ -231,18 +218,7 @@ namespace UE::Audio::Insights
 #endif // WITH_EDITOR
 
 		// Update audio meters (if any has been checked previously)
-		for (const TSharedPtr<IDashboardDataViewEntry>& Entry : DataViewEntries)
-		{
-			const FSubmixDashboardEntry& SubmixDashboardEntry = CastEntry(*Entry);
-
-			if (const bool* bFoundSubmixCheckboxCheckedState = SubmixCheckboxCheckedStates.Find(SubmixDashboardEntry.SubmixId))
-			{
-				if (*bFoundSubmixCheckboxCheckedState)
-				{
-					FAudioMeterProvider::OnUpdateAudioMeterInfo.ExecuteIfBound(SubmixDashboardEntry.SubmixId, SubmixDashboardEntry.AudioMeterInfo, SubmixDashboardEntry.GetDisplayName());
-				}
-			}
-		}
+		UpdateAudioMeters();
 
 		return DashboardWidget->AsShared();
 	}
@@ -361,6 +337,50 @@ namespace UE::Audio::Insights
 			}
 		}
 	}
+	
+	void FSubmixDashboardViewFactory::UpdateAudioMeters()
+	{
+		using namespace FSubmixDashboardViewFactoryPrivate;
+
+		for (const TSharedPtr<IDashboardDataViewEntry>& Entry : DataViewEntries)
+		{
+			const FSubmixDashboardEntry& SubmixDashboardEntry = CastEntry(*Entry);
+
+			if (CheckedSubmixes.Find(SubmixDashboardEntry.SubmixId) != nullptr)
+			{
+				SendSubmixEnvelopeFollowerCVar(true, SubmixDashboardEntry);
+				FAudioMeterProvider::OnUpdateAudioMeterInfo.ExecuteIfBound(SubmixDashboardEntry.SubmixId, SubmixDashboardEntry.AudioMeterInfo, SubmixDashboardEntry.GetDisplayName());
+			}
+		}
+	}
+
+	void FSubmixDashboardViewFactory::HandleOnSubmixAdded(const uint32 InSubmixId)
+	{
+		RequestListRefresh();
+	}
+
+	void FSubmixDashboardViewFactory::HandleOnSubmixRemoved(const uint32 InSubmixId)
+	{
+		using namespace FSubmixDashboardViewFactoryPrivate;
+
+		const TSharedPtr<IDashboardDataViewEntry>* FoundEntry = DataViewEntries.FindByPredicate([&InSubmixId](const TSharedPtr<IDashboardDataViewEntry>& Entry)
+		{
+			const FSubmixDashboardEntry& SubmixDashboardEntry = CastEntry(*Entry);
+
+			return SubmixDashboardEntry.SubmixId == InSubmixId;
+		});
+
+		if (FoundEntry)
+		{
+			const TSharedPtr<IDashboardDataViewEntry> Entry   = *FoundEntry;
+			const FSubmixDashboardEntry& SubmixDashboardEntry = CastEntry(*Entry);
+
+			SendSubmixEnvelopeFollowerCVar(false, SubmixDashboardEntry);
+		}
+
+		FAudioMeterProvider::OnRemoveAudioMeter.ExecuteIfBound(InSubmixId);
+		RequestListRefresh();
+	}
 
 	void FSubmixDashboardViewFactory::RequestListRefresh()
 	{
@@ -368,6 +388,8 @@ namespace UE::Audio::Insights
 		{
 			FilteredEntriesListView->RequestListRefresh();
 		}
+
+		UpdateAudioMeters();
 	}
 } // namespace UE::Audio::Insights
 
diff --git a/Engine/Plugins/AudioInsights/Source/AudioInsights/Public/Providers/AssetProvider.h b/Engine/Plugins/AudioInsights/Source/AudioInsights/Public/Providers/AssetProvider.h
--- a/Engine/Plugins/AudioInsights/Source/AudioInsights/Public/Providers/AssetProvider.h
+++ b/Engine/Plugins/AudioInsights/Source/AudioInsights/Public/Providers/AssetProvider.h
@@ -14,7 +14,7 @@
 
 namespace UE::Audio::Insights
 {
-	template <typename TAsset, typename TDashboardEntry>
+	template <typename TAsset>
 	class TAssetProvider
 	{
 	public:
@@ -52,19 +52,15 @@ namespace UE::Audio::Insights
 			UpdateAssetNames();
 		}
 
-		DECLARE_DELEGATE_OneParam(FOnAssetAdded, const TSharedPtr<TDashboardEntry>& /*DashboardEntry*/);
+		DECLARE_DELEGATE_OneParam(FOnAssetAdded, const FString& /*AssetPath*/);
 		FOnAssetAdded OnAssetAdded;
 
-		DECLARE_DELEGATE_OneParam(FOnAssetRemoved, const uint32 /*AssetId*/);
+		DECLARE_DELEGATE_OneParam(FOnAssetRemoved, const FString& /*AssetPath*/);
 		FOnAssetRemoved OnAssetRemoved;
 
-		DECLARE_DELEGATE_OneParam(FOnAssetListUpdated, const TArray<TSharedPtr<TDashboardEntry>>& /*DashboardEntries*/);
+		DECLARE_DELEGATE_OneParam(FOnAssetListUpdated, const TArray<FString>& /*AssetPaths*/);
 		FOnAssetListUpdated OnAssetListUpdated;
 
-	protected:
-		UE_API virtual uint32 GetAssetId(const TSharedPtr<TDashboardEntry>& InEntry) const = 0;
-		UE_API virtual void InitAssetEntry(TSharedPtr<TDashboardEntry>& InEntry, const FAssetData& InAssetData) = 0;
-
 	private:
 		void UpdateAssetNames()
 		{
@@ -74,64 +70,48 @@ namespace UE::Audio::Insights
 			FAssetRegistryModule& AssetRegistryModule = FModuleManager::LoadModuleChecked<FAssetRegistryModule>("AssetRegistry");
 			AssetRegistryModule.Get().GetAssetsByClass(FTopLevelAssetPath(TAsset::StaticClass()), AssetDataArray);
 
-			// Build AssetEntries
-			AssetEntries.Reset();
+			// Collect asset paths
+			AssetPaths.Reset();
 
 			constexpr bool bShouldNotify = false;
 			for (const FAssetData& AssetData : AssetDataArray)
 			{
-				AddAsset(AssetData, bShouldNotify);
+				AddAssetPath(AssetData.GetObjectPathString(), bShouldNotify);
 			}
 
-			AssetEntries.Sort([](const TSharedPtr<TDashboardEntry>& A, const TSharedPtr<TDashboardEntry>& B)
+			AssetPaths.Sort([](const FString& A, const FString& B)
 			{
-				return A->GetDisplayName().CompareToCaseIgnored(B->GetDisplayName()) < 0;
+				return A.Compare(B, ESearchCase::IgnoreCase) < 0;
 			});
 
-			OnAssetListUpdated.ExecuteIfBound(AssetEntries);
+			OnAssetListUpdated.ExecuteIfBound(AssetPaths);
 		}
 
-		void AddAsset(const FAssetData& InAssetData, const bool bInShouldNotify)
+		void AddAssetPath(const FString& InAssetPath, const bool bInShouldNotify)
 		{
-			const bool bIsAssetEntryAlreadAdded = AssetEntries.ContainsByPredicate(
-				[AssetName = InAssetData.GetObjectPathString()](const TSharedPtr<TDashboardEntry>& DashboardEntry)
+			const bool bIsAssetPathAlreadAdded = AssetPaths.ContainsByPredicate(
+				[&InAssetPath](const FString& AssetPath)
 				{
-					return DashboardEntry->Name == AssetName;
+					return AssetPath == InAssetPath;
 				});
 
-			if (!bIsAssetEntryAlreadAdded)
+			if (!bIsAssetPathAlreadAdded)
 			{
-				TSharedPtr<TDashboardEntry>& DashboardEntry = AssetEntries.Add_GetRef(MakeShared<TDashboardEntry>());
-
-				InitAssetEntry(DashboardEntry, InAssetData);
+				AssetPaths.Emplace(InAssetPath);
 
 				if (bInShouldNotify)
 				{
-					OnAssetAdded.ExecuteIfBound(DashboardEntry);
+					OnAssetAdded.ExecuteIfBound(InAssetPath);
 				}
 			}
 		}
 
-		void RemoveAsset(const FAssetData& InAssetData)
+		void RemoveAssetPath(const FString& InAssetPath)
 		{
-			const TSharedPtr<TDashboardEntry>* FoundDashboardEntry = AssetEntries.FindByPredicate(
-				[AssetName = InAssetData.GetObjectPathString()](const TSharedPtr<TDashboardEntry>& InDashboardEntry)
-				{
-					return InDashboardEntry->Name == AssetName;
-				});
-
-			if (FoundDashboardEntry)
+			const int32 NumRemovedItems = AssetPaths.Remove(InAssetPath);
+			if (NumRemovedItems > 0)
 			{
-				const TSharedPtr<TDashboardEntry> DashboardEntry = *FoundDashboardEntry;
-
-				if (DashboardEntry.IsValid())
-				{
-					const uint32 AssetId = GetAssetId(DashboardEntry);
-
-					AssetEntries.Remove(DashboardEntry);
-
-					OnAssetRemoved.ExecuteIfBound(AssetId);
-				}
+				OnAssetRemoved.ExecuteIfBound(InAssetPath);
 			}
 		}
 
@@ -140,7 +120,7 @@ namespace UE::Audio::Insights
 			if (bAreFilesLoaded && InAssetData.AssetClassPath == FTopLevelAssetPath(TAsset::StaticClass()))
 			{
 				constexpr bool bShouldNotify = true;
-				AddAsset(InAssetData, bShouldNotify);
+				AddAssetPath(InAssetData.GetObjectPathString(), bShouldNotify);
 			}
 		}
 
@@ -148,7 +128,7 @@ namespace UE::Audio::Insights
 		{
 			if (InAssetData.AssetClassPath == FTopLevelAssetPath(TAsset::StaticClass()))
 			{
-				RemoveAsset(InAssetData);
+				RemoveAssetPath(InAssetData.GetObjectPathString());
 			}
 		}
 
@@ -177,9 +157,8 @@ namespace UE::Audio::Insights
 		}
 
 		bool bAreFilesLoaded = false;
-		bool bAssetEntriesNeedRefreshing = false;
 
-		TArray<TSharedPtr<TDashboardEntry>> AssetEntries;
+		TArray<FString> AssetPaths;
 	};
 } // namespace UE::Audio::Insights
 
diff --git a/Engine/Plugins/AudioInsights/Source/AudioInsights/Public/Views/AudioBusDashboardViewFactory.h b/Engine/Plugins/AudioInsights/Source/AudioInsights/Public/Views/AudioBusDashboardViewFactory.h
--- a/Engine/Plugins/AudioInsights/Source/AudioInsights/Public/Views/AudioBusDashboardViewFactory.h
+++ b/Engine/Plugins/AudioInsights/Source/AudioInsights/Public/Views/AudioBusDashboardViewFactory.h
@@ -36,6 +36,8 @@ namespace UE::Audio::Insights
 		UE_API virtual void ProcessEntries(FTraceTableDashboardViewFactory::EProcessReason Reason) override;
 		UE_API virtual const TMap<FName, FTraceTableDashboardViewFactory::FColumnData>& GetColumns() const override;
 
+		UE_API void UpdateAudioMeters();
+
 		UE_API virtual void SortTable() override;
 
 		UE_API void FilterByAudioBusName();
@@ -47,7 +49,7 @@ namespace UE::Audio::Insights
 
 		TSharedPtr<FAudioBusTraceProvider> AudioBusProvider;
 
-		TMap<uint32, bool> AudioBusCheckboxCheckedStates;
+		TSet<uint32> CheckedAudioBuses;
 
 		using FComboboxSelectionItem = TPair<EAudioBusTypeComboboxSelection, FText>;
 		TArray<TSharedPtr<FComboboxSelectionItem>> AudioBusTypes;
diff --git a/Engine/Plugins/AudioInsights/Source/AudioInsights/Public/Views/SubmixDashboardViewFactory.h b/Engine/Plugins/AudioInsights/Source/AudioInsights/Public/Views/SubmixDashboardViewFactory.h
--- a/Engine/Plugins/AudioInsights/Source/AudioInsights/Public/Views/SubmixDashboardViewFactory.h
+++ b/Engine/Plugins/AudioInsights/Source/AudioInsights/Public/Views/SubmixDashboardViewFactory.h
@@ -30,13 +30,17 @@ namespace UE::Audio::Insights
 		UE_API virtual void ProcessEntries(FTraceTableDashboardViewFactory::EProcessReason Reason) override;
 		UE_API virtual const TMap<FName, FTraceTableDashboardViewFactory::FColumnData>& GetColumns() const override;
 
+		UE_API void UpdateAudioMeters();
+
 		UE_API virtual void SortTable() override;
 
+		void HandleOnSubmixAdded(const uint32 InSubmixId);
+		void HandleOnSubmixRemoved(const uint32 InSubmixId);
 		UE_API void RequestListRefresh();
 
 		TSharedPtr<FSubmixTraceProvider> SubmixProvider;
 
-		TMap<uint32, bool> SubmixCheckboxCheckedStates;
+		TSet<uint32> CheckedSubmixes;
 	};
 } // namespace UE::Audio::Insights
 
