[Audio Insights] Fix for missing plot data in trace files with very little data
* When a .utrace file has very little data inside (such as a single sound playing and stopping a couple of times) we can receive the OnSessionAnalysisComplete callback before we have processed the data inside our providers
* Instead of immediately setting the GameState inside the plots widget to ControlledByTimemarker on analysis complete, only do this when we have detected a time marker change. This avoids data loss if we are still processing loading in data.

diff --git a/Engine/Plugins/AudioInsights/Source/AudioInsights/Private/Views/SoundPlotsWidgetView.cpp b/Engine/Plugins/AudioInsights/Source/AudioInsights/Private/Views/SoundPlotsWidgetView.cpp
--- a/Engine/Plugins/AudioInsights/Source/AudioInsights/Private/Views/SoundPlotsWidgetView.cpp
+++ b/Engine/Plugins/AudioInsights/Source/AudioInsights/Private/Views/SoundPlotsWidgetView.cpp
@@ -273,19 +273,20 @@ namespace UE::Audio::Insights
 
 	void FSoundPlotsWidgetView::OnSessionAnalysisCompleted()
 	{
-		if (GameState == EGameState::LoadingTraceFile && PlotWidgetMetadataPerCurve.IsValid())
+		if (GameState == EGameState::LoadingTraceFile)
 		{
-			for (auto& [SourceID, MetaData] : *PlotWidgetMetadataPerCurve)
+			if (PlotWidgetMetadataPerCurve.IsValid())
 			{
-				MetaData.CurveColor.A = 0.0f;
-			}
+				for (auto& [SourceID, MetaData] : *PlotWidgetMetadataPerCurve)
+				{
+					MetaData.CurveColor.A = 0.0f;
+				}
 
-			if (PlotWidget.IsValid())
-			{
-				PlotWidget->SetCurvesMetadata(PlotWidgetMetadataPerCurve);
+				if (PlotWidget.IsValid())
+				{
+					PlotWidget->SetCurvesMetadata(PlotWidgetMetadataPerCurve);
+				}
 			}
-
-			GameState = EGameState::ControlledByTimemarker;
 		}
 		else
 		{
