[Audio Insights] When pausing caching and processing, make sure the selected row remains focused in the event log as the top events are removed. Only applies if the event was already visible.

diff --git a/Engine/Plugins/AudioInsights/Source/AudioInsights/Private/Views/AudioEventLogDashboardViewFactory.cpp b/Engine/Plugins/AudioInsights/Source/AudioInsights/Private/Views/AudioEventLogDashboardViewFactory.cpp
--- a/Engine/Plugins/AudioInsights/Source/AudioInsights/Private/Views/AudioEventLogDashboardViewFactory.cpp
+++ b/Engine/Plugins/AudioInsights/Source/AudioInsights/Private/Views/AudioEventLogDashboardViewFactory.cpp
@@ -1368,6 +1368,12 @@ namespace UE::Audio::Insights
 		if (TimestampIsInChunkMarkedToDelete(SelectedEventLogEntry.Timestamp))
 		{
 			AudioInsightsModule.GetTimingViewExtender().StopCachingAndProcessingNewMessages();
+
+			if (FilteredEntriesListView->IsItemVisible(SelectedItems[0]))
+			{
+				// Make sure the selected row remains visible if the top of the cache is removed
+				FocusedItem = SelectedItems[0];
+			}
 		}
 #endif // WITH_EDITOR
 	}
