[Audio Insights] Fix for early messages being lost when loading in a trace file
* AnalyzerQueue was using the timestamp sent from the runtime to decide whether the message was beyond the max history time. However, when loading small trace files with sparse data, it is entiredly possible for us to load in messages from a large time range before we process them.
* We now record the time a message entered the queue, and use this to limit the lifetime of messages in the queue. This functions the same when connected to live traces, but also correctly limits loading in trace files.
* The time limit restriction has been moved from Dequeue - the logic inside here was faulty, checking the difference of a message's timestamp with itself. Seeing as this feature is predominantly there to avoid the containers growing too large when adding, it doesn't make sense to use it when emptying the queue.
* Use HistoryLimitSec instead of MaxHistoryLimitSec on enqueue when controlling lifetimes - using MaxHistoryLimitSec makes the passed in limit to the constructor redundant.
* Removed derived classes that were originally created to allow us to grab the timestamps from the messages, these are no longer required.

diff --git a/Engine/Plugins/AudioInsights/Source/AudioInsights/Public/Cache/AudioInsightsCacheManager.h b/Engine/Plugins/AudioInsights/Source/AudioInsights/Public/Cache/AudioInsightsCacheManager.h
--- a/Engine/Plugins/AudioInsights/Source/AudioInsights/Public/Cache/AudioInsightsCacheManager.h
+++ b/Engine/Plugins/AudioInsights/Source/AudioInsights/Public/Cache/AudioInsightsCacheManager.h
@@ -199,7 +199,7 @@ namespace UE::Audio::Insights
 		TraceServices::FSlabAllocator Allocator;
 		TCircularBuffer<TUniquePtr<FAudioCachedMessageChunk>> Cache;
 
-		TCacheMessageQueue<TSharedPtr<IAudioCachedMessage>> NewMessageQueue { 5.0 };
+		TAnalyzerMessageQueue<TSharedPtr<IAudioCachedMessage>> NewMessageQueue { 5.0 };
 
 		uint32 WriteIndex = 0u;
 		uint32 StartIndex = 0u;
diff --git a/Engine/Plugins/AudioInsights/Source/AudioInsights/Public/Messages/AnalyzerMessageQueue.h b/Engine/Plugins/AudioInsights/Source/AudioInsights/Public/Messages/AnalyzerMessageQueue.h
--- a/Engine/Plugins/AudioInsights/Source/AudioInsights/Public/Messages/AnalyzerMessageQueue.h
+++ b/Engine/Plugins/AudioInsights/Source/AudioInsights/Public/Messages/AnalyzerMessageQueue.h
@@ -2,33 +2,24 @@
 #pragma once
 
 #include "Containers/MpscQueue.h"
+#include "HAL/PlatformTime.h"
 #include "Math/UnrealMathUtility.h"
 #include "Misc/ScopeLock.h"
 
 namespace UE::Audio::Insights
 {
 	template <typename T>
-	struct TMessageQueueBase
+	struct TAnalyzerMessageQueue
 	{
 		static constexpr double MaxHistoryLimitSec = 5.0;
 
-		explicit TMessageQueueBase(double InHistoryLimitSec)
+		explicit TAnalyzerMessageQueue(double InHistoryLimitSec)
 			: HistoryLimitSec(FMath::Clamp(InHistoryLimitSec, UE_DOUBLE_KINDA_SMALL_NUMBER, 5.0))
 		{
 		}
 
-		virtual ~TMessageQueueBase() = default;
+		virtual ~TAnalyzerMessageQueue() = default;
 
-	protected:
-		virtual double GetTimestampFromData() const = 0;
-		virtual double GetTimestampFromMessage(const T& Message) const = 0;
-
-		double HistoryLimitSec = 0.0;
-		TMpscQueue<T> Data;
-
-		mutable FCriticalSection AccessCritSec;
-
-	public:
 		bool IsEmpty() const
 		{
 			return Data.IsEmpty();
@@ -43,20 +34,10 @@ namespace UE::Audio::Insights
 
 				if (!Data.IsEmpty())
 				{
-					double FirstTimeStamp = GetTimestampFromData();
-
-					T Message;
-					while (Data.Dequeue(Message))
+					TimestampedMessage TimestampedMessage;
+					while (Data.Dequeue(TimestampedMessage))
 					{
-						if (HistoryLimitSec > GetTimestampFromMessage(Message) - FirstTimeStamp)
-						{
-							Output.Add(MoveTemp(Message));
-						}
-
-						if (!Data.IsEmpty())
-						{
-							FirstTimeStamp = GetTimestampFromData();
-						}
+						Output.Add(MoveTemp(TimestampedMessage.Message));
 					}
 				}
 			}
@@ -66,6 +47,8 @@ namespace UE::Audio::Insights
 
 		void Enqueue(T&& InMessage)
 		{
+			TimestampedMessage NewMessage(MoveTemp(InMessage));
+
 			{
 				// Clear queue if queue time limit hit.
 				// Have to lock as queues do not support
@@ -76,87 +59,39 @@ namespace UE::Audio::Insights
 
 				if (!Data.IsEmpty())
 				{
-					double FirstTimeStamp = GetTimestampFromData();
-					while (GetTimestampFromMessage(InMessage) - FirstTimeStamp > MaxHistoryLimitSec)
+					double FirstTimeStamp = Data.Peek()->TimestampWhenCreated;
+					while ((NewMessage.TimestampWhenCreated - FirstTimeStamp) > HistoryLimitSec)
 					{
 						if (!Data.Dequeue().IsSet() || Data.IsEmpty())
 						{
 							break;
 						}
-						FirstTimeStamp = GetTimestampFromData();
+						FirstTimeStamp = Data.Peek()->TimestampWhenCreated;
 					}
 				}
 			}
 
-			Data.Enqueue(MoveTemp(InMessage));
-		}
-
-		void Empty()
-		{
-			FScopeLock Lock(&AccessCritSec);
-			Data.Empty();
-		}
-	};
-
-	template<typename T>
-	concept TIsValidMessage = requires (T Message)
-	{
-		{ Message.Timestamp } -> std::same_as<double&>;
-	};
-
-	template <TIsValidMessage T>
-	struct TAnalyzerMessageQueue : public TMessageQueueBase<T>
-	{
-		explicit TAnalyzerMessageQueue(double InHistoryLimitSec = TMessageQueueBase<T>::MaxHistoryLimitSec)
-			: TMessageQueueBase<T>(InHistoryLimitSec)
-		{
+			Data.Enqueue(MoveTemp(NewMessage));
 		}
 
 	protected:
-		virtual double GetTimestampFromData() const override
+		struct TimestampedMessage
 		{
-			if (!TMessageQueueBase<T>::Data.IsEmpty())
-			{
-				return TMessageQueueBase<T>::Data.Peek()->Timestamp;
-			}
-
-			return 0.0;
-		}
-
-		virtual double GetTimestampFromMessage(const T& Message) const
-		{
-			return Message.Timestamp;
-		}
-	};
-
-	template<typename T>
-	concept TIsValidMessagePointer = requires (T Message)
-	{
-		{ Message->Timestamp } -> std::same_as<double&>;
-	};
+			TimestampedMessage() = default;
 
-	template <TIsValidMessagePointer T>
-	struct TCacheMessageQueue : public TMessageQueueBase<T>
-	{
-		explicit TCacheMessageQueue(double InHistoryLimitSec = TMessageQueueBase<T>::MaxHistoryLimitSec)
-			: TMessageQueueBase<T>(InHistoryLimitSec)
-		{
-		}
-
-	protected:
-		virtual double GetTimestampFromData() const override
-		{
-			if (!TMessageQueueBase<T>::Data.IsEmpty())
+			TimestampedMessage(T&& InMessage)
+				: Message(MoveTemp(InMessage))
+				, TimestampWhenCreated(FPlatformTime::Seconds())
 			{
-				return (*TMessageQueueBase<T>::Data.Peek())->Timestamp;
 			}
 
-			return 0.0;
-		}
+			T Message;
+			double TimestampWhenCreated = 0.0;
+		};
 
-		virtual double GetTimestampFromMessage(const T& Message) const
-		{
-			return Message->Timestamp;
-		}
+
+		double HistoryLimitSec = 0.0;
+		TMpscQueue<TimestampedMessage> Data;
+		mutable FCriticalSection AccessCritSec;
 	};
 } // namespace UE::Audio::Insights
