[Audio Insights] Fix accidental additional copy of column data when constructing FSoundPlotsWidgetView

diff --git a/Engine/Plugins/AudioInsights/Source/AudioInsights/Private/Views/SoundPlotsWidgetView.cpp b/Engine/Plugins/AudioInsights/Source/AudioInsights/Private/Views/SoundPlotsWidgetView.cpp
--- a/Engine/Plugins/AudioInsights/Source/AudioInsights/Private/Views/SoundPlotsWidgetView.cpp
+++ b/Engine/Plugins/AudioInsights/Source/AudioInsights/Private/Views/SoundPlotsWidgetView.cpp
@@ -26,7 +26,7 @@ namespace UE::Audio::Insights
 	}
 
 	FSoundPlotsWidgetView::FSoundPlotsWidgetView(TMap<FName, FSoundPlotsWidgetView::FPlotColumnInfo>&& InColumnInfo, const TFunction<bool(const IDashboardDataTreeViewEntry& InEntry)> InIsPlotEnabledForEntryFunc)
-		: ColumnInfo(InColumnInfo)
+		: ColumnInfo(MoveTemp(InColumnInfo))
 		, IsPlotEnabledForEntryFunc(InIsPlotEnabledForEntryFunc)
 	{
 		FTraceModule& AudioInsightsTraceModule = static_cast<FTraceModule&>(FAudioInsightsModule::GetChecked().GetTraceModule());
