[Audio Insights] FAudioMetersPanelDashboardViewFactory: don't add a meter if it was already added + don't invalidate AudioMeterWidgetsScrollBox if a meter hasn't been removed

diff --git a/Engine/Plugins/AudioInsights/Source/AudioInsights/Private/Views/AudioMetersPanelDashboardViewFactory.cpp b/Engine/Plugins/AudioInsights/Source/AudioInsights/Private/Views/AudioMetersPanelDashboardViewFactory.cpp
--- a/Engine/Plugins/AudioInsights/Source/AudioInsights/Private/Views/AudioMetersPanelDashboardViewFactory.cpp
+++ b/Engine/Plugins/AudioInsights/Source/AudioInsights/Private/Views/AudioMetersPanelDashboardViewFactory.cpp
@@ -82,6 +82,11 @@ namespace UE::Audio::Insights
 
 	void FAudioMetersPanelDashboardViewFactory::HandleOnAddAudioMeter(const uint32 InEntryId, const TSharedRef<FAudioMeterInfo>& InAudioMeterInfo, const FText& InName, const FSlateColor& InNameColor)
 	{
+		if (AudioMeterViews.Contains(InEntryId))
+		{
+			return;
+		}
+
 		const TSharedRef<SAudioMeterView>& AudioMeterView = AudioMeterViews.Emplace(InEntryId, SNew(SAudioMeterView, InAudioMeterInfo)
 			.NameText(InName)
 			.NameTextColor(InNameColor));
@@ -96,16 +101,20 @@ namespace UE::Audio::Insights
 
 	void FAudioMetersPanelDashboardViewFactory::HandleOnRemoveAudioMeter(const uint32 InEntryId)
 	{
-		if (TSharedRef<SAudioMeterView>* FoundAudioMeterView = AudioMeterViews.Find(InEntryId))
+		const TSharedRef<SAudioMeterView>* FoundAudioMeterView = AudioMeterViews.Find(InEntryId);
+
+		if (FoundAudioMeterView == nullptr)
 		{
-			if (AudioMeterWidgetsContainer.IsValid())
-			{
-				AudioMeterWidgetsContainer->RemoveSlot(*FoundAudioMeterView);
-			}
+			return;
+		}
 
-			AudioMeterViews.Remove(InEntryId);
+		if (AudioMeterWidgetsContainer.IsValid())
+		{
+			AudioMeterWidgetsContainer->RemoveSlot(*FoundAudioMeterView);
 		}
 
+		AudioMeterViews.Remove(InEntryId);
+
 		if (AudioMeterWidgetsScrollBox.IsValid())
 		{
 			AudioMeterWidgetsScrollBox->Invalidate(EInvalidateWidget::Layout);
