[Audio Insights] Added envelope follower enabled message to submix and audio bus message/entry + misc. code adjustments for homogeneity and readability

diff --git a/Engine/Plugins/AudioInsights/Source/AudioInsights/Private/Messages/AudioBusTraceMessages.cpp b/Engine/Plugins/AudioInsights/Source/AudioInsights/Private/Messages/AudioBusTraceMessages.cpp
--- a/Engine/Plugins/AudioInsights/Source/AudioInsights/Private/Messages/AudioBusTraceMessages.cpp
+++ b/Engine/Plugins/AudioInsights/Source/AudioInsights/Private/Messages/AudioBusTraceMessages.cpp
@@ -7,10 +7,12 @@ namespace UE::Audio::Insights
 	{
 		const FName Start = "AudioBusStart";
 		const FName HasActivity = "AudioBusHasActivity";
+		const FName EnvelopeFollowerEnabled = "AudioBusEnvelopeFollowerEnabled";
 		const FName EnvelopeValues = "AudioBusEnvelopeValues";
 		const FName Stop = "AudioBusStop";
 	};
 
+	// FAudioBusMessageBase
 	FAudioBusMessageBase::FAudioBusMessageBase(const Trace::IAnalyzer::FOnEventContext& InContext)
     {
 		const UE::Trace::IAnalyzer::FEventData& EventData = InContext.EventData;
@@ -20,6 +22,7 @@ namespace UE::Audio::Insights
 		Timestamp = InContext.EventTime.AsSeconds(EventData.GetValue<uint64>("Timestamp"));
     }
 
+	// FAudioBusStartMessage
 	FAudioBusStartMessage::FAudioBusStartMessage(const Trace::IAnalyzer::FOnEventContext& InContext)
 		: FAudioBusMessageBase(InContext)
 	{
@@ -42,6 +45,7 @@ namespace UE::Audio::Insights
 		return MemorySize;
 	}
 
+	// FAudioBusHasActivityMessage
 	FAudioBusHasActivityMessage::FAudioBusHasActivityMessage(const Trace::IAnalyzer::FOnEventContext& InContext)
 		: FAudioBusMessageBase(InContext)
 	{
@@ -50,11 +54,16 @@ namespace UE::Audio::Insights
 		bHasActivity = EventData.GetValue<bool>("HasActivity");
 	}
 
-	uint32 FAudioBusHasActivityMessage::GetSizeOf() const
+	// FAudioBusEnvelopeFollowerEnabledMessage
+	FAudioBusEnvelopeFollowerEnabledMessage::FAudioBusEnvelopeFollowerEnabledMessage(const Trace::IAnalyzer::FOnEventContext& InContext)
+		: FAudioBusMessageBase(InContext)
 	{
-		return sizeof(FAudioBusHasActivityMessage);
+		const UE::Trace::IAnalyzer::FEventData& EventData = InContext.EventData;
+
+		bEnvelopeFollowerEnabled = EventData.GetValue<bool>("EnvelopeFollowerEnabled");
 	}
 
+	// FAudioBusEnvelopeValuesMessage
 	FAudioBusEnvelopeValuesMessage::FAudioBusEnvelopeValuesMessage(const Trace::IAnalyzer::FOnEventContext& InContext)
 		: FAudioBusMessageBase(InContext)
 	{
@@ -82,13 +91,44 @@ namespace UE::Audio::Insights
 		return MemorySize;
 	}
 
+	// FAudioBusStopMessage
 	FAudioBusStopMessage::FAudioBusStopMessage(const Trace::IAnalyzer::FOnEventContext& InContext)
 		: FAudioBusMessageBase(InContext)
 	{
 	}
 
-	uint32 FAudioBusStopMessage::GetSizeOf() const
+	// FAudioBusDashboardEntry
+	FAudioBusDashboardEntry::FAudioBusDashboardEntry()
+		: AudioMeterInfo(MakeShared<FAudioMeterInfo>())
 	{
-		return sizeof(FAudioBusStopMessage);
 	}
-}
\ No newline at end of file
+
+	FAudioBusDashboardEntry::FAudioBusDashboardEntry(FAudioBusDashboardEntry& Other)
+		: AudioBusId(Other.AudioBusId)
+		, Timestamp(Other.Timestamp)
+		, AudioMeterInfo(MakeShared<FAudioMeterInfo>(*Other.AudioMeterInfo))
+		, Name(Other.Name)
+		, bHasActivity(Other.bHasActivity)
+		, bEnvelopeFollowerEnabled(Other.bEnvelopeFollowerEnabled)
+		, AudioBusType(Other.AudioBusType)
+	{
+	}
+
+	FAudioBusDashboardEntry& FAudioBusDashboardEntry::operator=(const FAudioBusDashboardEntry& Other)
+	{
+		if (this == &Other)
+		{
+			return *this;
+		}
+
+		AudioBusId = Other.AudioBusId;
+		Timestamp = Other.Timestamp;
+		AudioMeterInfo = MakeShared<FAudioMeterInfo>(*Other.AudioMeterInfo);
+		Name = Other.Name;
+		bHasActivity = Other.bHasActivity;
+		bEnvelopeFollowerEnabled = Other.bEnvelopeFollowerEnabled;
+		AudioBusType = Other.AudioBusType;
+
+		return *this;
+	}
+} // namespace UE::Audio::Insights
diff --git a/Engine/Plugins/AudioInsights/Source/AudioInsights/Private/Messages/SubmixTraceMessages.cpp b/Engine/Plugins/AudioInsights/Source/AudioInsights/Private/Messages/SubmixTraceMessages.cpp
--- a/Engine/Plugins/AudioInsights/Source/AudioInsights/Private/Messages/SubmixTraceMessages.cpp
+++ b/Engine/Plugins/AudioInsights/Source/AudioInsights/Private/Messages/SubmixTraceMessages.cpp
@@ -5,26 +5,29 @@ namespace UE::Audio::Insights
 {
 	namespace SubmixMessageNames
 	{
-		const FName SubmixLoaded = "SubmixLoaded";
-		const FName SubmixIsAlivePing = "SubmixAlivePing";
-		const FName SubmixHasActivity = "SubmixHasActivity";
-		const FName SubmixEnvelopeValues = "SubmixEnvelopeValues";
-		const FName SubmixUnloaded = "SubmixUnloaded";
+		const FName Loaded = "SubmixLoaded";
+		const FName IsAlivePing = "SubmixAlivePing";
+		const FName HasActivity = "SubmixHasActivity";
+		const FName EnvelopeFollowerEnabled = "SubmixEnvelopeFollowerEnabled";
+		const FName EnvelopeValues = "SubmixEnvelopeValues";
+		const FName Unloaded = "SubmixUnloaded";
 	}
 
-	FSubmixMessageBase::FSubmixMessageBase(const UE::Trace::IAnalyzer::FOnEventContext& InContext)
+	// FSubmixMessageBase
+	FSubmixMessageBase::FSubmixMessageBase(const Trace::IAnalyzer::FOnEventContext& InContext)
 	{
-		const UE::Trace::IAnalyzer::FEventData& EventData = InContext.EventData;
+		const Trace::IAnalyzer::FEventData& EventData = InContext.EventData;
 
 		DeviceId = static_cast<::Audio::FDeviceId>(EventData.GetValue<uint32>("DeviceId"));
 		SubmixId = EventData.GetValue<uint32>("SubmixId");
 		Timestamp = InContext.EventTime.AsSeconds(EventData.GetValue<uint64>("Timestamp"));
 	}
 
-	FSubmixLoadedMessage::FSubmixLoadedMessage(const UE::Trace::IAnalyzer::FOnEventContext& InContext)
+	// FSubmixLoadedMessage
+	FSubmixLoadedMessage::FSubmixLoadedMessage(const Trace::IAnalyzer::FOnEventContext& InContext)
 		: FSubmixMessageBase(InContext)
 	{
-		const UE::Trace::IAnalyzer::FEventData& EventData = InContext.EventData;
+		const Trace::IAnalyzer::FEventData& EventData = InContext.EventData;
 
 		EventData.GetString("Name", Name);
 		NumChannels = EventData.GetValue<int32>("NumChannels");
@@ -38,21 +41,33 @@ namespace UE::Audio::Insights
 
 		// Add any dynamically sized members
 		MemorySize += Name.GetAllocatedSize();
+
 		return MemorySize;
 	};
 
-	FSubmixHasActivityMessage::FSubmixHasActivityMessage(const UE::Trace::IAnalyzer::FOnEventContext& InContext)
+	// FSubmixHasActivityMessage
+	FSubmixHasActivityMessage::FSubmixHasActivityMessage(const Trace::IAnalyzer::FOnEventContext& InContext)
 		: FSubmixMessageBase(InContext)
 	{
-		const UE::Trace::IAnalyzer::FEventData& EventData = InContext.EventData;
+		const Trace::IAnalyzer::FEventData& EventData = InContext.EventData;
 
 		bHasActivity = EventData.GetValue<bool>("HasActivity");
 	}
 
-	FSubmixEnvelopeValuesMessage::FSubmixEnvelopeValuesMessage(const UE::Trace::IAnalyzer::FOnEventContext& InContext)
+	// FSubmixEnvelopeFollowerEnabledMessage
+	FSubmixEnvelopeFollowerEnabledMessage::FSubmixEnvelopeFollowerEnabledMessage(const Trace::IAnalyzer::FOnEventContext& InContext)
+		: FSubmixMessageBase(InContext)
+	{
+		const Trace::IAnalyzer::FEventData& EventData = InContext.EventData;
+
+		bEnvelopeFollowerEnabled = EventData.GetValue<bool>("EnvelopeFollowerEnabled");
+	}
+
+	// FSubmixEnvelopeValuesMessage
+	FSubmixEnvelopeValuesMessage::FSubmixEnvelopeValuesMessage(const Trace::IAnalyzer::FOnEventContext& InContext)
 		: FSubmixMessageBase(InContext)
 	{
-		const UE::Trace::IAnalyzer::FEventData& EventData = InContext.EventData;
+		const Trace::IAnalyzer::FEventData& EventData = InContext.EventData;
 
 		const TArrayView<const float> EnvelopeValuesView = EventData.GetArrayView<float>("EnvelopeValues");
 
@@ -72,13 +87,49 @@ namespace UE::Audio::Insights
 		{
 			MemorySize += EnvelopeValues.Num() * sizeof(EnvelopeValues[0]);
 		}
+
 		return MemorySize;
 	};
 
-	FSubmixUnloadedMessage::FSubmixUnloadedMessage(const UE::Trace::IAnalyzer::FOnEventContext& InContext)
+	// FSubmixUnloadedMessage
+	FSubmixUnloadedMessage::FSubmixUnloadedMessage(const Trace::IAnalyzer::FOnEventContext& InContext)
 		: FSubmixMessageBase(InContext)
 	{
 
 	};
 
+	// FSubmixDashboardEntry
+	FSubmixDashboardEntry::FSubmixDashboardEntry()
+		: AudioMeterInfo(MakeShared<FAudioMeterInfo>())
+	{
+	}
+
+	FSubmixDashboardEntry::FSubmixDashboardEntry(FSubmixDashboardEntry& Other)
+		: SubmixId(Other.SubmixId)
+		, Timestamp(Other.Timestamp)
+		, AudioMeterInfo(MakeShared<FAudioMeterInfo>(*Other.AudioMeterInfo))
+		, Name(Other.Name)
+		, bHasActivity(Other.bHasActivity)
+		, bEnvelopeFollowerEnabled(Other.bEnvelopeFollowerEnabled)
+		, bIsMainSubmix(Other.bIsMainSubmix)
+	{
+	}
+
+	FSubmixDashboardEntry& FSubmixDashboardEntry::operator=(const FSubmixDashboardEntry& Other)
+	{
+		if (this == &Other)
+		{
+			return *this;
+		}
+
+		SubmixId = Other.SubmixId;
+		Timestamp = Other.Timestamp;
+		AudioMeterInfo = MakeShared<FAudioMeterInfo>(*Other.AudioMeterInfo);
+		Name = Other.Name;
+		bHasActivity = Other.bHasActivity;
+		bEnvelopeFollowerEnabled = Other.bEnvelopeFollowerEnabled;
+		bIsMainSubmix = Other.bIsMainSubmix;
+
+		return *this;
+	}
 } // namespace UE::Audio::Insights
diff --git a/Engine/Plugins/AudioInsights/Source/AudioInsights/Private/Providers/SubmixTraceProvider.cpp b/Engine/Plugins/AudioInsights/Source/AudioInsights/Private/Providers/SubmixTraceProvider.cpp
--- a/Engine/Plugins/AudioInsights/Source/AudioInsights/Private/Providers/SubmixTraceProvider.cpp
+++ b/Engine/Plugins/AudioInsights/Source/AudioInsights/Private/Providers/SubmixTraceProvider.cpp
@@ -447,7 +447,7 @@ namespace UE::Audio::Insights
 
 #if WITH_EDITOR
 		const FAudioInsightsCacheManager& CacheManager = FAudioInsightsModule::GetChecked().GetCacheManager();
-		CacheManager.IterateTo<FSubmixLoadedMessage>(SubmixMessageNames::SubmixLoaded, TimeMarker, [&ProcessSubmixLoadedMessage](const FSubmixLoadedMessage& SubmixLoadedCachedMessage)
+		CacheManager.IterateTo<FSubmixLoadedMessage>(SubmixMessageNames::Loaded, TimeMarker, [&ProcessSubmixLoadedMessage](const FSubmixLoadedMessage& SubmixLoadedCachedMessage)
 		{
 			ProcessSubmixLoadedMessage(SubmixLoadedCachedMessage);
 		});
@@ -477,7 +477,7 @@ namespace UE::Audio::Insights
 		};
 
 #if WITH_EDITOR
-		CacheManager.IterateTo<FSubmixUnloadedMessage>(SubmixMessageNames::SubmixUnloaded, TimeMarker, [&ProcessSubmixUnloadedMessage](const FSubmixUnloadedMessage& SubmixUnloadedCachedMessage)
+		CacheManager.IterateTo<FSubmixUnloadedMessage>(SubmixMessageNames::Unloaded, TimeMarker, [&ProcessSubmixUnloadedMessage](const FSubmixUnloadedMessage& SubmixUnloadedCachedMessage)
 		{
 			ProcessSubmixUnloadedMessage(SubmixUnloadedCachedMessage);
 		});
@@ -520,7 +520,7 @@ namespace UE::Audio::Insights
 
 				// HasActivity
 #if WITH_EDITOR
-				const FSubmixHasActivityMessage* FoundHasActivityCachedMessage = CacheManager.FindClosestMessage<FSubmixHasActivityMessage>(SubmixMessageNames::SubmixHasActivity, TimeMarker, SubmixId);
+				const FSubmixHasActivityMessage* FoundHasActivityCachedMessage = CacheManager.FindClosestMessage<FSubmixHasActivityMessage>(SubmixMessageNames::HasActivity, TimeMarker, SubmixId);
 #else
 				const FSubmixHasActivityMessage* FoundHasActivityCachedMessage = FindClosestMessageToTimestamp(SessionCachedMessages->HasActivityCachedMessages, TimeMarker, SubmixId);
 #endif // WITH_EDITOR
@@ -531,7 +531,7 @@ namespace UE::Audio::Insights
 
 				// EnvelopeValues
 #if WITH_EDITOR
-				const FSubmixEnvelopeValuesMessage* FoundEnvelopeValuesCachedMessage = CacheManager.FindClosestMessage<FSubmixEnvelopeValuesMessage>(SubmixMessageNames::SubmixEnvelopeValues, TimeMarker, SubmixId);
+				const FSubmixEnvelopeValuesMessage* FoundEnvelopeValuesCachedMessage = CacheManager.FindClosestMessage<FSubmixEnvelopeValuesMessage>(SubmixMessageNames::EnvelopeValues, TimeMarker, SubmixId);
 #else
 				const FSubmixEnvelopeValuesMessage* FoundEnvelopeValuesCachedMessage = FindClosestMessageToTimestamp(SessionCachedMessages->EnvelopeValuesCachedMessages, TimeMarker, SubmixId);
 #endif // WITH_EDITOR
diff --git a/Engine/Plugins/AudioInsights/Source/AudioInsights/Public/Messages/AudioBusTraceMessages.h b/Engine/Plugins/AudioInsights/Source/AudioInsights/Public/Messages/AudioBusTraceMessages.h
--- a/Engine/Plugins/AudioInsights/Source/AudioInsights/Public/Messages/AudioBusTraceMessages.h
+++ b/Engine/Plugins/AudioInsights/Source/AudioInsights/Public/Messages/AudioBusTraceMessages.h
@@ -18,6 +18,7 @@ namespace UE::Audio::Insights
 	{
 		extern const FName Start;
 		extern const FName HasActivity;
+		extern const FName EnvelopeFollowerEnabled;
 		extern const FName EnvelopeValues;
 		extern const FName Stop;
 	};
@@ -58,12 +59,23 @@ namespace UE::Audio::Insights
 		FAudioBusHasActivityMessage() = default;
 		FAudioBusHasActivityMessage(const UE::Trace::IAnalyzer::FOnEventContext& InContext);
 
-		virtual uint32 GetSizeOf() const override;
+		virtual uint32 GetSizeOf() const override { return sizeof(FAudioBusHasActivityMessage); };
 		virtual const FName GetMessageName() const override { return AudioBusMessageNames::HasActivity; }
 
 		bool bHasActivity = false;
 	};
 
+	struct FAudioBusEnvelopeFollowerEnabledMessage : public FAudioBusMessageBase
+	{
+		FAudioBusEnvelopeFollowerEnabledMessage() = default;
+		FAudioBusEnvelopeFollowerEnabledMessage(const UE::Trace::IAnalyzer::FOnEventContext& InContext);
+
+		virtual uint32 GetSizeOf() const override { return sizeof(FAudioBusEnvelopeFollowerEnabledMessage); };
+		virtual const FName GetMessageName() const override { return AudioBusMessageNames::EnvelopeFollowerEnabled; }
+
+		bool bEnvelopeFollowerEnabled = false;
+	};
+
 	struct FAudioBusEnvelopeValuesMessage : public FAudioBusMessageBase
 	{
 		FAudioBusEnvelopeValuesMessage() = default;
@@ -80,7 +92,7 @@ namespace UE::Audio::Insights
 		FAudioBusStopMessage() = default;
 		FAudioBusStopMessage(const UE::Trace::IAnalyzer::FOnEventContext& InContext);
 
-		virtual uint32 GetSizeOf() const override;
+		virtual uint32 GetSizeOf() const override { return sizeof(FAudioBusStopMessage); };
 		virtual const FName GetMessageName() const override { return AudioBusMessageNames::Stop; }
 	};
 
@@ -88,6 +100,7 @@ namespace UE::Audio::Insights
 	{
 		TAnalyzerMessageQueue<FAudioBusStartMessage> StartMessages{ 0.1 };
 		TAnalyzerMessageQueue<FAudioBusHasActivityMessage> HasActivityMessages{ 0.1 };
+		TAnalyzerMessageQueue<FAudioBusEnvelopeFollowerEnabledMessage> EnvelopeFollowerEnabledMessages{ 0.1 };
 		TAnalyzerMessageQueue<FAudioBusEnvelopeValuesMessage> EnvelopeValuesMessages{ 0.1 };
 		TAnalyzerMessageQueue<FAudioBusStopMessage> StopMessages{ 0.1 };
 
@@ -100,6 +113,7 @@ namespace UE::Audio::Insights
 		FAudioBusSessionCachedMessages(TraceServices::IAnalysisSession& InSession)
 			: StartCachedMessages(InSession.GetLinearAllocator(), 4096)
 			, HasActivityCachedMessages(InSession.GetLinearAllocator(), 4096)
+			, EnvelopeFollowerEnabledCachedMessages(InSession.GetLinearAllocator(), 4096)
 			, EnvelopeValuesCachedMessages(InSession.GetLinearAllocator(), 4096)
 			, StopCachedMessages(InSession.GetLinearAllocator(), 4096)
 		{
@@ -108,6 +122,7 @@ namespace UE::Audio::Insights
 
 		TraceServices::TPagedArray<FAudioBusStartMessage> StartCachedMessages;
 		TraceServices::TPagedArray<FAudioBusHasActivityMessage> HasActivityCachedMessages;
+		TraceServices::TPagedArray<FAudioBusEnvelopeFollowerEnabledMessage> EnvelopeFollowerEnabledCachedMessages;
 		TraceServices::TPagedArray<FAudioBusEnvelopeValuesMessage> EnvelopeValuesCachedMessages;
 		TraceServices::TPagedArray<FAudioBusStopMessage> StopCachedMessages;
 	};
@@ -115,41 +130,12 @@ namespace UE::Audio::Insights
 
 	struct FAudioBusDashboardEntry : public IObjectDashboardEntry
 	{
-		FAudioBusDashboardEntry()
-			: AudioMeterInfo(MakeShared<FAudioMeterInfo>())
-		{
-
-		}
-
-		// Copy constructor - make sure we take a deep copy of FAudioMeterInfo
-		FAudioBusDashboardEntry(FAudioBusDashboardEntry& Other)
-			: AudioBusId(Other.AudioBusId)
-			, Timestamp(Other.Timestamp)
-			, AudioMeterInfo(MakeShared<FAudioMeterInfo>(*Other.AudioMeterInfo))
-			, Name(Other.Name)
-			, bHasActivity(Other.bHasActivity)
-			, AudioBusType(Other.AudioBusType)
-		{
-		}
+		FAudioBusDashboardEntry();
+		FAudioBusDashboardEntry(FAudioBusDashboardEntry& Other); // // Copy constructor - make sure we take a deep copy of FAudioMeterInfo
 
 		virtual ~FAudioBusDashboardEntry() = default;
 
-		FAudioBusDashboardEntry& operator=(const FAudioBusDashboardEntry& Other)
-		{
-			if (this == &Other)
-			{
-				return *this;
-			}
-
-			AudioBusId = Other.AudioBusId;
-			Timestamp = Other.Timestamp;
-			AudioMeterInfo = MakeShared<FAudioMeterInfo>(*Other.AudioMeterInfo);
-			Name = Other.Name;
-			bHasActivity = Other.bHasActivity;
-			AudioBusType = Other.AudioBusType;
-
-			return *this;
-		}
+		FAudioBusDashboardEntry& operator=(const FAudioBusDashboardEntry& Other);
 
 		virtual FText GetDisplayName() const override { return FText::FromString(AudioBusType == EAudioBusType::AssetBased ? FSoftObjectPath(Name).GetAssetName() : Name); }
 		virtual const UObject* GetObject() const override { return FSoftObjectPath(Name).ResolveObject(); }
@@ -163,7 +149,7 @@ namespace UE::Audio::Insights
 
 		FString Name;
 		bool bHasActivity = false;
-
+		bool bEnvelopeFollowerEnabled = false;
 		EAudioBusType AudioBusType = EAudioBusType::None;
 	};
 } // namespace UE::Audio::Insights
diff --git a/Engine/Plugins/AudioInsights/Source/AudioInsights/Public/Messages/SubmixTraceMessages.h b/Engine/Plugins/AudioInsights/Source/AudioInsights/Public/Messages/SubmixTraceMessages.h
--- a/Engine/Plugins/AudioInsights/Source/AudioInsights/Public/Messages/SubmixTraceMessages.h
+++ b/Engine/Plugins/AudioInsights/Source/AudioInsights/Public/Messages/SubmixTraceMessages.h
@@ -4,9 +4,6 @@
 #include "Cache/IAudioCachedMessage.h"
 #include "Providers/AudioMeterProvider.h"
 #include "Messages/AnalyzerMessageQueue.h"
-#if WITH_EDITOR
-#include "Sound/AudioSettings.h"
-#endif // WITH_EDITOR
 #include "Trace/Analyzer.h"
 #include "Views/TableDashboardViewFactory.h"
 
@@ -19,17 +16,18 @@ namespace UE::Audio::Insights
 {
 	namespace SubmixMessageNames
 	{
-		extern const FName SubmixLoaded;
-		extern const FName SubmixIsAlivePing;
-		extern const FName SubmixHasActivity;
-		extern const FName SubmixEnvelopeValues;
-		extern const FName SubmixUnloaded;
+		extern const FName Loaded;
+		extern const FName IsAlivePing;
+		extern const FName HasActivity;
+		extern const FName EnvelopeFollowerEnabled;
+		extern const FName EnvelopeValues;
+		extern const FName Unloaded;
 	};
 
 	struct FSubmixMessageBase : public IAudioCachedMessage
 	{
 		FSubmixMessageBase() = default;
-		FSubmixMessageBase(const UE::Trace::IAnalyzer::FOnEventContext& InContext);
+		FSubmixMessageBase(const Trace::IAnalyzer::FOnEventContext& InContext);
 
 		virtual uint32 GetID() const override { return SubmixId; };
 		virtual uint32 GetSizeOf() const override { return sizeof(FSubmixMessageBase); };
@@ -41,10 +39,10 @@ namespace UE::Audio::Insights
 	struct FSubmixLoadedMessage : public FSubmixMessageBase
 	{
 		FSubmixLoadedMessage() = default;
-		FSubmixLoadedMessage(const UE::Trace::IAnalyzer::FOnEventContext& InContext);
+		FSubmixLoadedMessage(const Trace::IAnalyzer::FOnEventContext& InContext);
 
 		virtual uint32 GetSizeOf() const override;
-		virtual const FName GetMessageName() const override { return SubmixMessageNames::SubmixLoaded; }
+		virtual const FName GetMessageName() const override { return SubmixMessageNames::Loaded; }
 
 		FString Name;
 		int32 NumChannels = 0;
@@ -54,21 +52,32 @@ namespace UE::Audio::Insights
 	struct FSubmixHasActivityMessage : public FSubmixMessageBase
 	{
 		FSubmixHasActivityMessage() = default;
-		FSubmixHasActivityMessage(const UE::Trace::IAnalyzer::FOnEventContext& InContext);
+		FSubmixHasActivityMessage(const Trace::IAnalyzer::FOnEventContext& InContext);
 
 		virtual uint32 GetSizeOf() const override { return sizeof(FSubmixHasActivityMessage); };
-		virtual const FName GetMessageName() const override { return SubmixMessageNames::SubmixHasActivity; }
+		virtual const FName GetMessageName() const override { return SubmixMessageNames::HasActivity; }
 
 		bool bHasActivity = false;
 	};
 
+	struct FSubmixEnvelopeFollowerEnabledMessage : public FSubmixMessageBase
+	{
+		FSubmixEnvelopeFollowerEnabledMessage() = default;
+		FSubmixEnvelopeFollowerEnabledMessage(const Trace::IAnalyzer::FOnEventContext& InContext);
+
+		virtual uint32 GetSizeOf() const override { return sizeof(FSubmixEnvelopeFollowerEnabledMessage); };
+		virtual const FName GetMessageName() const override { return SubmixMessageNames::EnvelopeFollowerEnabled; }
+
+		bool bEnvelopeFollowerEnabled = false;
+	};
+
 	struct FSubmixEnvelopeValuesMessage : public FSubmixMessageBase
 	{
 		FSubmixEnvelopeValuesMessage() = default;
-		FSubmixEnvelopeValuesMessage(const UE::Trace::IAnalyzer::FOnEventContext& InContext);
+		FSubmixEnvelopeValuesMessage(const Trace::IAnalyzer::FOnEventContext& InContext);
 
 		virtual uint32 GetSizeOf() const override;
-		virtual const FName GetMessageName() const override { return SubmixMessageNames::SubmixEnvelopeValues; }
+		virtual const FName GetMessageName() const override { return SubmixMessageNames::EnvelopeValues; }
 
 		TArray<float> EnvelopeValues;
 	};
@@ -76,20 +85,20 @@ namespace UE::Audio::Insights
 	struct FSubmixUnloadedMessage : public FSubmixMessageBase
 	{
 		FSubmixUnloadedMessage() = default;
-		FSubmixUnloadedMessage(const UE::Trace::IAnalyzer::FOnEventContext& InContext);
+		FSubmixUnloadedMessage(const Trace::IAnalyzer::FOnEventContext& InContext);
 
 		virtual uint32 GetSizeOf() const override { return sizeof(FSubmixUnloadedMessage); };
-		virtual const FName GetMessageName() const override { return SubmixMessageNames::SubmixUnloaded; }
+		virtual const FName GetMessageName() const override { return SubmixMessageNames::Unloaded; }
 
 	};
 
-
 	class FSubmixMessages
 	{
-		UE::Audio::Insights::TAnalyzerMessageQueue<FSubmixLoadedMessage> LoadedMessages{ 0.1 };
-		UE::Audio::Insights::TAnalyzerMessageQueue<FSubmixHasActivityMessage> HasActivityMessages{ 0.1 };
-		UE::Audio::Insights::TAnalyzerMessageQueue<FSubmixEnvelopeValuesMessage> EnvelopeValuesMessages{ 0.1 };
-		UE::Audio::Insights::TAnalyzerMessageQueue<FSubmixUnloadedMessage> UnloadedMessages{ 0.1 };
+		TAnalyzerMessageQueue<FSubmixLoadedMessage> LoadedMessages{ 0.1 };
+		TAnalyzerMessageQueue<FSubmixHasActivityMessage> HasActivityMessages{ 0.1 };
+		TAnalyzerMessageQueue<FSubmixEnvelopeFollowerEnabledMessage> EnvelopeFollowerEnabledMessages{ 0.1 };
+		TAnalyzerMessageQueue<FSubmixEnvelopeValuesMessage> EnvelopeValuesMessages{ 0.1 };
+		TAnalyzerMessageQueue<FSubmixUnloadedMessage> UnloadedMessages{ 0.1 };
 
 		friend class FSubmixTraceProvider;
 	};
@@ -100,6 +109,7 @@ namespace UE::Audio::Insights
 		FSubmixSessionCachedMessages(TraceServices::IAnalysisSession& InSession)
 			: LoadedCachedMessages(InSession.GetLinearAllocator(), 4096)
 			, HasActivityCachedMessages(InSession.GetLinearAllocator(), 4096)
+			, EnvelopeFollowerEnabledCachedMessages(InSession.GetLinearAllocator(), 4096)
 			, EnvelopeValuesCachedMessages(InSession.GetLinearAllocator(), 4096)
 			, UnloadedCachedMessages(InSession.GetLinearAllocator(), 4096)
 		{
@@ -108,6 +118,7 @@ namespace UE::Audio::Insights
 
 		TraceServices::TPagedArray<FSubmixLoadedMessage> LoadedCachedMessages;
 		TraceServices::TPagedArray<FSubmixHasActivityMessage> HasActivityCachedMessages;
+		TraceServices::TPagedArray<FSubmixEnvelopeFollowerEnabledMessage> EnvelopeFollowerEnabledCachedMessages;
 		TraceServices::TPagedArray<FSubmixEnvelopeValuesMessage> EnvelopeValuesCachedMessages;
 		TraceServices::TPagedArray<FSubmixUnloadedMessage> UnloadedCachedMessages;
 	};
@@ -115,51 +126,19 @@ namespace UE::Audio::Insights
 
 	struct FSubmixDashboardEntry : public IObjectDashboardEntry
 	{
-		FSubmixDashboardEntry()
-			: AudioMeterInfo(MakeShared<FAudioMeterInfo>())
-		{
-
-		}
-
-		// Copy constructor - make sure we take a deep copy of FAudioMeterInfo
-		FSubmixDashboardEntry(FSubmixDashboardEntry& Other)
-			: SubmixId(Other.SubmixId)
-			, Timestamp(Other.Timestamp)
-			, AudioMeterInfo(MakeShared<FAudioMeterInfo>(*Other.AudioMeterInfo))
-			, Name(Other.Name)
-			, bHasActivity(Other.bHasActivity)
-			, bIsMainSubmix(Other.bIsMainSubmix)
-		{
-		}
+		FSubmixDashboardEntry();
+		FSubmixDashboardEntry(FSubmixDashboardEntry& Other); // Copy constructor - make sure we take a deep copy of FAudioMeterInfo
 
 		virtual ~FSubmixDashboardEntry() = default;
 
-		FSubmixDashboardEntry& operator=(const FSubmixDashboardEntry& Other)
-		{
-			if (this == &Other)
-			{
-				return *this;
-			}
-
-			SubmixId = Other.SubmixId;
-			Timestamp = Other.Timestamp;
-			AudioMeterInfo = MakeShared<FAudioMeterInfo>(*Other.AudioMeterInfo);
-			Name = Other.Name;
-			bHasActivity = Other.bHasActivity;
-			bIsMainSubmix = Other.bIsMainSubmix;
-
-			return *this;
-		}
+		FSubmixDashboardEntry& operator=(const FSubmixDashboardEntry& Other);
 
 		virtual FText GetDisplayName() const override { return FText::FromString(FSoftObjectPath(Name).GetAssetName()); }
 		virtual const UObject* GetObject() const override { return FSoftObjectPath(Name).ResolveObject(); }
 		virtual UObject* GetObject() override { return FSoftObjectPath(Name).ResolveObject(); }
 		virtual bool IsValid() const override { return SubmixId != static_cast<uint32>(INDEX_NONE); }
 
-		bool IsMainSubmix() const
-		{
-			return bIsMainSubmix;
-		}
+		bool IsMainSubmix() const { return bIsMainSubmix; }
 
 		uint32 SubmixId = INDEX_NONE;
 		double Timestamp = 0.0;
@@ -168,6 +147,7 @@ namespace UE::Audio::Insights
 
 		FString Name;
 		bool bHasActivity = false;
+		bool bEnvelopeFollowerEnabled = false;
 		bool bIsMainSubmix = false;
 	};
 } // namespace UE::Audio::Insights
