[Audio Insights] Fix for missing plots data when scrubbing through a busy scene
* Caused by some parameter data sets being bigger than the maximum size of the buffers inside sound dashboard entries, causing us to lose the end of the data.
* Max increased to accomodate larger data sets

#FailedValidations Preflight

diff --git a/Engine/Plugins/AudioInsights/Source/AudioInsights/Private/Providers/SoundTraceProvider.cpp b/Engine/Plugins/AudioInsights/Source/AudioInsights/Private/Providers/SoundTraceProvider.cpp
--- a/Engine/Plugins/AudioInsights/Source/AudioInsights/Private/Providers/SoundTraceProvider.cpp
+++ b/Engine/Plugins/AudioInsights/Source/AudioInsights/Private/Providers/SoundTraceProvider.cpp
@@ -1410,7 +1410,7 @@ namespace UE::Audio::Insights
 
 			if (ActiveSoundEntry.IsValid())
 			{
-				constexpr uint32 DataPointsCapacity = 512u;
+				constexpr uint32 DataPointsCapacity = 2048u;
 				ActiveSoundEntry->ResetDataBuffers(DataPointsCapacity);
 
 				FoundEntries.Add(InActiveSoundPlayOrderID, ActiveSoundEntry);
