[Audio Insights] Make sure we draw the plot for a sound when enabling the plot button, including when we have paused on a specific timestamp
* We were not telling the plots widget to update after enabling a plot as this is usually done in PropcessEntries. Because this isn't called when scrubbing, we need to force an update.

#FailedValidations Preflight

diff --git a/Engine/Plugins/AudioInsights/Source/AudioInsights/Private/Views/SoundDashboardViewFactory.cpp b/Engine/Plugins/AudioInsights/Source/AudioInsights/Private/Views/SoundDashboardViewFactory.cpp
--- a/Engine/Plugins/AudioInsights/Source/AudioInsights/Private/Views/SoundDashboardViewFactory.cpp
+++ b/Engine/Plugins/AudioInsights/Source/AudioInsights/Private/Views/SoundDashboardViewFactory.cpp
@@ -3135,6 +3135,11 @@ namespace UE::Audio::Insights
 			}
 		}
 
+		if (FilteredEntriesListView.IsValid())
+		{
+			OnProcessPlotData.Broadcast(DataViewEntries, FilteredEntriesListView->GetSelectedItems());
+		}
+
 		OnUpdatePlotVisibility.Broadcast(DataViewEntries);
 	}
 
