[Audio Insights] Fix for plots not showing data when scrubbed back and forth after stopping monitoring
* Bug caused by resetting plots when we stop tracing, whereas most dashboards have moved to reset when tracing starts. Moving ResetPlots to OnAnalysisStarting allows us to scrub back through the captured monitoring session.
* Additional logic added to improve the behaviour of the Plot's timeline and make it more natural - the widget will now pause after tracing has finished, and will not start moving forward in time after scrubbing if PIE is paused or tracing is inactive

#FailedValidations Preflight

diff --git a/Engine/Plugins/AudioInsights/Source/AudioInsights/Private/Views/SoundPlotsWidgetView.cpp b/Engine/Plugins/AudioInsights/Source/AudioInsights/Private/Views/SoundPlotsWidgetView.cpp
--- a/Engine/Plugins/AudioInsights/Source/AudioInsights/Private/Views/SoundPlotsWidgetView.cpp
+++ b/Engine/Plugins/AudioInsights/Source/AudioInsights/Private/Views/SoundPlotsWidgetView.cpp
@@ -229,8 +229,12 @@ namespace UE::Audio::Insights
 
 	void FSoundPlotsWidgetView::OnAnalysisStarting(const double Timestamp)
 	{
+		ResetPlots();
+
 #if WITH_EDITOR
 		BeginTimestamp = Timestamp - GStartTime;
+		GameState = EGameState::Running;
+		bIsTraceActive = true;
 #else
 		BeginTimestamp = 0.0;
 #endif // WITH_EDITOR
@@ -240,26 +244,31 @@ namespace UE::Audio::Insights
 	void FSoundPlotsWidgetView::OnPIEStarted(bool bSimulating)
 	{
 		GameState = EGameState::Running;
+		bIsPIEPaused = false;
 	}
 
 	void FSoundPlotsWidgetView::OnPIEStopped(bool bSimulating)
 	{
 		GameState = EGameState::Stopped;
+		bIsPIEPaused = false;
 	}
 
 	void FSoundPlotsWidgetView::OnPIEPaused(bool bSimulating)
 	{
 		GameState = EGameState::Paused;
+		bIsPIEPaused = true;
 	}
 
 	void FSoundPlotsWidgetView::OnPIEResumed(bool bSimulating)
 	{
 		GameState = EGameState::Running;
+		bIsPIEPaused = false;
 	}
 
 	void FSoundPlotsWidgetView::OnTraceStopped(FTraceAuxiliary::EConnectionType InTraceType, const FString& InTraceDestination)
 	{
-		ResetPlots();
+		GameState = EGameState::Paused;
+		bIsTraceActive = false;
 	}
 #else
 	void FSoundPlotsWidgetView::OnAudioInsightsComponentTabSpawn()
@@ -307,7 +316,7 @@ namespace UE::Audio::Insights
 		bTimeMarkerMoved = false;
 
 #if WITH_EDITOR
-		GameState = EGameState::Running;
+		GameState = bIsTraceActive && !bIsPIEPaused ? EGameState::Running : EGameState::Paused;
 #else
 		const TSharedPtr<const FAudioInsightsComponent> AudioInsightsComponent = FAudioInsightsModule::GetChecked().GetAudioInsightsComponent();
 		if (AudioInsightsComponent.IsValid() && AudioInsightsComponent->GetIsLiveSession())
diff --git a/Engine/Plugins/AudioInsights/Source/AudioInsights/Public/Views/SoundPlotsWidgetView.h b/Engine/Plugins/AudioInsights/Source/AudioInsights/Public/Views/SoundPlotsWidgetView.h
--- a/Engine/Plugins/AudioInsights/Source/AudioInsights/Public/Views/SoundPlotsWidgetView.h
+++ b/Engine/Plugins/AudioInsights/Source/AudioInsights/Public/Views/SoundPlotsWidgetView.h
@@ -133,6 +133,11 @@ namespace UE::Audio::Insights
 #endif // !WITH_EDITOR
 
 		bool bTimeMarkerMoved = false;
+
+#if WITH_EDITOR
+		bool bIsTraceActive = false;
+		bool bIsPIEPaused = false;
+#endif // WITH_EDITOR
 	};
 } // namespace UE::Audio::Insights
 
