[Audio Insights] Added envelope follower enabled logic in submix/audio bus trace providers + added OnTimeMarkerUpdated delegate so the dashboards get notified when a timeline scrub with its associated data update has been executed

diff --git a/Engine/Plugins/AudioInsights/Source/AudioInsights/Private/Providers/AudioBusTraceProvider.cpp b/Engine/Plugins/AudioInsights/Source/AudioInsights/Private/Providers/AudioBusTraceProvider.cpp
--- a/Engine/Plugins/AudioInsights/Source/AudioInsights/Private/Providers/AudioBusTraceProvider.cpp
+++ b/Engine/Plugins/AudioInsights/Source/AudioInsights/Private/Providers/AudioBusTraceProvider.cpp
@@ -96,10 +96,11 @@ namespace UE::Audio::Insights
 
 				UE::Trace::IAnalyzer::FInterfaceBuilder& Builder = Context.InterfaceBuilder;
 
-				Builder.RouteEvent(RouteId_Start,          "Audio", "AudioBusStart");
-				Builder.RouteEvent(RouteId_HasActivity,    "Audio", "AudioBusHasActivity");
-				Builder.RouteEvent(RouteId_EnvelopeValues, "Audio", "AudioBusEnvelopeValues");
-				Builder.RouteEvent(RouteId_Stop,           "Audio", "AudioBusStop");
+				Builder.RouteEvent(RouteId_Start,                   "Audio", "AudioBusStart");
+				Builder.RouteEvent(RouteId_HasActivity,             "Audio", "AudioBusHasActivity");
+				Builder.RouteEvent(RouteId_EnvelopeFollowerEnabled, "Audio", "AudioBusEnvelopeFollowerEnabled");
+				Builder.RouteEvent(RouteId_EnvelopeValues,          "Audio", "AudioBusEnvelopeValues");
+				Builder.RouteEvent(RouteId_Stop,                    "Audio", "AudioBusStop");
 			}
 
 			virtual bool OnHandleEvent(uint16 RouteId, EStyle Style, const FOnEventContext& Context) override
@@ -130,6 +131,16 @@ namespace UE::Audio::Insights
 						break;
 					}
 
+					case RouteId_EnvelopeFollowerEnabled:
+					{
+#if WITH_EDITOR
+						CacheMessage<FAudioBusEnvelopeFollowerEnabledMessage>(Context, Messages.EnvelopeFollowerEnabledMessages);
+#else
+						Messages.EnvelopeFollowerEnabledMessages.Enqueue(FAudioBusEnvelopeFollowerEnabledMessage{ Context });
+#endif // WITH_EDITOR
+						break;
+					}
+
 					case RouteId_EnvelopeValues:
 					{
 #if WITH_EDITOR
@@ -171,6 +182,7 @@ namespace UE::Audio::Insights
 			{
 				RouteId_Start,
 				RouteId_HasActivity,
+				RouteId_EnvelopeFollowerEnabled,
 				RouteId_EnvelopeValues,
 				RouteId_Stop
 			};
@@ -259,6 +271,24 @@ namespace UE::Audio::Insights
 			}
 		});
 
+		ProcessMessageQueue<FAudioBusEnvelopeFollowerEnabledMessage>(TraceMessages.EnvelopeFollowerEnabledMessages, GetEntry,
+		[this](const FAudioBusEnvelopeFollowerEnabledMessage& Msg, TSharedPtr<FAudioBusDashboardEntry>* OutEntry)
+		{
+#if !WITH_EDITOR
+			if (SessionCachedMessages.IsValid())
+			{
+				SessionCachedMessages->EnvelopeFollowerEnabledCachedMessages.EmplaceBack(Msg);
+			}
+#endif // !WITH_EDITOR
+
+			if (OutEntry)
+			{
+				FAudioBusDashboardEntry& EntryRef = *OutEntry->Get();
+				EntryRef.bEnvelopeFollowerEnabled = Msg.bEnvelopeFollowerEnabled;
+				EntryRef.Timestamp = Msg.Timestamp;
+			}
+		});
+
 		ProcessMessageQueue<FAudioBusEnvelopeValuesMessage>(TraceMessages.EnvelopeValuesMessages, GetEntry,
 		[this](const FAudioBusEnvelopeValuesMessage& Msg, TSharedPtr<FAudioBusDashboardEntry>* OutEntry)
 		{
@@ -514,6 +544,7 @@ namespace UE::Audio::Insights
 			struct CachedEntryInfo
 			{
 				TOptional<FAudioBusHasActivityMessage> HasActivityMessage;
+				TOptional<FAudioBusEnvelopeFollowerEnabledMessage> EnvelopeFollowerEnabledMessage;
 				TOptional<FAudioBusEnvelopeValuesMessage> EnvelopeValuesMessage;
 			};
 
@@ -543,6 +574,17 @@ namespace UE::Audio::Insights
 					CachedEntryInfos[Index].HasActivityMessage = *FoundHasActivityCachedMessage;
 				}
 
+				// EnvelopeFollowerEnabled
+#if WITH_EDITOR
+				const FAudioBusEnvelopeFollowerEnabledMessage* FoundEnvelopeFollowerEnabledCachedMessage = CacheManager.FindClosestMessage<FAudioBusEnvelopeFollowerEnabledMessage>(AudioBusMessageNames::EnvelopeFollowerEnabled, TimeMarker, AudioBusId);
+#else
+				const FAudioBusEnvelopeFollowerEnabledMessage* FoundEnvelopeFollowerEnabledCachedMessage = FindClosestMessageToTimestamp(SessionCachedMessages->EnvelopeFollowerEnabledCachedMessages, TimeMarker, AudioBusId);
+#endif // WITH_EDITOR
+				if (FoundEnvelopeFollowerEnabledCachedMessage)
+				{
+					CachedEntryInfos[Index].EnvelopeFollowerEnabledMessage = *FoundEnvelopeFollowerEnabledCachedMessage;
+				}
+
 				// Envelope Values
 #if WITH_EDITOR
 				const FAudioBusEnvelopeValuesMessage* FoundEnvelopeValuesCachedMessage = CacheManager.FindClosestMessage<FAudioBusEnvelopeValuesMessage>(AudioBusMessageNames::EnvelopeValues, TimeMarker, AudioBusId);
@@ -576,6 +618,24 @@ namespace UE::Audio::Insights
 					});
 				}
 
+				if (CachedEntryInfo.EnvelopeFollowerEnabledMessage.IsSet())
+				{
+					const FAudioBusEnvelopeFollowerEnabledMessage& EnvelopeFollowerEnabledMessage = CachedEntryInfo.EnvelopeFollowerEnabledMessage.GetValue();
+
+					UpdateDeviceEntry(EnvelopeFollowerEnabledMessage.DeviceId, EnvelopeFollowerEnabledMessage.AudioBusId, [&EnvelopeFollowerEnabledMessage](TSharedPtr<FAudioBusDashboardEntry>& Entry)
+					{
+						if (!Entry.IsValid())
+						{
+							Entry = MakeShared<FAudioBusDashboardEntry>();
+							Entry->AudioBusId = EnvelopeFollowerEnabledMessage.AudioBusId;
+						}
+
+						Entry->Timestamp = EnvelopeFollowerEnabledMessage.Timestamp;
+
+						Entry->bEnvelopeFollowerEnabled = EnvelopeFollowerEnabledMessage.bEnvelopeFollowerEnabled;
+					});
+				}
+
 				if (CachedEntryInfo.EnvelopeValuesMessage.IsSet())
 				{
 					const FAudioBusEnvelopeValuesMessage& EnvelopeValuesMessage = CachedEntryInfo.EnvelopeValuesMessage.GetValue();
@@ -599,6 +659,8 @@ namespace UE::Audio::Insights
 
 		// Call parent method to update LastMessageId
 		FTraceProviderBase::OnTimingViewTimeMarkerChanged(TimeMarker);
+
+		OnTimeMarkerUpdated.ExecuteIfBound();
 	}
 
 #if WITH_EDITOR
@@ -607,5 +669,5 @@ namespace UE::Audio::Insights
 		AudioBusAssetProvider.RequestEntriesUpdate();
 		OnAudioBusListUpdated.ExecuteIfBound();
 	}
-#endif
+#endif // WITH_EDITOR
 } // namespace UE::Audio::Insights
diff --git a/Engine/Plugins/AudioInsights/Source/AudioInsights/Private/Providers/AudioBusTraceProvider.h b/Engine/Plugins/AudioInsights/Source/AudioInsights/Private/Providers/AudioBusTraceProvider.h
--- a/Engine/Plugins/AudioInsights/Source/AudioInsights/Private/Providers/AudioBusTraceProvider.h
+++ b/Engine/Plugins/AudioInsights/Source/AudioInsights/Private/Providers/AudioBusTraceProvider.h
@@ -35,6 +35,9 @@ namespace UE::Audio::Insights
 		DECLARE_DELEGATE(FOnAudioBusListUpdated);
 		FOnAudioBusListUpdated OnAudioBusListUpdated;
 
+		DECLARE_DELEGATE(FOnTimeMarkerUpdated);
+		FOnTimeMarkerUpdated OnTimeMarkerUpdated;
+
 	private:
 		virtual bool ProcessMessages() override;
 
diff --git a/Engine/Plugins/AudioInsights/Source/AudioInsights/Private/Providers/SubmixTraceProvider.cpp b/Engine/Plugins/AudioInsights/Source/AudioInsights/Private/Providers/SubmixTraceProvider.cpp
--- a/Engine/Plugins/AudioInsights/Source/AudioInsights/Private/Providers/SubmixTraceProvider.cpp
+++ b/Engine/Plugins/AudioInsights/Source/AudioInsights/Private/Providers/SubmixTraceProvider.cpp
@@ -102,11 +102,12 @@ namespace UE::Audio::Insights
 
 				UE::Trace::IAnalyzer::FInterfaceBuilder& Builder = Context.InterfaceBuilder;
 
-				Builder.RouteEvent(RouteId_Loaded,         "Audio", "SubmixAlivePing");
-				Builder.RouteEvent(RouteId_Loaded,         "Audio", "SubmixLoaded");
-				Builder.RouteEvent(RouteId_HasActivity,    "Audio", "SubmixHasActivity");
-				Builder.RouteEvent(RouteId_EnvelopeValues, "Audio", "SubmixEnvelopeValues");
-				Builder.RouteEvent(RouteId_Unloaded,       "Audio", "SubmixUnloaded");
+				Builder.RouteEvent(RouteId_Loaded,                  "Audio", "SubmixAlivePing");
+				Builder.RouteEvent(RouteId_Loaded,                  "Audio", "SubmixLoaded");
+				Builder.RouteEvent(RouteId_HasActivity,             "Audio", "SubmixHasActivity");
+				Builder.RouteEvent(RouteId_EnvelopeFollowerEnabled, "Audio", "SubmixEnvelopeFollowerEnabled");
+				Builder.RouteEvent(RouteId_EnvelopeValues,          "Audio", "SubmixEnvelopeValues");
+				Builder.RouteEvent(RouteId_Unloaded,                "Audio", "SubmixUnloaded");
 			}
 
 			virtual bool OnHandleEvent(uint16 RouteId, EStyle Style, const FOnEventContext& Context) override
@@ -137,6 +138,16 @@ namespace UE::Audio::Insights
 						break;
 					}
 
+					case RouteId_EnvelopeFollowerEnabled:
+					{
+#if WITH_EDITOR
+						CacheMessage<FSubmixEnvelopeFollowerEnabledMessage>(Context, Messages.EnvelopeFollowerEnabledMessages);
+#else
+						Messages.EnvelopeFollowerEnabledMessages.Enqueue(FSubmixEnvelopeFollowerEnabledMessage{ Context });
+#endif // WITH_EDITOR
+						break;
+					}
+
 					case RouteId_EnvelopeValues:
 					{
 #if WITH_EDITOR
@@ -178,6 +189,7 @@ namespace UE::Audio::Insights
 			{
 				RouteId_Loaded,
 				RouteId_HasActivity,
+				RouteId_EnvelopeFollowerEnabled,
 				RouteId_EnvelopeValues,
 				RouteId_Unloaded
 			};
@@ -254,6 +266,24 @@ namespace UE::Audio::Insights
 			}
 		});
 
+		ProcessMessageQueue<FSubmixEnvelopeFollowerEnabledMessage>(TraceMessages.EnvelopeFollowerEnabledMessages, GetEntry,
+		[this](const FSubmixEnvelopeFollowerEnabledMessage& Msg, TSharedPtr<FSubmixDashboardEntry>* OutEntry)
+		{
+#if !WITH_EDITOR
+			if (SessionCachedMessages.IsValid())
+			{
+				SessionCachedMessages->EnvelopeFollowerEnabledCachedMessages.EmplaceBack(Msg);
+			}
+#endif // !WITH_EDITOR
+
+			if (OutEntry)
+			{
+				FSubmixDashboardEntry& EntryRef = *OutEntry->Get();
+				EntryRef.bEnvelopeFollowerEnabled = Msg.bEnvelopeFollowerEnabled;
+				EntryRef.Timestamp = Msg.Timestamp;
+			}
+		});
+
 		ProcessMessageQueue<FSubmixEnvelopeValuesMessage>(TraceMessages.EnvelopeValuesMessages, GetEntry,
 		[this](const FSubmixEnvelopeValuesMessage& Msg, TSharedPtr<FSubmixDashboardEntry>* OutEntry)
 		{
@@ -500,6 +530,7 @@ namespace UE::Audio::Insights
 			struct CachedEntryInfo
 			{
 				TOptional<FSubmixHasActivityMessage> HasActivityMessage;
+				TOptional<FSubmixEnvelopeFollowerEnabledMessage> EnvelopeFollowerEnabledMessage;
 				TOptional<FSubmixEnvelopeValuesMessage> EnvelopeValuesMessage;
 			};
 
@@ -529,6 +560,17 @@ namespace UE::Audio::Insights
 					CachedEntryInfos[Index].HasActivityMessage = *FoundHasActivityCachedMessage;
 				}
 
+				// EnvelopeFollowerEnabled
+#if WITH_EDITOR
+				const FSubmixEnvelopeFollowerEnabledMessage* FoundEnvelopeFollowerEnabledCachedMessage = CacheManager.FindClosestMessage<FSubmixEnvelopeFollowerEnabledMessage>(SubmixMessageNames::EnvelopeFollowerEnabled, TimeMarker, SubmixId);
+#else
+				const FSubmixEnvelopeFollowerEnabledMessage* FoundEnvelopeFollowerEnabledCachedMessage = FindClosestMessageToTimestamp(SessionCachedMessages->EnvelopeFollowerEnabledCachedMessages, TimeMarker, SubmixId);
+#endif // WITH_EDITOR
+				if (FoundEnvelopeFollowerEnabledCachedMessage)
+				{
+					CachedEntryInfos[Index].EnvelopeFollowerEnabledMessage = *FoundEnvelopeFollowerEnabledCachedMessage;
+				}
+
 				// EnvelopeValues
 #if WITH_EDITOR
 				const FSubmixEnvelopeValuesMessage* FoundEnvelopeValuesCachedMessage = CacheManager.FindClosestMessage<FSubmixEnvelopeValuesMessage>(SubmixMessageNames::EnvelopeValues, TimeMarker, SubmixId);
@@ -562,6 +604,24 @@ namespace UE::Audio::Insights
 					});
 				}
 
+				if (CachedEntryInfo.EnvelopeFollowerEnabledMessage.IsSet())
+				{
+					const FSubmixEnvelopeFollowerEnabledMessage& EnvelopeFollowerEnabledMessage = CachedEntryInfo.EnvelopeFollowerEnabledMessage.GetValue();
+
+					UpdateDeviceEntry(EnvelopeFollowerEnabledMessage.DeviceId, EnvelopeFollowerEnabledMessage.SubmixId, [&EnvelopeFollowerEnabledMessage](TSharedPtr<FSubmixDashboardEntry>& Entry)
+					{
+						if (!Entry.IsValid())
+						{
+							Entry = MakeShared<FSubmixDashboardEntry>();
+							Entry->SubmixId = EnvelopeFollowerEnabledMessage.SubmixId;
+						}
+
+						Entry->Timestamp = EnvelopeFollowerEnabledMessage.Timestamp;
+
+						Entry->bEnvelopeFollowerEnabled = EnvelopeFollowerEnabledMessage.bEnvelopeFollowerEnabled;
+					});
+				}
+
 				if (CachedEntryInfo.EnvelopeValuesMessage.IsSet())
 				{
 					const FSubmixEnvelopeValuesMessage& EnvelopeValuesMessage = CachedEntryInfo.EnvelopeValuesMessage.GetValue();
@@ -585,6 +645,8 @@ namespace UE::Audio::Insights
 
 		// Call parent method to update LastMessageId
 		FTraceProviderBase::OnTimingViewTimeMarkerChanged(TimeMarker);
+
+		OnTimeMarkerUpdated.ExecuteIfBound();
 	}
 
 #if WITH_EDITOR
diff --git a/Engine/Plugins/AudioInsights/Source/AudioInsights/Private/Providers/SubmixTraceProvider.h b/Engine/Plugins/AudioInsights/Source/AudioInsights/Private/Providers/SubmixTraceProvider.h
--- a/Engine/Plugins/AudioInsights/Source/AudioInsights/Private/Providers/SubmixTraceProvider.h
+++ b/Engine/Plugins/AudioInsights/Source/AudioInsights/Private/Providers/SubmixTraceProvider.h
@@ -35,6 +35,9 @@ namespace UE::Audio::Insights
 		DECLARE_DELEGATE(FOnSubmixListUpdated);
 		FOnSubmixListUpdated OnSubmixListUpdated;
 
+		DECLARE_DELEGATE(FOnTimeMarkerUpdated);
+		FOnTimeMarkerUpdated OnTimeMarkerUpdated;
+
 	private:
 		virtual bool ProcessMessages() override;
 
