[Audio Insights] Submix and Audio Bus dashboards: Disable checkboxes when inspecting Cache (Editor) / Offline traces (Standalone) + added logic to handle audio meter updates when scrubbing

diff --git a/Engine/Plugins/AudioInsights/Source/AudioInsights/Private/Views/AudioBusDashboardViewFactory.cpp b/Engine/Plugins/AudioInsights/Source/AudioInsights/Private/Views/AudioBusDashboardViewFactory.cpp
--- a/Engine/Plugins/AudioInsights/Source/AudioInsights/Private/Views/AudioBusDashboardViewFactory.cpp
+++ b/Engine/Plugins/AudioInsights/Source/AudioInsights/Private/Views/AudioBusDashboardViewFactory.cpp
@@ -1,7 +1,9 @@
 // Copyright Epic Games, Inc. All Rights Reserved.
 #include "Views/AudioBusDashboardViewFactory.h"
 
+#include "AudioInsightsModule.h"
 #include "AudioInsightsStyle.h"
+#include "AudioInsightsTimingViewExtender.h"
 #include "IAudioInsightsModule.h"
 #include "Internationalization/Text.h"
 #include "Providers/AudioMeterProvider.h"
@@ -12,6 +14,8 @@
 #if WITH_EDITOR
 #include "Editor.h"
 #include "Subsystems/AssetEditorSubsystem.h"
+#else
+#include "AudioInsightsComponent.h"
 #endif // WITH_EDITOR
 
 #define LOCTEXT_NAMESPACE "AudioInsights"
@@ -63,6 +67,7 @@ namespace UE::Audio::Insights
 		AudioBusProvider->OnAudioBusAdded.BindRaw(this, &FAudioBusDashboardViewFactory::HandleOnAudioBusAdded);
 		AudioBusProvider->OnAudioBusRemoved.BindRaw(this, &FAudioBusDashboardViewFactory::HandleOnAudioBusRemoved);
 		AudioBusProvider->OnAudioBusListUpdated.BindRaw(this, &FAudioBusDashboardViewFactory::RequestListRefresh);
+		AudioBusProvider->OnTimeMarkerUpdated.BindRaw(this, &FAudioBusDashboardViewFactory::HandleOnTimeMarkerUpdated);
 	}
 
 	FAudioBusDashboardViewFactory::~FAudioBusDashboardViewFactory()
@@ -70,6 +75,7 @@ namespace UE::Audio::Insights
 		AudioBusProvider->OnAudioBusAdded.Unbind();
 		AudioBusProvider->OnAudioBusRemoved.Unbind();
 		AudioBusProvider->OnAudioBusListUpdated.Unbind();
+		AudioBusProvider->OnTimeMarkerUpdated.Unbind();
 	}
 
 	FName FAudioBusDashboardViewFactory::GetName() const
@@ -127,11 +133,10 @@ namespace UE::Audio::Insights
 			}
 
 			const FAudioBusDashboardEntry& AudioBusDashboardEntry = CastEntry(*InRowData);
-			const uint32 AudioBusId = AudioBusDashboardEntry.AudioBusId;
 
-			if (CheckedAudioBuses.Find(AudioBusId) != nullptr)
+			if (ShouldUpdateAudioMeter(AudioBusDashboardEntry))
 			{
-				FAudioMeterProvider::OnUpdateAudioMeterInfo.ExecuteIfBound(AudioBusId, AudioBusDashboardEntry.AudioMeterInfo, AudioBusDashboardEntry.GetDisplayName());
+				FAudioMeterProvider::OnUpdateAudioMeterInfo.ExecuteIfBound(AudioBusDashboardEntry.AudioBusId, AudioBusDashboardEntry.AudioMeterInfo, AudioBusDashboardEntry.GetDisplayName());
 			}
 
 			return SNew(SHorizontalBox)
@@ -140,9 +145,39 @@ namespace UE::Audio::Insights
 				.AutoWidth()
 				[
 					SNew(SCheckBox)
-					.IsChecked_Lambda([this, AudioBusId]()
+					.IsEnabled_Lambda([this, InRowData]()
 					{
-						return CheckedAudioBuses.Find(AudioBusId) != nullptr ? ECheckBoxState::Checked : ECheckBoxState::Unchecked;
+#if WITH_EDITOR
+						return AudioBusProvider && AudioBusProvider->GetMessageCacheAndProcessingStatus() == ECacheAndProcess::Latest;
+#else
+						const TSharedPtr<const FAudioInsightsComponent> AudioInsightsComponent = FAudioInsightsModule::GetChecked().GetAudioInsightsComponent();
+						return AudioInsightsComponent.IsValid() && AudioInsightsComponent->GetIsLiveSession();
+#endif // WITH_EDITOR
+					})
+					.IsChecked_Lambda([this, InRowData]()
+					{
+#if WITH_EDITOR
+						if (!AudioBusProvider.IsValid())
+#else
+						const TSharedPtr<const FAudioInsightsComponent> AudioInsightsComponent = FAudioInsightsModule::GetChecked().GetAudioInsightsComponent();
+						if (!AudioInsightsComponent.IsValid())
+#endif // WITH_EDITOR
+						{
+							return ECheckBoxState::Unchecked;
+						}
+
+						const FAudioBusDashboardEntry& AudioBusDashboardEntry = CastEntry(*InRowData);
+
+#if WITH_EDITOR
+						if (AudioBusProvider->GetMessageCacheAndProcessingStatus() == ECacheAndProcess::Latest)
+#else
+						if (AudioInsightsComponent->GetIsLiveSession())
+#endif // WITH_EDITOR
+						{
+							return CheckedAudioBuses.Contains(AudioBusDashboardEntry.AudioBusId) ? ECheckBoxState::Checked : ECheckBoxState::Unchecked;
+						}
+
+						return AudioBusDashboardEntry.bEnvelopeFollowerEnabled ? ECheckBoxState::Checked : ECheckBoxState::Unchecked;
 					})
 					.OnCheckStateChanged_Lambda([this, InRowData](ECheckBoxState NewState)
 					{
@@ -206,6 +241,17 @@ namespace UE::Audio::Insights
 	{
 		FilterByAudioBusName();
 		FilterByAudioBusType();
+
+		if (bShouldRebuildAudioMeters)
+		{
+			RebuildAudioMeters();
+			bShouldRebuildAudioMeters = false;
+		}
+		else if (bListRefreshed)
+		{
+			UpdateAudioMeters();
+			bListRefreshed = false;
+		}
 	}
 
 	TSharedRef<SWidget> FAudioBusDashboardViewFactory::MakeAudioBusTypeFilterWidget()
@@ -276,6 +322,9 @@ namespace UE::Audio::Insights
 	{
 		if (!DashboardWidget.IsValid())
 		{
+			FAudioInsightsModule& AudioInsightsModule = FAudioInsightsModule::GetChecked();
+			AudioInsightsModule.GetTimingViewExtender().OnTimeControlMethodReset.AddSP(this, &FAudioBusDashboardViewFactory::HandleOnTimeControlMethodReset);
+
 			SAssignNew(DashboardWidget, SVerticalBox)
 			+ SVerticalBox::Slot()
 			.AutoHeight()
@@ -445,7 +494,49 @@ namespace UE::Audio::Insights
 		}
 	}
 
-	void FAudioBusDashboardViewFactory::UpdateAudioMeters()
+	void FAudioBusDashboardViewFactory::RebuildAudioMeters() const
+	{
+		using namespace FAudioBusDashboardViewFactoryPrivate;
+
+		for (const TSharedPtr<IDashboardDataViewEntry>& Entry : DataViewEntries)
+		{
+			const FAudioBusDashboardEntry& AudioBusDashboardEntry = CastEntry(*Entry);
+
+			if (ShouldUpdateAudioMeter(AudioBusDashboardEntry))
+			{
+				FAudioMeterProvider::OnAddAudioMeter.ExecuteIfBound(AudioBusDashboardEntry.AudioBusId, AudioBusDashboardEntry.AudioMeterInfo, AudioBusDashboardEntry.GetDisplayName(), AudioBusMeterTextColor);
+				FAudioMeterProvider::OnUpdateAudioMeterInfo.ExecuteIfBound(AudioBusDashboardEntry.AudioBusId, AudioBusDashboardEntry.AudioMeterInfo, AudioBusDashboardEntry.GetDisplayName());
+			}
+			else
+			{
+				FAudioMeterProvider::OnRemoveAudioMeter.ExecuteIfBound(AudioBusDashboardEntry.AudioBusId);
+			}
+		}
+	}
+
+	bool FAudioBusDashboardViewFactory::ShouldUpdateAudioMeter(const FAudioBusDashboardEntry& InAudioBusDashboardEntry) const
+	{
+#if WITH_EDITOR
+		if (!AudioBusProvider.IsValid())
+		{
+			return false;
+		}
+
+		const bool bProcessingTraceMessages = AudioBusProvider->GetMessageCacheAndProcessingStatus() == ECacheAndProcess::Latest;
+#else
+		const TSharedPtr<const FAudioInsightsComponent> AudioInsightsComponent = FAudioInsightsModule::GetChecked().GetAudioInsightsComponent();
+		if (!AudioInsightsComponent.IsValid())
+		{
+			return false;
+		}
+
+		const bool bProcessingTraceMessages = AudioInsightsComponent->GetIsLiveSession();
+#endif // WITH_EDITOR
+
+		return bProcessingTraceMessages ? CheckedAudioBuses.Contains(InAudioBusDashboardEntry.AudioBusId) : InAudioBusDashboardEntry.bEnvelopeFollowerEnabled;
+	}
+
+	void FAudioBusDashboardViewFactory::UpdateAudioMeters() const
 	{
 		using namespace FAudioBusDashboardViewFactoryPrivate;
 
@@ -475,6 +566,16 @@ namespace UE::Audio::Insights
 		RequestListRefresh();
 	}
 
+	void FAudioBusDashboardViewFactory::HandleOnTimeMarkerUpdated()
+	{
+		bShouldRebuildAudioMeters = true;
+	}
+
+	void FAudioBusDashboardViewFactory::HandleOnTimeControlMethodReset()
+	{
+		bShouldRebuildAudioMeters = true;
+	}
+
 	void FAudioBusDashboardViewFactory::RequestListRefresh()
 	{
 		if (FilteredEntriesListView.IsValid())
@@ -482,7 +583,7 @@ namespace UE::Audio::Insights
 			FilteredEntriesListView->RequestListRefresh();
 		}
 
-		UpdateAudioMeters();
+		bListRefreshed = true;
 	}
 } // namespace UE::Audio::Insights
 
diff --git a/Engine/Plugins/AudioInsights/Source/AudioInsights/Private/Views/SubmixDashboardViewFactory.cpp b/Engine/Plugins/AudioInsights/Source/AudioInsights/Private/Views/SubmixDashboardViewFactory.cpp
--- a/Engine/Plugins/AudioInsights/Source/AudioInsights/Private/Views/SubmixDashboardViewFactory.cpp
+++ b/Engine/Plugins/AudioInsights/Source/AudioInsights/Private/Views/SubmixDashboardViewFactory.cpp
@@ -1,7 +1,9 @@
 // Copyright Epic Games, Inc. All Rights Reserved.
 #include "Views/SubmixDashboardViewFactory.h"
 
+#include "AudioInsightsModule.h"
 #include "AudioInsightsStyle.h"
+#include "AudioInsightsTimingViewExtender.h"
 #include "IAudioInsightsModule.h"
 #include "Internationalization/Text.h"
 #include "Providers/AudioMeterProvider.h"
@@ -11,6 +13,8 @@
 #if WITH_EDITOR
 #include "Editor.h"
 #include "Subsystems/AssetEditorSubsystem.h"
+#else
+#include "AudioInsightsComponent.h"
 #endif // WITH_EDITOR
 
 #define LOCTEXT_NAMESPACE "AudioInsights"
@@ -46,6 +50,7 @@ namespace UE::Audio::Insights
 		SubmixProvider->OnSubmixAdded.BindRaw(this, &FSubmixDashboardViewFactory::HandleOnSubmixAdded);
 		SubmixProvider->OnSubmixRemoved.BindRaw(this, &FSubmixDashboardViewFactory::HandleOnSubmixRemoved);
 		SubmixProvider->OnSubmixListUpdated.BindRaw(this, &FSubmixDashboardViewFactory::RequestListRefresh);
+		SubmixProvider->OnTimeMarkerUpdated.BindRaw(this, &FSubmixDashboardViewFactory::HandleOnTimeMarkerUpdated);
 	}
 
 	FSubmixDashboardViewFactory::~FSubmixDashboardViewFactory()
@@ -53,6 +58,7 @@ namespace UE::Audio::Insights
 		SubmixProvider->OnSubmixAdded.Unbind();
 		SubmixProvider->OnSubmixRemoved.Unbind();
 		SubmixProvider->OnSubmixListUpdated.Unbind();
+		SubmixProvider->OnTimeMarkerUpdated.Unbind();
 	}
 
 	FName FSubmixDashboardViewFactory::GetName() const
@@ -110,11 +116,10 @@ namespace UE::Audio::Insights
 			}
 
 			const FSubmixDashboardEntry& SubmixDashboardEntry = CastEntry(*InRowData);
-			const uint32 SubmixId = SubmixDashboardEntry.SubmixId;
 
-			if (CheckedSubmixes.Find(SubmixId) != nullptr)
+			if (ShouldUpdateAudioMeter(SubmixDashboardEntry))
 			{
-				FAudioMeterProvider::OnUpdateAudioMeterInfo.ExecuteIfBound(SubmixId, SubmixDashboardEntry.AudioMeterInfo, SubmixDashboardEntry.GetDisplayName());
+				FAudioMeterProvider::OnUpdateAudioMeterInfo.ExecuteIfBound(SubmixDashboardEntry.SubmixId, SubmixDashboardEntry.AudioMeterInfo, SubmixDashboardEntry.GetDisplayName());
 			}
 
 			return SNew(SHorizontalBox)
@@ -123,9 +128,39 @@ namespace UE::Audio::Insights
 				.AutoWidth()
 				[
 					SNew(SCheckBox)
-					.IsChecked_Lambda([this, SubmixId]()
+					.IsEnabled_Lambda([this, InRowData]()
+					{
+#if WITH_EDITOR
+						return SubmixProvider && SubmixProvider->GetMessageCacheAndProcessingStatus() == ECacheAndProcess::Latest;
+#else
+						const TSharedPtr<const FAudioInsightsComponent> AudioInsightsComponent = FAudioInsightsModule::GetChecked().GetAudioInsightsComponent();
+						return AudioInsightsComponent.IsValid() && AudioInsightsComponent->GetIsLiveSession();
+#endif // WITH_EDITOR
+					})
+					.IsChecked_Lambda([this, InRowData]()
 					{
-						return CheckedSubmixes.Find(SubmixId) != nullptr ? ECheckBoxState::Checked : ECheckBoxState::Unchecked;
+#if WITH_EDITOR
+						if (!SubmixProvider.IsValid())
+#else
+						const TSharedPtr<const FAudioInsightsComponent> AudioInsightsComponent = FAudioInsightsModule::GetChecked().GetAudioInsightsComponent();
+						if (!AudioInsightsComponent.IsValid())
+#endif // WITH_EDITOR
+						{
+							return ECheckBoxState::Unchecked;
+						}
+
+						const FSubmixDashboardEntry& SubmixDashboardEntry = CastEntry(*InRowData);
+
+#if WITH_EDITOR
+						if (SubmixProvider->GetMessageCacheAndProcessingStatus() == ECacheAndProcess::Latest)
+#else
+						if (AudioInsightsComponent->GetIsLiveSession())
+#endif // WITH_EDITOR
+						{
+							return CheckedSubmixes.Contains(SubmixDashboardEntry.SubmixId) ? ECheckBoxState::Checked : ECheckBoxState::Unchecked;
+						}
+
+						return SubmixDashboardEntry.bEnvelopeFollowerEnabled ? ECheckBoxState::Checked : ECheckBoxState::Unchecked;
 					})
 					.OnCheckStateChanged_Lambda([this, InRowData](ECheckBoxState NewState)
 					{
@@ -194,12 +229,26 @@ namespace UE::Audio::Insights
 			
 			return SubmixEntry.GetDisplayName().ToString().Contains(FilterString);
 		});
+
+		if (bShouldRebuildAudioMeters)
+		{
+			RebuildAudioMeters();
+			bShouldRebuildAudioMeters = false;
+		}
+		else if (bListRefreshed)
+		{
+			UpdateAudioMeters();
+			bListRefreshed = false;
+		}
 	}
 
 	TSharedRef<SWidget> FSubmixDashboardViewFactory::MakeWidget(TSharedRef<SDockTab> OwnerTab, const FSpawnTabArgs& SpawnTabArgs)
 	{
 		if (!DashboardWidget.IsValid())
 		{
+			FAudioInsightsModule& AudioInsightsModule = FAudioInsightsModule::GetChecked();
+			AudioInsightsModule.GetTimingViewExtender().OnTimeControlMethodReset.AddSP(this, &FSubmixDashboardViewFactory::HandleOnTimeControlMethodReset);
+
 			DashboardWidget = FTraceTableDashboardViewFactory::MakeWidget(OwnerTab, SpawnTabArgs);
 
 			if (FilteredEntriesListView.IsValid())
@@ -337,8 +386,50 @@ namespace UE::Audio::Insights
 			}
 		}
 	}
-	
-	void FSubmixDashboardViewFactory::UpdateAudioMeters()
+
+	void FSubmixDashboardViewFactory::RebuildAudioMeters() const
+	{
+		using namespace FSubmixDashboardViewFactoryPrivate;
+
+		for (const TSharedPtr<IDashboardDataViewEntry>& Entry : DataViewEntries)
+		{
+			const FSubmixDashboardEntry& SubmixDashboardEntry = CastEntry(*Entry);
+
+			if (ShouldUpdateAudioMeter(SubmixDashboardEntry))
+			{
+				FAudioMeterProvider::OnAddAudioMeter.ExecuteIfBound(SubmixDashboardEntry.SubmixId, SubmixDashboardEntry.AudioMeterInfo, SubmixDashboardEntry.GetDisplayName(), SubmixMeterTextColor);
+				FAudioMeterProvider::OnUpdateAudioMeterInfo.ExecuteIfBound(SubmixDashboardEntry.SubmixId, SubmixDashboardEntry.AudioMeterInfo, SubmixDashboardEntry.GetDisplayName());
+			}
+			else
+			{
+				FAudioMeterProvider::OnRemoveAudioMeter.ExecuteIfBound(SubmixDashboardEntry.SubmixId);
+			}
+		}
+	}
+
+	bool FSubmixDashboardViewFactory::ShouldUpdateAudioMeter(const FSubmixDashboardEntry& InSubmixDashboardEntry) const
+	{
+#if WITH_EDITOR
+		if (!SubmixProvider.IsValid())
+		{
+			return false;
+		}
+
+		const bool bProcessingTraceMessages = SubmixProvider->GetMessageCacheAndProcessingStatus() == ECacheAndProcess::Latest;
+#else
+		const TSharedPtr<const FAudioInsightsComponent> AudioInsightsComponent = FAudioInsightsModule::GetChecked().GetAudioInsightsComponent();
+		if (!AudioInsightsComponent.IsValid())
+		{
+			return false;
+		}
+
+		const bool bProcessingTraceMessages = AudioInsightsComponent->GetIsLiveSession();
+#endif // WITH_EDITOR
+
+		return bProcessingTraceMessages ? CheckedSubmixes.Contains(InSubmixDashboardEntry.SubmixId) : InSubmixDashboardEntry.bEnvelopeFollowerEnabled;
+	}
+
+	void FSubmixDashboardViewFactory::UpdateAudioMeters() const
 	{
 		using namespace FSubmixDashboardViewFactoryPrivate;
 
@@ -382,6 +473,16 @@ namespace UE::Audio::Insights
 		RequestListRefresh();
 	}
 
+	void FSubmixDashboardViewFactory::HandleOnTimeMarkerUpdated()
+	{
+		bShouldRebuildAudioMeters = true;
+	}
+
+	void FSubmixDashboardViewFactory::HandleOnTimeControlMethodReset()
+	{
+		bShouldRebuildAudioMeters = true;
+	}
+
 	void FSubmixDashboardViewFactory::RequestListRefresh()
 	{
 		if (FilteredEntriesListView.IsValid())
@@ -389,7 +490,7 @@ namespace UE::Audio::Insights
 			FilteredEntriesListView->RequestListRefresh();
 		}
 
-		UpdateAudioMeters();
+		bListRefreshed = true;
 	}
 } // namespace UE::Audio::Insights
 
diff --git a/Engine/Plugins/AudioInsights/Source/AudioInsights/Public/Views/AudioBusDashboardViewFactory.h b/Engine/Plugins/AudioInsights/Source/AudioInsights/Public/Views/AudioBusDashboardViewFactory.h
--- a/Engine/Plugins/AudioInsights/Source/AudioInsights/Public/Views/AudioBusDashboardViewFactory.h
+++ b/Engine/Plugins/AudioInsights/Source/AudioInsights/Public/Views/AudioBusDashboardViewFactory.h
@@ -9,6 +9,7 @@
 namespace UE::Audio::Insights
 {
 	class FAudioBusTraceProvider;
+	struct FAudioBusDashboardEntry;
 
 	class FAudioBusDashboardViewFactory : public FTraceObjectTableDashboardViewFactory
 	{
@@ -36,7 +37,9 @@ namespace UE::Audio::Insights
 		UE_API virtual void ProcessEntries(FTraceTableDashboardViewFactory::EProcessReason Reason) override;
 		UE_API virtual const TMap<FName, FTraceTableDashboardViewFactory::FColumnData>& GetColumns() const override;
 
-		UE_API void UpdateAudioMeters();
+		UE_API void RebuildAudioMeters() const;
+		UE_API bool ShouldUpdateAudioMeter(const FAudioBusDashboardEntry& InAudioBusDashboardEntry) const;
+		UE_API void UpdateAudioMeters() const;
 
 		UE_API virtual void SortTable() override;
 
@@ -45,12 +48,18 @@ namespace UE::Audio::Insights
 
 		UE_API void HandleOnAudioBusAdded(const uint32 InAudioBusId);
 		UE_API void HandleOnAudioBusRemoved(const uint32 InAudioBusId);
+		UE_API void HandleOnTimeMarkerUpdated();
+		UE_API void HandleOnTimeControlMethodReset();
+
 		UE_API void RequestListRefresh();
 
 		TSharedPtr<FAudioBusTraceProvider> AudioBusProvider;
 
 		TSet<uint32> CheckedAudioBuses;
 
+		bool bShouldRebuildAudioMeters = false;
+		bool bListRefreshed = false;
+
 		using FComboboxSelectionItem = TPair<EAudioBusTypeComboboxSelection, FText>;
 		TArray<TSharedPtr<FComboboxSelectionItem>> AudioBusTypes;
 		TSharedPtr<FComboboxSelectionItem> SelectedAudioBusType;
diff --git a/Engine/Plugins/AudioInsights/Source/AudioInsights/Public/Views/SubmixDashboardViewFactory.h b/Engine/Plugins/AudioInsights/Source/AudioInsights/Public/Views/SubmixDashboardViewFactory.h
--- a/Engine/Plugins/AudioInsights/Source/AudioInsights/Public/Views/SubmixDashboardViewFactory.h
+++ b/Engine/Plugins/AudioInsights/Source/AudioInsights/Public/Views/SubmixDashboardViewFactory.h
@@ -8,8 +8,8 @@
 
 namespace UE::Audio::Insights
 {
-	struct FSubmixDashboardEntry;
 	class FSubmixTraceProvider;
+	struct FSubmixDashboardEntry;
 
 	class FSubmixDashboardViewFactory : public FTraceObjectTableDashboardViewFactory
 	{
@@ -30,17 +30,25 @@ namespace UE::Audio::Insights
 		UE_API virtual void ProcessEntries(FTraceTableDashboardViewFactory::EProcessReason Reason) override;
 		UE_API virtual const TMap<FName, FTraceTableDashboardViewFactory::FColumnData>& GetColumns() const override;
 
-		UE_API void UpdateAudioMeters();
+		UE_API void RebuildAudioMeters() const;
+		UE_API bool ShouldUpdateAudioMeter(const FSubmixDashboardEntry& InSubmixDashboardEntry) const;
+		UE_API void UpdateAudioMeters() const;
 
 		UE_API virtual void SortTable() override;
 
-		void HandleOnSubmixAdded(const uint32 InSubmixId);
-		void HandleOnSubmixRemoved(const uint32 InSubmixId);
+		UE_API void HandleOnSubmixAdded(const uint32 InSubmixId);
+		UE_API void HandleOnSubmixRemoved(const uint32 InSubmixId);
+		UE_API void HandleOnTimeMarkerUpdated();
+		UE_API void HandleOnTimeControlMethodReset();
+
 		UE_API void RequestListRefresh();
 
 		TSharedPtr<FSubmixTraceProvider> SubmixProvider;
 
 		TSet<uint32> CheckedSubmixes;
+
+		bool bShouldRebuildAudioMeters = false;
+		bool bListRefreshed = false;
 	};
 } // namespace UE::Audio::Insights
 
