[Audio Insights] Scrubbing timestamps via plots part 01: AudioInsightsTimingViewExtender
* Set and provide API to enquire which system is being used to set the timestamp. If this came from outside Audio Insights, set to External.
* Store the current plotting range - used to determine what range of data should be drawn from the cache.

diff --git a/Engine/Plugins/AudioInsights/Source/AudioInsights/Private/AudioInsightsTimingViewExtender.cpp b/Engine/Plugins/AudioInsights/Source/AudioInsights/Private/AudioInsightsTimingViewExtender.cpp
--- a/Engine/Plugins/AudioInsights/Source/AudioInsights/Private/AudioInsightsTimingViewExtender.cpp
+++ b/Engine/Plugins/AudioInsights/Source/AudioInsights/Private/AudioInsightsTimingViewExtender.cpp
@@ -13,6 +13,8 @@ namespace UE::Audio::Insights
 			TimingView = &InSession;
 		}
 
+		SystemControllingTime.Reset();
+
 		InSession.OnTimeMarkerChanged().AddRaw(this, &FAudioInsightsTimingViewExtender::OnTimeMarkerChanged);
 	}
 
@@ -36,6 +38,8 @@ namespace UE::Audio::Insights
 	void FAudioInsightsTimingViewExtender::ResetMessageProcessType()
 	{
 		CacheAndProcessMessageStatus = ECacheAndProcess::Latest;
+		SystemControllingTime.Reset();
+
 		OnTimeControlMethodReset.Broadcast();
 	}
 
@@ -44,8 +48,20 @@ namespace UE::Audio::Insights
 		return TraceDurationSeconds;
 	}
 
-	void FAudioInsightsTimingViewExtender::SetTimeMarker(const double Timestamp)
+	void FAudioInsightsTimingViewExtender::SetTimeMarker(const double Timestamp, const ESystemControllingTimeMarker ControllingSystem, TOptional<TRange<double>> InPlottingRange /*= TOptional<TRange<double>>()*/)
 	{
+		SystemControllingTime = ControllingSystem;
+		bUserInputDetected = true;
+
+		if (InPlottingRange.IsSet())
+		{
+			PlottingRange = InPlottingRange.GetValue();
+		}
+		else
+		{
+			PlottingRange = TRange<double>(Timestamp - 5.0, Timestamp);
+		}
+
 		// If you set the time marker on the TimingView, it will Broadcast OnTimingViewTimeMarkerChanged. Otherwise we gotta do it ourselves.
 		if (TimingView != nullptr)
 		{
@@ -61,6 +77,7 @@ namespace UE::Audio::Insights
 	{
 		if (AnalysisSession == nullptr)
 		{
+			bUserInputDetected = false;
 			return;
 		}
 		
@@ -71,7 +88,18 @@ namespace UE::Audio::Insights
 		// Scrub Audio Insights when all data is loaded in non-live sessions
 		if (AnalysisSession->IsAnalysisComplete() && bIsValidTime)
 		{
+			if (!bUserInputDetected)
+			{
+				SystemControllingTime = ESystemControllingTimeMarker::External;
+			}
+
 			OnTimingViewTimeMarkerChanged.Broadcast(InTimeMarker);
 		}
+		else
+		{
+			SystemControllingTime.Reset();
+		}
+
+		bUserInputDetected = false;
 	}
 } // namespace UE::Audio::Insights
diff --git a/Engine/Plugins/AudioInsights/Source/AudioInsights/Private/AudioInsightsTimingViewExtender.h b/Engine/Plugins/AudioInsights/Source/AudioInsights/Private/AudioInsightsTimingViewExtender.h
--- a/Engine/Plugins/AudioInsights/Source/AudioInsights/Private/AudioInsightsTimingViewExtender.h
+++ b/Engine/Plugins/AudioInsights/Source/AudioInsights/Private/AudioInsightsTimingViewExtender.h
@@ -3,6 +3,8 @@
 
 #include "Delegates/Delegate.h"
 #include "Insights/ITimingViewExtender.h"
+#include "Math/Range.h"
+#include "Misc/Optional.h"
 
 namespace Insights
 {
@@ -24,6 +26,13 @@ namespace UE::Audio::Insights
 		None					// Do not cache or process any new messages
 	};
 
+	enum class ESystemControllingTimeMarker
+	{
+		EventLog = 0,
+		PlotsWidget,
+		External
+	};
+
 	class  FAudioInsightsTimingViewExtender : public UE::Insights::Timing::ITimingViewExtender
 	{
 	public:
@@ -33,14 +42,16 @@ namespace UE::Audio::Insights
 		virtual void Tick(UE::Insights::Timing::ITimingViewSession& InSession, const TraceServices::IAnalysisSession& InAnalysisSession) override;
 
 		ECacheAndProcess GetMessageCacheAndProcessingStatus() const { return CacheAndProcessMessageStatus; }
+		TOptional<ESystemControllingTimeMarker> TryGetSystemControllingTimeMarker() const { return SystemControllingTime; }
 
 		void StopProcessingNewMessages() { CacheAndProcessMessageStatus = ECacheAndProcess::CacheLatestNoProcess; }
 		void StopCachingAndProcessingNewMessages() { CacheAndProcessMessageStatus = ECacheAndProcess::None; }
 		void ResetMessageProcessType();
 
 		double GetCurrentDurationSeconds() const;
+		TRange<double> GetPlottingRange() const { return PlottingRange; }
 
-		void SetTimeMarker(const double Timestamp);
+		void SetTimeMarker(const double Timestamp, const ESystemControllingTimeMarker ControllingSystem, TOptional<TRange<double>> InPlottingRange = TOptional<TRange<double>>());
 
 		DECLARE_MULTICAST_DELEGATE_OneParam(FOnTimingViewTimeMarkerChanged, double /*TimeMarker*/);
 		FOnTimingViewTimeMarkerChanged OnTimingViewTimeMarkerChanged;
@@ -55,6 +66,9 @@ namespace UE::Audio::Insights
 		UE::Insights::Timing::ITimingViewSession* TimingView = nullptr;
 
 		ECacheAndProcess CacheAndProcessMessageStatus = ECacheAndProcess::Latest;
+		TOptional<ESystemControllingTimeMarker> SystemControllingTime;
+		TRange<double> PlottingRange { 0.0, 5.0 };
 		double TraceDurationSeconds = 0.0;
+		bool bUserInputDetected = false;
 	};
 } // namespace UE::Audio::Insights
diff --git a/Engine/Plugins/AudioInsights/Source/AudioInsights/Private/Views/AudioEventLogDashboardViewFactory.cpp b/Engine/Plugins/AudioInsights/Source/AudioInsights/Private/Views/AudioEventLogDashboardViewFactory.cpp
--- a/Engine/Plugins/AudioInsights/Source/AudioInsights/Private/Views/AudioEventLogDashboardViewFactory.cpp
+++ b/Engine/Plugins/AudioInsights/Source/AudioInsights/Private/Views/AudioEventLogDashboardViewFactory.cpp
@@ -216,10 +216,11 @@ namespace UE::Audio::Insights
 
 		const FAudioEventLogEditorCommands& Commands = FAudioEventLogEditorCommands::Get();
 
+		CommandList->MapAction(Commands.GetResetInspectTimestampCommand(), FExecuteAction::CreateLambda([this]() { ResetInspectTimestamp(); }));
+
 #if WITH_EDITOR
 		CommandList->MapAction(Commands.GetBrowseCommand(), FExecuteAction::CreateLambda([this]() { BrowseToAsset(); }));
 		CommandList->MapAction(Commands.GetEditCommand(), FExecuteAction::CreateLambda([this]() { OpenAsset(); }));
-		CommandList->MapAction(Commands.GetResetInspectTimestampCommand(), FExecuteAction::CreateLambda([this]() { ResetInspectTimestamp(); }));
 
 		CommandList->MapAction(
 			Commands.GetAutoStopCachingWhenLastInCacheCommand(),
@@ -958,33 +959,7 @@ namespace UE::Audio::Insights
 		if (SelectInfo != ESelectInfo::Type::Direct)
 		{
 			FAudioInsightsModule& AudioInsightsModule = FAudioInsightsModule::GetChecked();
-			AudioInsightsModule.GetTimingViewExtender().SetTimeMarker(SelectedEventLogEntry.Timestamp);
-
-#if WITH_EDITOR
-			if (AutoStopCachingMode == EAutoStopCachingMode::OnInspect)
-			{
-				AudioInsightsModule.GetTimingViewExtender().StopCachingAndProcessingNewMessages();
-				bAutoScroll = false;
-			}
-			else if (AutoStopCachingMode == EAutoStopCachingMode::WhenInLastChunk)
-			{
-				if (TimestampIsInChunkMarkedToDelete(SelectedEventLogEntry.Timestamp))
-				{
-					AudioInsightsModule.GetTimingViewExtender().StopCachingAndProcessingNewMessages();
-					bAutoScroll = false;
-				}
-				else if (AudioInsightsModule.GetTimingViewExtender().GetMessageCacheAndProcessingStatus() != ECacheAndProcess::None)
-				{
-					AudioInsightsModule.GetTimingViewExtender().StopProcessingNewMessages();
-				}
-			}
-			else
-			{
-				AudioInsightsModule.GetTimingViewExtender().StopProcessingNewMessages();
-			}
-#else
-			bAutoScroll = false;
-#endif	// WITH_EDITOR
+			AudioInsightsModule.GetTimingViewExtender().SetTimeMarker(SelectedEventLogEntry.Timestamp, ESystemControllingTimeMarker::EventLog);
 		}
 	}
 
@@ -1296,12 +1271,39 @@ namespace UE::Audio::Insights
 	{
 		using namespace AudioEventLogPrivate;
 
+#if WITH_EDITOR
+		FAudioInsightsModule& AudioInsightsModule = FAudioInsightsModule::GetChecked();
+
+		if (AutoStopCachingMode == EAutoStopCachingMode::OnInspect)
+		{
+			AudioInsightsModule.GetTimingViewExtender().StopCachingAndProcessingNewMessages();
+			bAutoScroll = false;
+		}
+		else if (AutoStopCachingMode == EAutoStopCachingMode::WhenInLastChunk)
+		{
+			if (TimestampIsInChunkMarkedToDelete(InTimeMarker))
+			{
+				AudioInsightsModule.GetTimingViewExtender().StopCachingAndProcessingNewMessages();
+				bAutoScroll = false;
+			}
+			else if (AudioInsightsModule.GetTimingViewExtender().GetMessageCacheAndProcessingStatus() != ECacheAndProcess::None)
+			{
+				AudioInsightsModule.GetTimingViewExtender().StopProcessingNewMessages();
+			}
+		}
+		else
+		{
+			AudioInsightsModule.GetTimingViewExtender().StopProcessingNewMessages();
+		}
+#endif // WITH_EDITOR
+
+		bAutoScroll = false;
+
 		if (!FilteredEntriesListView.IsValid())
 		{
 			return;
 		}
 
-#if !WITH_EDITOR
 		// We assume the data entries are in timestamp order (earliest to latest)
 		for (int32 Index = 0; Index < DataViewEntries.Num(); ++Index)
 		{
@@ -1338,7 +1340,6 @@ namespace UE::Audio::Insights
 				FilteredEntriesListView->SetSelection(Entry);
 			}
 		}
-#endif // !WITH_EDITOR
 	}
 
 	void FAudioEventLogDashboardViewFactory::UpdateCacheMessageProcessMethod()
