[Audio Insights] Scrubbing timestamps via plots part 02: Split out "current"parameter values and store them separately to the data buffers inside SoundDashboardEntries.
* Previously we assumed the current value of a sound dashboard entry was to last in the data buffer. However, the timemarker can now be located in the middle of the plots window.
* We now seperate the two concepts of data for the plots widget and data to show in the sounds tab (i.e. data at the time marker).
* Use plotting range gathered from Timing Authority rather than assuming the range is timemarker - 5s
* Extra fix in FSoundTraceProvider::ProcessManuallyUpdatedEntries - make sure we are actually removing some items before returning true

diff --git a/Engine/Plugins/AudioInsights/Source/AudioInsights/Private/Messages/SoundTraceMessages.cpp b/Engine/Plugins/AudioInsights/Source/AudioInsights/Private/Messages/SoundTraceMessages.cpp
--- a/Engine/Plugins/AudioInsights/Source/AudioInsights/Private/Messages/SoundTraceMessages.cpp
+++ b/Engine/Plugins/AudioInsights/Source/AudioInsights/Private/Messages/SoundTraceMessages.cpp
@@ -95,59 +95,27 @@ namespace UE::Audio::Insights
 	{
 		constexpr uint32 DataPointsCapacity = 256u;
 
-		PriorityDataPoints.SetCapacity(DataPointsCapacity);
-		DistanceDataPoints.SetCapacity(DataPointsCapacity);
-		DistanceAttenuationDataPoints.SetCapacity(DataPointsCapacity);
-		LPFFreqDataPoints.SetCapacity(DataPointsCapacity);
-		HPFFreqDataPoints.SetCapacity(DataPointsCapacity);
-		AmplitudeDataPoints.SetCapacity(DataPointsCapacity);
-		VolumeDataPoints.SetCapacity(DataPointsCapacity);
-		PitchDataPoints.SetCapacity(DataPointsCapacity);
-		RelativeRenderCostDataPoints.SetCapacity(DataPointsCapacity);
+		PriorityDataRange.SetCapacity(DataPointsCapacity);
+		DistanceDataRange.SetCapacity(DataPointsCapacity);
+		DistanceAttenuationDataRange.SetCapacity(DataPointsCapacity);
+		LPFFreqDataRange.SetCapacity(DataPointsCapacity);
+		HPFFreqDataRange.SetCapacity(DataPointsCapacity);
+		AmplitudeDataRange.SetCapacity(DataPointsCapacity);
+		VolumeDataRange.SetCapacity(DataPointsCapacity);
+		PitchDataRange.SetCapacity(DataPointsCapacity);
+		RelativeRenderCostDataRange.SetCapacity(DataPointsCapacity);
 	}
 
 	void FSoundDashboardEntry::ResetDataBuffers(const uint32 DataPointsCapacity)
 	{
-		PriorityDataPoints.Reset(DataPointsCapacity);
-		DistanceDataPoints.Reset(DataPointsCapacity);
-		DistanceAttenuationDataPoints.Reset(DataPointsCapacity);
-		LPFFreqDataPoints.Reset(DataPointsCapacity);
-		HPFFreqDataPoints.Reset(DataPointsCapacity);
-		AmplitudeDataPoints.Reset(DataPointsCapacity);
-		VolumeDataPoints.Reset(DataPointsCapacity);
-		PitchDataPoints.Reset(DataPointsCapacity);
-		RelativeRenderCostDataPoints.Reset(DataPointsCapacity);
-	}
-
-	void FSoundDashboardEntry::SwitchBuffersToSingleValues()
-	{
-		constexpr uint32 DataPointsCapacity = 1u;
-
-		PriorityDataPoints.Empty();
-		PriorityDataPoints.SetCapacity(DataPointsCapacity);
-
-		DistanceDataPoints.Empty();
-		DistanceDataPoints.SetCapacity(DataPointsCapacity);
-
-		DistanceAttenuationDataPoints.Empty();
-		DistanceAttenuationDataPoints.SetCapacity(DataPointsCapacity);
-
-		LPFFreqDataPoints.Empty();
-		LPFFreqDataPoints.SetCapacity(DataPointsCapacity);
-
-		HPFFreqDataPoints.Empty();
-		HPFFreqDataPoints.SetCapacity(DataPointsCapacity);
-
-		AmplitudeDataPoints.Empty();
-		AmplitudeDataPoints.SetCapacity(DataPointsCapacity);
-
-		VolumeDataPoints.Empty();
-		VolumeDataPoints.SetCapacity(DataPointsCapacity);
-
-		PitchDataPoints.Empty();
-		PitchDataPoints.SetCapacity(DataPointsCapacity);
-
-		RelativeRenderCostDataPoints.Empty();
-		RelativeRenderCostDataPoints.SetCapacity(DataPointsCapacity);
+		PriorityDataRange.Reset(DataPointsCapacity);
+		DistanceDataRange.Reset(DataPointsCapacity);
+		DistanceAttenuationDataRange.Reset(DataPointsCapacity);
+		LPFFreqDataRange.Reset(DataPointsCapacity);
+		HPFFreqDataRange.Reset(DataPointsCapacity);
+		AmplitudeDataRange.Reset(DataPointsCapacity);
+		VolumeDataRange.Reset(DataPointsCapacity);
+		PitchDataRange.Reset(DataPointsCapacity);
+		RelativeRenderCostDataRange.Reset(DataPointsCapacity);
 	}
 }
\ No newline at end of file
diff --git a/Engine/Plugins/AudioInsights/Source/AudioInsights/Private/Providers/SoundTraceProvider.cpp b/Engine/Plugins/AudioInsights/Source/AudioInsights/Private/Providers/SoundTraceProvider.cpp
--- a/Engine/Plugins/AudioInsights/Source/AudioInsights/Private/Providers/SoundTraceProvider.cpp
+++ b/Engine/Plugins/AudioInsights/Source/AudioInsights/Private/Providers/SoundTraceProvider.cpp
@@ -77,15 +77,16 @@ namespace UE::Audio::Insights
 			return DataPointsDisjointedArrayView.FirstBuffer.Last();
 		}
 
-		void UpdateParameterEntry(const FSoundParameterMessage& Msg, const float DataPoint, ::Audio::TCircularAudioBuffer<FDataPoint>& OutDataPoints, TArray<int32, TInlineAllocator<64>>& OutEntriesWithPoppedDataPoints)
+		void UpdateParameterEntry(const FSoundParameterMessage& Msg, const float DataPoint, float& OutDataPoint, ::Audio::TCircularAudioBuffer<FDataPoint>& OutDataRange, TArray<int32, TInlineAllocator<64>>& OutEntriesWithPoppedDataPoints)
 		{
 			if (!OutEntriesWithPoppedDataPoints.Contains(Msg.PlayOrder))
 			{
-				OutDataPoints.Pop(OutDataPoints.Num());
+				OutDataRange.Pop(OutDataRange.Num());
 				OutEntriesWithPoppedDataPoints.Add(Msg.PlayOrder);
 			}
 
-			OutDataPoints.Push({ Msg.Timestamp, DataPoint });
+			OutDataPoint = DataPoint;
+			OutDataRange.Push({ Msg.Timestamp, DataPoint });
 		}
 
 		//////////////////////////////////////////////////////////////////////////////////////////////////////////////
@@ -589,7 +590,7 @@ namespace UE::Audio::Insights
 			}
 #endif // !WITH_EDITOR
 
-			FSoundTraceProviderPrivate::UpdateParameterEntry(Msg, Msg.Priority, SoundEntry->PriorityDataPoints, EntriesWithPoppedDataPoints);
+			FSoundTraceProviderPrivate::UpdateParameterEntry(Msg, Msg.Priority, SoundEntry->PriorityDataPoint, SoundEntry->PriorityDataRange, EntriesWithPoppedDataPoints);
 		});
 		EntriesWithPoppedDataPoints.Reset();
 
@@ -609,7 +610,7 @@ namespace UE::Audio::Insights
 			}
 #endif // !WITH_EDITOR
 
-			FSoundTraceProviderPrivate::UpdateParameterEntry(Msg, Msg.Distance, SoundEntry->DistanceDataPoints, EntriesWithPoppedDataPoints);
+			FSoundTraceProviderPrivate::UpdateParameterEntry(Msg, Msg.Distance, SoundEntry->DistanceDataPoint, SoundEntry->DistanceDataRange, EntriesWithPoppedDataPoints);
 		});
 		EntriesWithPoppedDataPoints.Reset();
 
@@ -629,7 +630,7 @@ namespace UE::Audio::Insights
 			}
 #endif // !WITH_EDITOR
 
-			FSoundTraceProviderPrivate::UpdateParameterEntry(Msg, Msg.DistanceAttenuation, SoundEntry->DistanceAttenuationDataPoints, EntriesWithPoppedDataPoints);
+			FSoundTraceProviderPrivate::UpdateParameterEntry(Msg, Msg.DistanceAttenuation, SoundEntry->DistanceAttenuationDataPoint, SoundEntry->DistanceAttenuationDataRange, EntriesWithPoppedDataPoints);
 		});
 		EntriesWithPoppedDataPoints.Reset();
 
@@ -649,7 +650,7 @@ namespace UE::Audio::Insights
 			}
 #endif // !WITH_EDITOR
 
-			FSoundTraceProviderPrivate::UpdateParameterEntry(Msg, Msg.LPFFrequency, SoundEntry->LPFFreqDataPoints, EntriesWithPoppedDataPoints);
+			FSoundTraceProviderPrivate::UpdateParameterEntry(Msg, Msg.LPFFrequency, SoundEntry->LPFFreqDataPoint, SoundEntry->LPFFreqDataRange, EntriesWithPoppedDataPoints);
 		});
 		EntriesWithPoppedDataPoints.Reset();
 
@@ -669,7 +670,7 @@ namespace UE::Audio::Insights
 			}
 #endif // !WITH_EDITOR
 
-			FSoundTraceProviderPrivate::UpdateParameterEntry(Msg, Msg.HPFFrequency, SoundEntry->HPFFreqDataPoints, EntriesWithPoppedDataPoints);
+			FSoundTraceProviderPrivate::UpdateParameterEntry(Msg, Msg.HPFFrequency, SoundEntry->HPFFreqDataPoint, SoundEntry->HPFFreqDataRange, EntriesWithPoppedDataPoints);
 		});
 		EntriesWithPoppedDataPoints.Reset();
 
@@ -689,7 +690,7 @@ namespace UE::Audio::Insights
 			}
 #endif // !WITH_EDITOR
 
-			FSoundTraceProviderPrivate::UpdateParameterEntry(Msg, Msg.Envelope, SoundEntry->AmplitudeDataPoints, EntriesWithPoppedDataPoints);
+			FSoundTraceProviderPrivate::UpdateParameterEntry(Msg, Msg.Envelope, SoundEntry->AmplitudeDataPoint, SoundEntry->AmplitudeDataRange, EntriesWithPoppedDataPoints);
 		});
 		EntriesWithPoppedDataPoints.Reset();
 
@@ -709,7 +710,7 @@ namespace UE::Audio::Insights
 			}
 #endif // !WITH_EDITOR
 
-			FSoundTraceProviderPrivate::UpdateParameterEntry(Msg, Msg.Volume, SoundEntry->VolumeDataPoints, EntriesWithPoppedDataPoints);
+			FSoundTraceProviderPrivate::UpdateParameterEntry(Msg, Msg.Volume, SoundEntry->VolumeDataPoint, SoundEntry->VolumeDataRange, EntriesWithPoppedDataPoints);
 		});
 		EntriesWithPoppedDataPoints.Reset();
 
@@ -729,7 +730,7 @@ namespace UE::Audio::Insights
 			}
 #endif // !WITH_EDITOR
 
-			FSoundTraceProviderPrivate::UpdateParameterEntry(Msg, Msg.Pitch, SoundEntry->PitchDataPoints, EntriesWithPoppedDataPoints);
+			FSoundTraceProviderPrivate::UpdateParameterEntry(Msg, Msg.Pitch, SoundEntry->PitchDataPoint, SoundEntry->PitchDataRange, EntriesWithPoppedDataPoints);
 		});
 		EntriesWithPoppedDataPoints.Reset();
 
@@ -749,7 +750,7 @@ namespace UE::Audio::Insights
 			}
 #endif // !WITH_EDITOR
 
-			FSoundTraceProviderPrivate::UpdateParameterEntry(Msg, Msg.RelativeRenderCost, SoundEntry->RelativeRenderCostDataPoints, EntriesWithPoppedDataPoints);
+			FSoundTraceProviderPrivate::UpdateParameterEntry(Msg, Msg.RelativeRenderCost, SoundEntry->RelativeRenderCostDataPoint, SoundEntry->RelativeRenderCostDataRange, EntriesWithPoppedDataPoints);
 		});
 		EntriesWithPoppedDataPoints.Reset();
 
@@ -830,7 +831,7 @@ namespace UE::Audio::Insights
 			}
 		}
 
-		if (NumToTrim >= 0)
+		if (NumToTrim > 0)
 		{
 			EntriesTimingOut.RemoveAt(0, NumToTrim, EAllowShrinking::No);
 			return true;
@@ -1004,46 +1005,29 @@ namespace UE::Audio::Insights
 			return;
 		}
 
-		auto GetMaxDatapoint = [](const FDataPoint& First, const FDataPoint& Second)
-		{
-			return First.Value >= Second.Value ? First : Second;
-		};
-
-		auto GetMinDatapoint = [](const FDataPoint& First, const FDataPoint& Second)
-		{
-			return First.Value < Second.Value ? First : Second;
-		};
-
 		struct FCumulativeDataPoint
 		{
-			FDataPoint TotalAccumulationDataPoint = { 0.0, 0.0f };
+			float TotalAccumulationDataPoint = { 0.0f };
 			int32 NumDataPoints = 0;
 		};
 
-		auto AccumulateDataPoints = [](const FCumulativeDataPoint& AccumulativeValue, const ::Audio::TCircularAudioBuffer<FDataPoint>& DataBuffer) -> FCumulativeDataPoint
+		auto AccumulateDataPoints = [](const FCumulativeDataPoint& AccumulativeValue, const float DataPoint) -> FCumulativeDataPoint
 		{
-			if (DataBuffer.Num() == 0)
-			{
-				return AccumulativeValue;
-			}
-
-			const FDataPoint AdditionalValue = PeekLastValue(DataBuffer);
-
 			return
 			{
-				{ FMath::Max(AccumulativeValue.TotalAccumulationDataPoint.Key, AdditionalValue.Key), AccumulativeValue.TotalAccumulationDataPoint.Value + AdditionalValue.Value },
+				AccumulativeValue.TotalAccumulationDataPoint + DataPoint,
 				AccumulativeValue.NumDataPoints + 1
 			};
 		};
 
-		auto AverageDataPoints = [](const FCumulativeDataPoint& DataPointsTotal) -> FDataPoint
+		auto AverageDataPoints = [](const FCumulativeDataPoint& DataPointsTotal) -> float
 		{
 			if (DataPointsTotal.NumDataPoints == 0)
 			{
-				return { 0.0, 0.0f };
+				return 0.0f;
 			}
 
-			return { DataPointsTotal.TotalAccumulationDataPoint.Key, DataPointsTotal.TotalAccumulationDataPoint.Value / static_cast<float>(DataPointsTotal.NumDataPoints) };
+			return DataPointsTotal.TotalAccumulationDataPoint / static_cast<float>(DataPointsTotal.NumDataPoints);
 		};
 
 		auto AddSingleValueToBuffer = [](::Audio::TCircularAudioBuffer<FDataPoint>& To, const FDataPoint& Value)
@@ -1056,23 +1040,23 @@ namespace UE::Audio::Insights
 			To.Push(Value);
 		};
 
-		FDataPoint MaxVolume { 0.0, 0.0f };
-		FDataPoint MinDistance { 0.0, TNumericLimits<float>::Max() };
-		FDataPoint MaxPriority { 0.0, 0.0f };
+		float MaxVolume = 0.0f;
+		float MinDistance = TNumericLimits<float>::Max();
+		float MaxPriority = 0.0f;
 
-		FDataPoint AveragePeakAmp { 0.0, 0.0f };
+		float AveragePeakAmp = 0.0f;
 		FCumulativeDataPoint CumulativePeakAmplitude;
 
-		FDataPoint AveragePitch { 0.0, 0.0f };
+		float AveragePitch = 0.0f;
 		FCumulativeDataPoint CumulativePitch;
 
-		FDataPoint AverageDistanceAttenuation { 0.0, 0.0f };
+		float AverageDistanceAttenuation = 0.0f;
 		FCumulativeDataPoint CumulativeDistanceAttenuation;
 
-		FDataPoint AverageLPFFreq { 0.0, 0.0f };
+		float AverageLPFFreq = 0.0f;
 		FCumulativeDataPoint CumulativeLPFFreq;
 
-		FDataPoint AverageHPFFreq { 0.0, 0.0f };
+		float AverageHPFFreq = 0.0f;
 		FCumulativeDataPoint CumulativeHPFFreq;
 
 		FCumulativeDataPoint CumulativeRelativeRenderCost;
@@ -1081,16 +1065,16 @@ namespace UE::Audio::Insights
 		{
 			const FSoundDashboardEntry& WaveSoundEntry = CastEntry(*WaveInstanceEntry);
 
-			MaxVolume = GetMaxDatapoint(MaxVolume, PeekLastValue(WaveSoundEntry.VolumeDataPoints, MaxVolume.Value));
-			MinDistance = GetMinDatapoint(MinDistance, PeekLastValue(WaveSoundEntry.DistanceDataPoints, MinDistance.Value));
-			MaxPriority = GetMaxDatapoint(MaxPriority, PeekLastValue(WaveSoundEntry.PriorityDataPoints, MaxPriority.Value));
+			MaxVolume = FMath::Max(MaxVolume, WaveSoundEntry.VolumeDataPoint);
+			MinDistance = FMath::Min(MinDistance, WaveSoundEntry.DistanceDataPoint);
+			MaxPriority = FMath::Max(MaxPriority, WaveSoundEntry.PriorityDataPoint);
 
-			CumulativePeakAmplitude = AccumulateDataPoints(CumulativePeakAmplitude, WaveSoundEntry.AmplitudeDataPoints);
-			CumulativePitch = AccumulateDataPoints(CumulativePitch, WaveSoundEntry.PitchDataPoints);
-			CumulativeDistanceAttenuation = AccumulateDataPoints(CumulativeDistanceAttenuation, WaveSoundEntry.DistanceAttenuationDataPoints);
-			CumulativeLPFFreq = AccumulateDataPoints(CumulativeLPFFreq, WaveSoundEntry.LPFFreqDataPoints);
-			CumulativeHPFFreq = AccumulateDataPoints(CumulativeHPFFreq, WaveSoundEntry.HPFFreqDataPoints);
-			CumulativeRelativeRenderCost = AccumulateDataPoints(CumulativeRelativeRenderCost, WaveSoundEntry.RelativeRenderCostDataPoints);
+			CumulativePeakAmplitude = AccumulateDataPoints(CumulativePeakAmplitude, WaveSoundEntry.AmplitudeDataPoint);
+			CumulativePitch = AccumulateDataPoints(CumulativePitch, WaveSoundEntry.PitchDataPoint);
+			CumulativeDistanceAttenuation = AccumulateDataPoints(CumulativeDistanceAttenuation, WaveSoundEntry.DistanceAttenuationDataPoint);
+			CumulativeLPFFreq = AccumulateDataPoints(CumulativeLPFFreq, WaveSoundEntry.LPFFreqDataPoint);
+			CumulativeHPFFreq = AccumulateDataPoints(CumulativeHPFFreq, WaveSoundEntry.HPFFreqDataPoint);
+			CumulativeRelativeRenderCost = AccumulateDataPoints(CumulativeRelativeRenderCost, WaveSoundEntry.RelativeRenderCostDataPoint);
 		}
 
 		AveragePeakAmp = AverageDataPoints(CumulativePeakAmplitude);
@@ -1099,15 +1083,32 @@ namespace UE::Audio::Insights
 		AverageLPFFreq = AverageDataPoints(CumulativeLPFFreq);
 		AverageHPFFreq = AverageDataPoints(CumulativeHPFFreq);
 
-		AddSingleValueToBuffer(ActiveSoundEntry.VolumeDataPoints, MaxVolume);
-		AddSingleValueToBuffer(ActiveSoundEntry.DistanceDataPoints, MinDistance);
-		AddSingleValueToBuffer(ActiveSoundEntry.DistanceAttenuationDataPoints, AverageDistanceAttenuation);
-		AddSingleValueToBuffer(ActiveSoundEntry.LPFFreqDataPoints, AverageLPFFreq);
-		AddSingleValueToBuffer(ActiveSoundEntry.HPFFreqDataPoints, AverageHPFFreq);
-		AddSingleValueToBuffer(ActiveSoundEntry.PriorityDataPoints, MaxPriority);
-		AddSingleValueToBuffer(ActiveSoundEntry.AmplitudeDataPoints, AveragePeakAmp);
-		AddSingleValueToBuffer(ActiveSoundEntry.PitchDataPoints, AveragePitch);
-		AddSingleValueToBuffer(ActiveSoundEntry.RelativeRenderCostDataPoints, CumulativeRelativeRenderCost.TotalAccumulationDataPoint);
+		AddSingleValueToBuffer(ActiveSoundEntry.VolumeDataRange, { ActiveSoundEntry.Timestamp, MaxVolume });
+		ActiveSoundEntry.VolumeDataPoint = MaxVolume;
+
+		AddSingleValueToBuffer(ActiveSoundEntry.DistanceDataRange, { ActiveSoundEntry.Timestamp, MinDistance });
+		ActiveSoundEntry.DistanceDataPoint = MinDistance;
+
+		AddSingleValueToBuffer(ActiveSoundEntry.DistanceAttenuationDataRange, { ActiveSoundEntry.Timestamp, AverageDistanceAttenuation });
+		ActiveSoundEntry.DistanceAttenuationDataPoint = AverageDistanceAttenuation;
+
+		AddSingleValueToBuffer(ActiveSoundEntry.LPFFreqDataRange, { ActiveSoundEntry.Timestamp, AverageLPFFreq });
+		ActiveSoundEntry.LPFFreqDataPoint = AverageLPFFreq;
+
+		AddSingleValueToBuffer(ActiveSoundEntry.HPFFreqDataRange, { ActiveSoundEntry.Timestamp, AverageHPFFreq });
+		ActiveSoundEntry.HPFFreqDataPoint = AverageHPFFreq;
+
+		AddSingleValueToBuffer(ActiveSoundEntry.PriorityDataRange, { ActiveSoundEntry.Timestamp, MaxPriority });
+		ActiveSoundEntry.PriorityDataPoint = MaxPriority;
+
+		AddSingleValueToBuffer(ActiveSoundEntry.AmplitudeDataRange, { ActiveSoundEntry.Timestamp, AveragePeakAmp });
+		ActiveSoundEntry.AmplitudeDataPoint = AveragePeakAmp;
+
+		AddSingleValueToBuffer(ActiveSoundEntry.PitchDataRange, { ActiveSoundEntry.Timestamp, AveragePitch });
+		ActiveSoundEntry.PitchDataPoint = AveragePitch;
+
+		AddSingleValueToBuffer(ActiveSoundEntry.RelativeRenderCostDataRange, { ActiveSoundEntry.Timestamp, CumulativeRelativeRenderCost.TotalAccumulationDataPoint });
+		ActiveSoundEntry.RelativeRenderCostDataPoint = CumulativeRelativeRenderCost.TotalAccumulationDataPoint;
 	}
 
 	void FSoundTraceProvider::OnTimingViewTimeMarkerChanged(double TimeMarker)
@@ -1440,112 +1441,160 @@ namespace UE::Audio::Insights
 
 		constexpr uint32 NumMessageTypesToProcess = static_cast<uint32>(ESoundParamMessageType::MAX);
 
+		const FAudioInsightsModule& AudioInsightsModule = UE::Audio::Insights::FAudioInsightsModule::GetChecked();
+		const TRange<double> PlottingRange = AudioInsightsModule.GetTimingViewExtender().GetPlottingRange();
+
+		// Add a small margin to either side to avoid gaps at the edges
+		constexpr double Margin = 0.2;
+		const double StartTime = FMath::Max(0.0f, PlottingRange.GetLowerBoundValue() - Margin);
+		const double EndTime = FMath::Max(0.0f, PlottingRange.GetUpperBoundValue() + Margin);
+
 		// Using ParallelFor to speed-up the cached messages retrieval, using a traditional for loop is unacceptably slower, specially in large traces.
 		ParallelFor(NumMessageTypesToProcess,
-		[&FoundEntries, &CacheManager, TimeMarker, this](const int32 Index)
+		[&FoundEntries, &CacheManager, TimeMarker, StartTime, EndTime, this](const int32 Index)
 		{
-			// 5s is the current window size of the plots widget
-			constexpr double LookbackTime = 5.0;
-			constexpr double Margin = 0.2;
-			const double StartTime = FMath::Max(0.0f, TimeMarker - LookbackTime - Margin);
-
 			switch (static_cast<ESoundParamMessageType>(Index))
 			{
 			case ESoundParamMessageType::Priority:
-				CacheManager.IterateOverRange<FSoundPriorityMessage>(SoundMessageNames::PriorityParam, StartTime, TimeMarker, [&FoundEntries](const FSoundPriorityMessage& Message)
+				CacheManager.IterateOverRange<FSoundPriorityMessage>(SoundMessageNames::PriorityParam, StartTime, EndTime, [&FoundEntries, TimeMarker](const FSoundPriorityMessage& Message)
 				{
 					if (TSharedPtr<FSoundDashboardEntry>* SoundEntry = FoundEntries.Find(Message.ActiveSoundPlayOrder))
 					{
-						(*SoundEntry)->PriorityDataPoints.Push({ Message.Timestamp, Message.Priority });
+						(*SoundEntry)->PriorityDataRange.Push({ Message.Timestamp, Message.Priority });
 						(*SoundEntry)->Timestamp = Message.Timestamp;
+
+						if (Message.Timestamp <= TimeMarker)
+						{
+							(*SoundEntry)->PriorityDataPoint = Message.Priority;
+						}
 					}
 				});
 				break;
 
 			case ESoundParamMessageType::Distance:
-				CacheManager.IterateOverRange<FSoundDistanceMessage>(SoundMessageNames::DistanceParam, StartTime, TimeMarker, [&FoundEntries](const FSoundDistanceMessage& Message)
+				CacheManager.IterateOverRange<FSoundDistanceMessage>(SoundMessageNames::DistanceParam, StartTime, EndTime, [&FoundEntries, TimeMarker](const FSoundDistanceMessage& Message)
 				{
 					if (TSharedPtr<FSoundDashboardEntry>* SoundEntry = FoundEntries.Find(Message.ActiveSoundPlayOrder))
 					{
-						(*SoundEntry)->DistanceDataPoints.Push({ Message.Timestamp, Message.Distance });
+						(*SoundEntry)->DistanceDataRange.Push({ Message.Timestamp, Message.Distance });
 						(*SoundEntry)->Timestamp = Message.Timestamp;
+
+						if (Message.Timestamp <= TimeMarker)
+						{
+							(*SoundEntry)->DistanceDataPoint = Message.Distance;
+						}
 					}
 				});
 				break;
 
 			case ESoundParamMessageType::DistanceAttenuation:
-				CacheManager.IterateOverRange<FSoundDistanceAttenuationMessage>(SoundMessageNames::DistanceAttenuationParam, StartTime, TimeMarker, [&FoundEntries](const FSoundDistanceAttenuationMessage& Message)
+				CacheManager.IterateOverRange<FSoundDistanceAttenuationMessage>(SoundMessageNames::DistanceAttenuationParam, StartTime, EndTime, [&FoundEntries, TimeMarker](const FSoundDistanceAttenuationMessage& Message)
 				{
 					if (TSharedPtr<FSoundDashboardEntry>* SoundEntry = FoundEntries.Find(Message.ActiveSoundPlayOrder))
 					{
-						(*SoundEntry)->DistanceAttenuationDataPoints.Push({ Message.Timestamp, Message.DistanceAttenuation });
+						(*SoundEntry)->DistanceAttenuationDataRange.Push({ Message.Timestamp, Message.DistanceAttenuation });
 						(*SoundEntry)->Timestamp = Message.Timestamp;
+
+						if (Message.Timestamp <= TimeMarker)
+						{
+							(*SoundEntry)->DistanceAttenuationDataPoint = Message.DistanceAttenuation;
+						}
 					}
 				});
 				break;
 
 			case ESoundParamMessageType::LPFFreq:
-				CacheManager.IterateOverRange<FSoundLPFFreqMessage>(SoundMessageNames::LPFFreqParam, StartTime, TimeMarker, [&FoundEntries](const FSoundLPFFreqMessage& Message)
+				CacheManager.IterateOverRange<FSoundLPFFreqMessage>(SoundMessageNames::LPFFreqParam, StartTime, EndTime, [&FoundEntries, TimeMarker](const FSoundLPFFreqMessage& Message)
 				{
 					if (TSharedPtr<FSoundDashboardEntry>* SoundEntry = FoundEntries.Find(Message.ActiveSoundPlayOrder))
 					{
-						(*SoundEntry)->LPFFreqDataPoints.Push({ Message.Timestamp, Message.LPFFrequency });
+						(*SoundEntry)->LPFFreqDataRange.Push({ Message.Timestamp, Message.LPFFrequency });
 						(*SoundEntry)->Timestamp = Message.Timestamp;
+
+						if (Message.Timestamp <= TimeMarker)
+						{
+							(*SoundEntry)->LPFFreqDataPoint = Message.LPFFrequency;
+						}
 					}
 				});
 				break;
 
 			case ESoundParamMessageType::HPFFreq:
-				CacheManager.IterateOverRange<FSoundHPFFreqMessage>(SoundMessageNames::HPFFreqParam, StartTime, TimeMarker, [&FoundEntries](const FSoundHPFFreqMessage& Message)
+				CacheManager.IterateOverRange<FSoundHPFFreqMessage>(SoundMessageNames::HPFFreqParam, StartTime, EndTime, [&FoundEntries, TimeMarker](const FSoundHPFFreqMessage& Message)
 				{
 					if (TSharedPtr<FSoundDashboardEntry>* SoundEntry = FoundEntries.Find(Message.ActiveSoundPlayOrder))
 					{
-						(*SoundEntry)->HPFFreqDataPoints.Push({ Message.Timestamp, Message.HPFFrequency });
+						(*SoundEntry)->HPFFreqDataRange.Push({ Message.Timestamp, Message.HPFFrequency });
 						(*SoundEntry)->Timestamp = Message.Timestamp;
+
+						if (Message.Timestamp <= TimeMarker)
+						{
+							(*SoundEntry)->HPFFreqDataPoint = Message.HPFFrequency;
+						}
 					}
 				});
 				break;
 
 			case ESoundParamMessageType::Envelope:
-				CacheManager.IterateOverRange<FSoundEnvelopeMessage>(SoundMessageNames::EnvelopeParam, StartTime, TimeMarker, [&FoundEntries](const FSoundEnvelopeMessage& Message)
+				CacheManager.IterateOverRange<FSoundEnvelopeMessage>(SoundMessageNames::EnvelopeParam, StartTime, EndTime, [&FoundEntries, TimeMarker](const FSoundEnvelopeMessage& Message)
 				{
 					if (TSharedPtr<FSoundDashboardEntry>* SoundEntry = FoundEntries.Find(Message.ActiveSoundPlayOrder))
 					{
-						(*SoundEntry)->AmplitudeDataPoints.Push({ Message.Timestamp, Message.Envelope });
+						(*SoundEntry)->AmplitudeDataRange.Push({ Message.Timestamp, Message.Envelope });
 						(*SoundEntry)->Timestamp = Message.Timestamp;
+
+						if (Message.Timestamp <= TimeMarker)
+						{
+							(*SoundEntry)->AmplitudeDataPoint = Message.Envelope;
+						}
 					}
 				});
 				break;
 
 			case ESoundParamMessageType::Volume:
-				CacheManager.IterateOverRange<FSoundVolumeMessage>(SoundMessageNames::VolumeParam, StartTime, TimeMarker, [&FoundEntries](const FSoundVolumeMessage& Message)
+				CacheManager.IterateOverRange<FSoundVolumeMessage>(SoundMessageNames::VolumeParam, StartTime, EndTime, [&FoundEntries, TimeMarker](const FSoundVolumeMessage& Message)
 				{
 					if (TSharedPtr<FSoundDashboardEntry>* SoundEntry = FoundEntries.Find(Message.ActiveSoundPlayOrder))
 					{
-						(*SoundEntry)->VolumeDataPoints.Push({ Message.Timestamp, Message.Volume });
+						(*SoundEntry)->VolumeDataRange.Push({ Message.Timestamp, Message.Volume });
 						(*SoundEntry)->Timestamp = Message.Timestamp;
+
+						if (Message.Timestamp <= TimeMarker)
+						{
+							(*SoundEntry)->VolumeDataPoint = Message.Volume;
+						}
 					}
 				});
 				break;
 
 			case ESoundParamMessageType::Pitch:
-				CacheManager.IterateOverRange<FSoundPitchMessage>(SoundMessageNames::PitchParam, StartTime, TimeMarker, [&FoundEntries](const FSoundPitchMessage& Message)
+				CacheManager.IterateOverRange<FSoundPitchMessage>(SoundMessageNames::PitchParam, StartTime, EndTime, [&FoundEntries, TimeMarker](const FSoundPitchMessage& Message)
 				{
 					if (TSharedPtr<FSoundDashboardEntry>* SoundEntry = FoundEntries.Find(Message.ActiveSoundPlayOrder))
 					{
-						(*SoundEntry)->PitchDataPoints.Push({ Message.Timestamp, Message.Pitch });
+						(*SoundEntry)->PitchDataRange.Push({ Message.Timestamp, Message.Pitch });
 						(*SoundEntry)->Timestamp = Message.Timestamp;
+
+						if (Message.Timestamp <= TimeMarker)
+						{
+							(*SoundEntry)->PitchDataPoint = Message.Pitch;
+						}
 					}
 				});
 				break;
 
 			case ESoundParamMessageType::RelativeRenderCost:
-				CacheManager.IterateOverRange<FSoundRelativeRenderCostMessage>(SoundMessageNames::RelativeRenderCostParam, StartTime, TimeMarker, [&FoundEntries](const FSoundRelativeRenderCostMessage& Message)
+				CacheManager.IterateOverRange<FSoundRelativeRenderCostMessage>(SoundMessageNames::RelativeRenderCostParam, StartTime, EndTime, [&FoundEntries, TimeMarker](const FSoundRelativeRenderCostMessage& Message)
 				{
 					if (TSharedPtr<FSoundDashboardEntry>* SoundEntry = FoundEntries.Find(Message.ActiveSoundPlayOrder))
 					{
-						(*SoundEntry)->RelativeRenderCostDataPoints.Push({ Message.Timestamp, Message.RelativeRenderCost });
+						(*SoundEntry)->RelativeRenderCostDataRange.Push({ Message.Timestamp, Message.RelativeRenderCost });
 						(*SoundEntry)->Timestamp = Message.Timestamp;
+
+						if (Message.Timestamp <= TimeMarker)
+						{
+							(*SoundEntry)->RelativeRenderCostDataPoint = Message.RelativeRenderCost;
+						}
 					}
 				});
 				break;
@@ -1674,50 +1723,46 @@ namespace UE::Audio::Insights
 					continue;
 				}
 
-				// Once we have gathered all of the plot data, we are only interested in the closest value to the time marker to display inside the Sounds Dashboard
-				// Changing our data buffers to single points helps avoid issues with ordering data while moving the timemarker back and forth
-				SoundEntry->SwitchBuffersToSingleValues();
-
-				SoundEntry->PriorityDataPoints.Push({ CachedEntryInfo.PriorityMessage.GetValue().Timestamp, CachedEntryInfo.PriorityMessage.GetValue().Priority });
+				SoundEntry->PriorityDataPoint = CachedEntryInfo.PriorityMessage.GetValue().Priority;
 
 				if (CachedEntryInfo.DistanceMessage.IsSet())
 				{
-					SoundEntry->DistanceDataPoints.Push({ CachedEntryInfo.DistanceMessage.GetValue().Timestamp, CachedEntryInfo.DistanceMessage.GetValue().Distance });
+					SoundEntry->DistanceDataPoint = CachedEntryInfo.DistanceMessage.GetValue().Distance;
 				}
 
 				if (CachedEntryInfo.DistanceAttenuationMessage.IsSet())
 				{
-					SoundEntry->DistanceAttenuationDataPoints.Push({ CachedEntryInfo.DistanceAttenuationMessage.GetValue().Timestamp, CachedEntryInfo.DistanceAttenuationMessage.GetValue().DistanceAttenuation });
+					SoundEntry->DistanceAttenuationDataPoint = CachedEntryInfo.DistanceAttenuationMessage.GetValue().DistanceAttenuation;
 				}
 
 				if (CachedEntryInfo.LPFFreqMessage.IsSet())
 				{
-					SoundEntry->LPFFreqDataPoints.Push({ CachedEntryInfo.LPFFreqMessage.GetValue().Timestamp, CachedEntryInfo.LPFFreqMessage.GetValue().LPFFrequency });
+					SoundEntry->LPFFreqDataPoint = CachedEntryInfo.LPFFreqMessage.GetValue().LPFFrequency;
 				}
 
 				if (CachedEntryInfo.HPFFreqMessage.IsSet())
 				{
-					SoundEntry->HPFFreqDataPoints.Push({ CachedEntryInfo.HPFFreqMessage.GetValue().Timestamp, CachedEntryInfo.HPFFreqMessage.GetValue().HPFFrequency });
+					SoundEntry->HPFFreqDataPoint = CachedEntryInfo.HPFFreqMessage.GetValue().HPFFrequency;
 				}
 
 				if (CachedEntryInfo.AmplitudeMessage.IsSet())
 				{
-					SoundEntry->AmplitudeDataPoints.Push({ CachedEntryInfo.AmplitudeMessage.GetValue().Timestamp, CachedEntryInfo.AmplitudeMessage.GetValue().Envelope });
+					SoundEntry->AmplitudeDataPoint = CachedEntryInfo.AmplitudeMessage.GetValue().Envelope;
 				}
 
 				if (CachedEntryInfo.VolumeMessage.IsSet())
 				{
-					SoundEntry->VolumeDataPoints.Push({ CachedEntryInfo.VolumeMessage.GetValue().Timestamp, CachedEntryInfo.VolumeMessage.GetValue().Volume });
+					SoundEntry->VolumeDataPoint = CachedEntryInfo.VolumeMessage.GetValue().Volume;
 				}
 
 				if (CachedEntryInfo.PitchMessage.IsSet())
 				{
-					SoundEntry->PitchDataPoints.Push({ CachedEntryInfo.PitchMessage.GetValue().Timestamp, CachedEntryInfo.PitchMessage.GetValue().Pitch });
+					SoundEntry->PitchDataPoint = CachedEntryInfo.PitchMessage.GetValue().Pitch;
 				}
 
 				if (CachedEntryInfo.RelativeRenderCostMessage.IsSet())
 				{
-					SoundEntry->RelativeRenderCostDataPoints.Push({ CachedEntryInfo.RelativeRenderCostMessage.GetValue().Timestamp, CachedEntryInfo.RelativeRenderCostMessage.GetValue().RelativeRenderCost });
+					SoundEntry->RelativeRenderCostDataPoint = CachedEntryInfo.RelativeRenderCostMessage.GetValue().RelativeRenderCost;
 				}
 
 				SoundEntry->Timestamp = CachedEntryInfo.PriorityMessage.GetValue().Timestamp;
diff --git a/Engine/Plugins/AudioInsights/Source/AudioInsights/Private/Views/SoundDashboardViewFactory.cpp b/Engine/Plugins/AudioInsights/Source/AudioInsights/Private/Views/SoundDashboardViewFactory.cpp
--- a/Engine/Plugins/AudioInsights/Source/AudioInsights/Private/Views/SoundDashboardViewFactory.cpp
+++ b/Engine/Plugins/AudioInsights/Source/AudioInsights/Private/Views/SoundDashboardViewFactory.cpp
@@ -49,17 +49,6 @@ namespace UE::Audio::Insights
 #endif // WITH_EDITOR
 		}
 
-		float GetLastEntryArrayValue(const ::Audio::TCircularAudioBuffer<FDataPoint>& InDataPoints)
-		{
-			if (InDataPoints.Num() > 0)
-			{
-				const ::Audio::DisjointedArrayView<const FDataPoint> DataPointsDisjointedArrayView = InDataPoints.PeekInPlace(InDataPoints.Num());
-				return DataPointsDisjointedArrayView.FirstBuffer.Last().Value;
-			}
-
-			return 0.0f;
-		};
-
 		bool SetFilteredVisibility(IDashboardDataTreeViewEntry& InEntry, const FString& InFilterString)
 		{
 			FSoundDashboardEntry& SoundEntry = FSoundDashboardViewFactoryPrivate::CastEntry(InEntry);
@@ -596,32 +585,15 @@ namespace UE::Audio::Insights
 			PinnedEntry.TimeoutTimestamp = OriginalEntry.TimeoutTimestamp;
 			PinnedEntry.bIsVisible = OriginalEntry.bIsVisible;
 
-			// Just copy the last entry in the buffer rather than the whole buffer
-			auto AddLastValueInBuffer = [](::Audio::TCircularAudioBuffer<FDataPoint>& To, const ::Audio::TCircularAudioBuffer<FDataPoint>& From)
-			{
-				if (From.Num() == 0)
-				{
-					return;
-				}
-
-				if (To.Num() > 0)
-				{
-					To.Pop(1);
-				}
-
-				const ::Audio::DisjointedArrayView<const FDataPoint> DataPointsDisjointedArrayView = From.PeekInPlace(From.Num());
-				To.Push(DataPointsDisjointedArrayView.FirstBuffer.Last());
-			};
-
-			AddLastValueInBuffer(PinnedEntry.PriorityDataPoints, OriginalEntry.PriorityDataPoints);
-			AddLastValueInBuffer(PinnedEntry.DistanceDataPoints, OriginalEntry.DistanceDataPoints);
-			AddLastValueInBuffer(PinnedEntry.DistanceAttenuationDataPoints, OriginalEntry.DistanceAttenuationDataPoints);
-			AddLastValueInBuffer(PinnedEntry.LPFFreqDataPoints, OriginalEntry.LPFFreqDataPoints);
-			AddLastValueInBuffer(PinnedEntry.HPFFreqDataPoints, OriginalEntry.HPFFreqDataPoints);
-			AddLastValueInBuffer(PinnedEntry.AmplitudeDataPoints, OriginalEntry.AmplitudeDataPoints);
-			AddLastValueInBuffer(PinnedEntry.VolumeDataPoints, OriginalEntry.VolumeDataPoints);
-			AddLastValueInBuffer(PinnedEntry.PitchDataPoints, OriginalEntry.PitchDataPoints);
-			AddLastValueInBuffer(PinnedEntry.RelativeRenderCostDataPoints, OriginalEntry.RelativeRenderCostDataPoints);
+			PinnedEntry.PriorityDataPoint = OriginalEntry.PriorityDataPoint;
+			PinnedEntry.DistanceDataPoint = OriginalEntry.DistanceDataPoint;
+			PinnedEntry.DistanceAttenuationDataPoint = OriginalEntry.DistanceAttenuationDataPoint;
+			PinnedEntry.LPFFreqDataPoint = OriginalEntry.LPFFreqDataPoint;
+			PinnedEntry.HPFFreqDataPoint = OriginalEntry.HPFFreqDataPoint;
+			PinnedEntry.AmplitudeDataPoint = OriginalEntry.AmplitudeDataPoint;
+			PinnedEntry.VolumeDataPoint = OriginalEntry.VolumeDataPoint;
+			PinnedEntry.PitchDataPoint = OriginalEntry.PitchDataPoint;
+			PinnedEntry.RelativeRenderCostDataPoint = OriginalEntry.RelativeRenderCostDataPoint;
 		};
 	} // namespace FPinnedSoundEntryWrapperPrivate
 
@@ -2032,7 +2004,7 @@ namespace UE::Audio::Insights
 							return FText::GetEmpty();
 						}
 
-						const float PriorityValue = GetLastEntryArrayValue(SoundDashboardEntry.PriorityDataPoints);
+						const float PriorityValue = SoundDashboardEntry.PriorityDataPoint;
 
 						// Max priority as defined in SoundWave.cpp
 						static constexpr float VolumeWeightedMaxPriority = TNumericLimits<float>::Max() / MAX_VOLUME;
@@ -2054,7 +2026,7 @@ namespace UE::Audio::Insights
 							return FText::GetEmpty();
 						}
 
-						return FText::AsNumber(GetLastEntryArrayValue(SoundDashboardEntry.DistanceDataPoints), FSlateStyle::Get().GetDefaultFloatFormat());
+						return FText::AsNumber(SoundDashboardEntry.DistanceDataPoint, FSlateStyle::Get().GetDefaultFloatFormat());
 					}
 				}
 			},
@@ -2069,7 +2041,7 @@ namespace UE::Audio::Insights
 							return FText::GetEmpty();
 						}
 
-						return FText::AsNumber(GetLastEntryArrayValue(SoundDashboardEntry.DistanceAttenuationDataPoints), FSlateStyle::Get().GetLinearVolumeFloatFormat());
+						return FText::AsNumber(SoundDashboardEntry.DistanceAttenuationDataPoint, FSlateStyle::Get().GetLinearVolumeFloatFormat());
 					}
 				}
 			},
@@ -2084,7 +2056,7 @@ namespace UE::Audio::Insights
 							return FText::GetEmpty();
 						}
 
-						const float AmplitudeValue = GetLastEntryArrayValue(SoundDashboardEntry.AmplitudeDataPoints);
+						const float AmplitudeValue = SoundDashboardEntry.AmplitudeDataPoint;
 
 						return bDisplayAmpPeakInDb 
 							? AudioInsightsUtils::ConvertToDecibelsText(AmplitudeValue) 
@@ -2103,7 +2075,7 @@ namespace UE::Audio::Insights
 							return FText::GetEmpty();
 						}
 
-						return FText::AsNumber(GetLastEntryArrayValue(SoundDashboardEntry.VolumeDataPoints), FSlateStyle::Get().GetLinearVolumeFloatFormat());
+						return FText::AsNumber(SoundDashboardEntry.VolumeDataPoint, FSlateStyle::Get().GetLinearVolumeFloatFormat());
 					}
 				}
 			},
@@ -2118,7 +2090,7 @@ namespace UE::Audio::Insights
 							return FText::GetEmpty();
 						}
 
-						return FText::AsNumber(GetLastEntryArrayValue(SoundDashboardEntry.LPFFreqDataPoints), FSlateStyle::Get().GetFreqFloatFormat());
+						return FText::AsNumber(SoundDashboardEntry.LPFFreqDataPoint, FSlateStyle::Get().GetFreqFloatFormat());
 					}
 				}
 			},
@@ -2133,7 +2105,7 @@ namespace UE::Audio::Insights
 							return FText::GetEmpty();
 						}
 
-						return FText::AsNumber(GetLastEntryArrayValue(SoundDashboardEntry.HPFFreqDataPoints), FSlateStyle::Get().GetFreqFloatFormat());
+						return FText::AsNumber(SoundDashboardEntry.HPFFreqDataPoint, FSlateStyle::Get().GetFreqFloatFormat());
 					}
 				}
 			},
@@ -2148,7 +2120,7 @@ namespace UE::Audio::Insights
 							return FText::GetEmpty();
 						}
 
-						return FText::AsNumber(GetLastEntryArrayValue(SoundDashboardEntry.PitchDataPoints), FSlateStyle::Get().GetPitchFloatFormat());
+						return FText::AsNumber(SoundDashboardEntry.PitchDataPoint, FSlateStyle::Get().GetPitchFloatFormat());
 					}
 				}
 			},
@@ -2163,7 +2135,7 @@ namespace UE::Audio::Insights
 							return FText::GetEmpty();
 						}
 
-						return FText::AsNumber(GetLastEntryArrayValue(SoundDashboardEntry.RelativeRenderCostDataPoints), FSlateStyle::Get().GetRelativeRenderCostFormat());
+						return FText::AsNumber(SoundDashboardEntry.RelativeRenderCostDataPoint, FSlateStyle::Get().GetRelativeRenderCostFormat());
 					}
 				}
 			},
@@ -2220,7 +2192,7 @@ namespace UE::Audio::Insights
 
 		auto SortByPriority = [&SortByPlayOrder](const FSoundDashboardEntry& First, const FSoundDashboardEntry& Second)
 		{
-			const float ComparisonDiff = GetLastEntryArrayValue(First.PriorityDataPoints) - GetLastEntryArrayValue(Second.PriorityDataPoints);
+			const float ComparisonDiff = First.PriorityDataPoint - Second.PriorityDataPoint;
 			if (FMath::IsNearlyZero(ComparisonDiff, UE_KINDA_SMALL_NUMBER))
 			{
 				return SortByPlayOrder(First, Second);
@@ -2231,7 +2203,7 @@ namespace UE::Audio::Insights
 
 		auto SortByDistance = [&SortByPlayOrder](const FSoundDashboardEntry& First, const FSoundDashboardEntry& Second)
 		{
-			const float ComparisonDiff = GetLastEntryArrayValue(First.DistanceDataPoints) - GetLastEntryArrayValue(Second.DistanceDataPoints);
+			const float ComparisonDiff = First.DistanceDataPoint- Second.DistanceDataPoint;
 			if (FMath::IsNearlyZero(ComparisonDiff, UE_KINDA_SMALL_NUMBER))
 			{
 				return SortByPlayOrder(First, Second);
@@ -2242,7 +2214,7 @@ namespace UE::Audio::Insights
 
 		auto SortByDistanceAttenuation = [&SortByPlayOrder](const FSoundDashboardEntry& First, const FSoundDashboardEntry& Second)
 		{
-			const float ComparisonDiff = GetLastEntryArrayValue(First.DistanceAttenuationDataPoints) - GetLastEntryArrayValue(Second.DistanceAttenuationDataPoints);
+			const float ComparisonDiff = First.DistanceAttenuationDataPoint - Second.DistanceAttenuationDataPoint;
 			if (FMath::IsNearlyZero(ComparisonDiff, UE_KINDA_SMALL_NUMBER))
 			{
 				return SortByPlayOrder(First, Second);
@@ -2253,7 +2225,7 @@ namespace UE::Audio::Insights
 
 		auto SortByLPFFreq = [&SortByPlayOrder](const FSoundDashboardEntry& First, const FSoundDashboardEntry& Second)
 		{
-			const float ComparisonDiff = GetLastEntryArrayValue(First.LPFFreqDataPoints) - GetLastEntryArrayValue(Second.LPFFreqDataPoints);
+			const float ComparisonDiff = First.LPFFreqDataPoint - Second.LPFFreqDataPoint;
 			if (FMath::IsNearlyZero(ComparisonDiff, UE_KINDA_SMALL_NUMBER))
 			{
 				return SortByPlayOrder(First, Second);
@@ -2264,7 +2236,7 @@ namespace UE::Audio::Insights
 
 		auto SortByHPFFreq = [&SortByPlayOrder](const FSoundDashboardEntry& First, const FSoundDashboardEntry& Second)
 		{
-			const float ComparisonDiff = GetLastEntryArrayValue(First.HPFFreqDataPoints) - GetLastEntryArrayValue(Second.HPFFreqDataPoints);
+			const float ComparisonDiff = First.HPFFreqDataPoint - Second.HPFFreqDataPoint;
 			if (FMath::IsNearlyZero(ComparisonDiff, UE_KINDA_SMALL_NUMBER))
 			{
 				return SortByPlayOrder(First, Second);
@@ -2275,7 +2247,7 @@ namespace UE::Audio::Insights
 
 		auto SortByAmplitude = [&SortByPlayOrder](const FSoundDashboardEntry& First, const FSoundDashboardEntry& Second)
 		{
-			const float ComparisonDiff = GetLastEntryArrayValue(First.AmplitudeDataPoints) - GetLastEntryArrayValue(Second.AmplitudeDataPoints);
+			const float ComparisonDiff = First.AmplitudeDataPoint - Second.AmplitudeDataPoint;
 			if (FMath::IsNearlyZero(ComparisonDiff, UE_KINDA_SMALL_NUMBER))
 			{
 				return SortByPlayOrder(First, Second);
@@ -2286,7 +2258,7 @@ namespace UE::Audio::Insights
 
 		auto SortByVolume = [&SortByPlayOrder](const FSoundDashboardEntry& First, const FSoundDashboardEntry& Second)
 		{
-			const float ComparisonDiff = GetLastEntryArrayValue(First.VolumeDataPoints) - GetLastEntryArrayValue(Second.VolumeDataPoints);
+			const float ComparisonDiff = First.VolumeDataPoint - Second.VolumeDataPoint;
 			if (FMath::IsNearlyZero(ComparisonDiff, UE_KINDA_SMALL_NUMBER))
 			{
 				return SortByPlayOrder(First, Second);
@@ -2297,7 +2269,7 @@ namespace UE::Audio::Insights
 
 		auto SortByPitch = [&SortByPlayOrder](const FSoundDashboardEntry& First, const FSoundDashboardEntry& Second)
 		{
-			const float ComparisonDiff = GetLastEntryArrayValue(First.PitchDataPoints) - GetLastEntryArrayValue(Second.PitchDataPoints);
+			const float ComparisonDiff = First.PitchDataPoint - Second.PitchDataPoint;
 			if (FMath::IsNearlyZero(ComparisonDiff, UE_KINDA_SMALL_NUMBER))
 			{
 				return SortByPlayOrder(First, Second);
@@ -2308,7 +2280,7 @@ namespace UE::Audio::Insights
 
 		auto SortByRelativeRenderCost = [&SortByPlayOrder](const FSoundDashboardEntry& First, const FSoundDashboardEntry& Second)
 		{
-			const float ComparisonDiff = GetLastEntryArrayValue(First.RelativeRenderCostDataPoints) - GetLastEntryArrayValue(Second.RelativeRenderCostDataPoints);
+			const float ComparisonDiff = First.RelativeRenderCostDataPoint - Second.RelativeRenderCostDataPoint;
 			if (FMath::IsNearlyZero(ComparisonDiff, UE_KINDA_SMALL_NUMBER))
 			{
 				return SortByPlayOrder(First, Second);
@@ -3226,7 +3198,7 @@ namespace UE::Audio::Insights
 					AmplitudeColumnName,
 					{
 						AmplitudeColumnDisplayName,
-						[](const IDashboardDataTreeViewEntry& InEntry) -> const ::Audio::TCircularAudioBuffer<FDataPoint>& { return CastEntry(InEntry).AmplitudeDataPoints; },
+						[](const IDashboardDataTreeViewEntry& InEntry) -> const ::Audio::TCircularAudioBuffer<FDataPoint>& { return CastEntry(InEntry).AmplitudeDataRange; },
 						[this](const float InValue) { return bDisplayAmpPeakInDb ? ::Audio::ConvertToDecibels(InValue) : InValue; },
 						[this]() { return bDisplayAmpPeakInDb ? FSlateStyle::Get().GetAmpFloatFormat() : FSlateStyle::Get().GetLinearVolumeFloatFormat(); },
 						[this]()
@@ -3272,7 +3244,7 @@ namespace UE::Audio::Insights
 					VolumeColumnName,
 					{
 						VolumeColumnDisplayName,
-						[](const IDashboardDataTreeViewEntry& InEntry) -> const ::Audio::TCircularAudioBuffer<FDataPoint>& { return CastEntry(InEntry).VolumeDataPoints; },
+						[](const IDashboardDataTreeViewEntry& InEntry) -> const ::Audio::TCircularAudioBuffer<FDataPoint>& { return CastEntry(InEntry).VolumeDataRange; },
 						nullptr,
 						[this]() { return FSlateStyle::Get().GetLinearVolumeFloatFormat(); },
 						[]() { return FSoundPlotsDashboardPlotRanges::StaticStruct()->FindPropertyByName(GET_MEMBER_NAME_CHECKED(FSoundPlotsDashboardPlotRanges, Volume)); },
@@ -3286,7 +3258,7 @@ namespace UE::Audio::Insights
 					DistanceColumnName,
 					{
 						DistanceColumnDisplayName,
-						[](const IDashboardDataTreeViewEntry& InEntry) -> const ::Audio::TCircularAudioBuffer<FDataPoint>& { return CastEntry(InEntry).DistanceDataPoints; },
+						[](const IDashboardDataTreeViewEntry& InEntry) -> const ::Audio::TCircularAudioBuffer<FDataPoint>& { return CastEntry(InEntry).DistanceDataRange; },
 						nullptr,
 						[this]() { return FSlateStyle::Get().GetFreqFloatFormat(); },
 						[]() { return FSoundPlotsDashboardPlotRanges::StaticStruct()->FindPropertyByName(GET_MEMBER_NAME_CHECKED(FSoundPlotsDashboardPlotRanges, Distance)); },
@@ -3300,7 +3272,7 @@ namespace UE::Audio::Insights
 					DistanceAttenuationColumnName,
 					{
 						DistanceAttenuationColumnDisplayName,
-						[](const IDashboardDataTreeViewEntry& InEntry) -> const ::Audio::TCircularAudioBuffer<FDataPoint>& { return CastEntry(InEntry).DistanceAttenuationDataPoints; },
+						[](const IDashboardDataTreeViewEntry& InEntry) -> const ::Audio::TCircularAudioBuffer<FDataPoint>& { return CastEntry(InEntry).DistanceAttenuationDataRange; },
 						nullptr,
 						[this]() { return FSlateStyle::Get().GetLinearVolumeFloatFormat(); },
 						[]() { return FSoundPlotsDashboardPlotRanges::StaticStruct()->FindPropertyByName(GET_MEMBER_NAME_CHECKED(FSoundPlotsDashboardPlotRanges, DistanceAttenuation)); },
@@ -3314,7 +3286,7 @@ namespace UE::Audio::Insights
 					PitchColumnName,
 					{
 						PitchColumnDisplayName,
-						[](const IDashboardDataTreeViewEntry& InEntry) -> const ::Audio::TCircularAudioBuffer<FDataPoint>& { return CastEntry(InEntry).PitchDataPoints; },
+						[](const IDashboardDataTreeViewEntry& InEntry) -> const ::Audio::TCircularAudioBuffer<FDataPoint>& { return CastEntry(InEntry).PitchDataRange; },
 						nullptr,
 						[this]() { return FSlateStyle::Get().GetPitchFloatFormat(); },
 						[]() { return FSoundPlotsDashboardPlotRanges::StaticStruct()->FindPropertyByName(GET_MEMBER_NAME_CHECKED(FSoundPlotsDashboardPlotRanges, Pitch)); },
@@ -3328,7 +3300,7 @@ namespace UE::Audio::Insights
 					PriorityColumnName,
 					{
 						PriorityColumnDisplayName,
-						[](const IDashboardDataTreeViewEntry& InEntry) -> const ::Audio::TCircularAudioBuffer<FDataPoint>& { return CastEntry(InEntry).PriorityDataPoints; },
+						[](const IDashboardDataTreeViewEntry& InEntry) -> const ::Audio::TCircularAudioBuffer<FDataPoint>& { return CastEntry(InEntry).PriorityDataRange; },
 						nullptr,
 						[this]() { return FSlateStyle::Get().GetLinearVolumeFloatFormat(); },
 						[]() { return FSoundPlotsDashboardPlotRanges::StaticStruct()->FindPropertyByName(GET_MEMBER_NAME_CHECKED(FSoundPlotsDashboardPlotRanges, Priority)); },
@@ -3342,7 +3314,7 @@ namespace UE::Audio::Insights
 					LPFFreqColumnName,
 					{
 						LPFFreqColumnDisplayName,
-						[](const IDashboardDataTreeViewEntry& InEntry) -> const ::Audio::TCircularAudioBuffer<FDataPoint>& { return CastEntry(InEntry).LPFFreqDataPoints; },
+						[](const IDashboardDataTreeViewEntry& InEntry) -> const ::Audio::TCircularAudioBuffer<FDataPoint>& { return CastEntry(InEntry).LPFFreqDataRange; },
 						nullptr,
 						[this]() { return FSlateStyle::Get().GetFreqFloatFormat(); },
 						[]() { return FSoundPlotsDashboardPlotRanges::StaticStruct()->FindPropertyByName(GET_MEMBER_NAME_CHECKED(FSoundPlotsDashboardPlotRanges, LPFFreq)); },
@@ -3356,7 +3328,7 @@ namespace UE::Audio::Insights
 					HPFFreqColumnName,
 					{
 						HPFFreqColumnDisplayName,
-						[](const IDashboardDataTreeViewEntry& InEntry) -> const ::Audio::TCircularAudioBuffer<FDataPoint>& { return CastEntry(InEntry).HPFFreqDataPoints; },
+						[](const IDashboardDataTreeViewEntry& InEntry) -> const ::Audio::TCircularAudioBuffer<FDataPoint>& { return CastEntry(InEntry).HPFFreqDataRange; },
 						nullptr,
 						[this]() { return FSlateStyle::Get().GetFreqFloatFormat(); },
 						[]() { return FSoundPlotsDashboardPlotRanges::StaticStruct()->FindPropertyByName(GET_MEMBER_NAME_CHECKED(FSoundPlotsDashboardPlotRanges, HPFFreq)); },
diff --git a/Engine/Plugins/AudioInsights/Source/AudioInsights/Public/Messages/SoundTraceMessages.h b/Engine/Plugins/AudioInsights/Source/AudioInsights/Public/Messages/SoundTraceMessages.h
--- a/Engine/Plugins/AudioInsights/Source/AudioInsights/Public/Messages/SoundTraceMessages.h
+++ b/Engine/Plugins/AudioInsights/Source/AudioInsights/Public/Messages/SoundTraceMessages.h
@@ -285,7 +285,6 @@ namespace UE::Audio::Insights
 		virtual void ResetHasSetInitExpansion() override { bHasSetInitExpansion = true; }
 
 		void ResetDataBuffers(const uint32 DataPointsCapacity);
-		void SwitchBuffersToSingleValues();
 
 		enum class EPinnedEntryType : uint8
 		{
@@ -308,15 +307,25 @@ namespace UE::Audio::Insights
 		bool bIsPlotActive = false;
 		bool bForceKeepEntryAlive = false;
 
-		::Audio::TCircularAudioBuffer<FDataPoint> PriorityDataPoints;
-		::Audio::TCircularAudioBuffer<FDataPoint> DistanceDataPoints;
-		::Audio::TCircularAudioBuffer<FDataPoint> DistanceAttenuationDataPoints;
-		::Audio::TCircularAudioBuffer<FDataPoint> LPFFreqDataPoints;
-		::Audio::TCircularAudioBuffer<FDataPoint> HPFFreqDataPoints;
-		::Audio::TCircularAudioBuffer<FDataPoint> AmplitudeDataPoints;
-		::Audio::TCircularAudioBuffer<FDataPoint> VolumeDataPoints;
-		::Audio::TCircularAudioBuffer<FDataPoint> PitchDataPoints;
-		::Audio::TCircularAudioBuffer<FDataPoint> RelativeRenderCostDataPoints;
+		::Audio::TCircularAudioBuffer<FDataPoint> PriorityDataRange;
+		::Audio::TCircularAudioBuffer<FDataPoint> DistanceDataRange;
+		::Audio::TCircularAudioBuffer<FDataPoint> DistanceAttenuationDataRange;
+		::Audio::TCircularAudioBuffer<FDataPoint> LPFFreqDataRange;
+		::Audio::TCircularAudioBuffer<FDataPoint> HPFFreqDataRange;
+		::Audio::TCircularAudioBuffer<FDataPoint> AmplitudeDataRange;
+		::Audio::TCircularAudioBuffer<FDataPoint> VolumeDataRange;
+		::Audio::TCircularAudioBuffer<FDataPoint> PitchDataRange;
+		::Audio::TCircularAudioBuffer<FDataPoint> RelativeRenderCostDataRange;
+
+		float PriorityDataPoint = 0.0f;
+		float DistanceDataPoint = 0.0f;
+		float DistanceAttenuationDataPoint = 0.0f;
+		float LPFFreqDataPoint = 20000.0f;
+		float HPFFreqDataPoint = 0.0f;
+		float AmplitudeDataPoint = -160.0f;
+		float VolumeDataPoint = 0.0f;
+		float PitchDataPoint = 1.0f;
+		float RelativeRenderCostDataPoint = 0.0f;
 
 		FText ActorLabel;
 
