[Audio Insights] Scrubbing timestamps via plots part 03: Enabling scrubbing via the plots widget
* User can now left-click/drag on a plot to inspect the timestamp at that position.
* Added ESC as a means to resume processing latest data
* Added tooltip to inform user on new functionality

diff --git a/Engine/Plugins/AudioInsights/Source/AudioInsights/Private/SoundPlotsWidgetCommands.cpp b/Engine/Plugins/AudioInsights/Source/AudioInsights/Private/SoundPlotsWidgetCommands.cpp
new file mode 100644
--- /dev/null
+++ b/Engine/Plugins/AudioInsights/Source/AudioInsights/Private/SoundPlotsWidgetCommands.cpp
@@ -0,0 +1,24 @@
+// Copyright Epic Games, Inc. All Rights Reserved.
+#include "SoundPlotsWidgetCommands.h"
+
+#include "AudioInsightsStyle.h"
+#include "Framework/Commands/InputChord.h"
+#include "Framework/Commands/UICommandInfo.h"
+
+#define LOCTEXT_NAMESPACE "AudioInsights"
+
+namespace UE::Audio::Insights
+{
+	FSoundPlotsWidgetCommands::FSoundPlotsWidgetCommands()
+		: TCommands<FSoundPlotsWidgetCommands>("SoundPlotsWidgetCommands", LOCTEXT("SoundPlotsWidgetCommands_ContextDescText", "Sound Plots Widget Commands"), NAME_None, UE::Audio::Insights::FSlateStyle::GetStyleName())
+	{
+	}
+
+	void FSoundPlotsWidgetCommands::RegisterCommands()
+	{
+		UI_COMMAND(ResetInspectTimestamp, "Reset Inspect Timestamp", "Resumes gathering the latest audio data and auto-scrolling in the Plots tab if previously paused.", EUserInterfaceActionType::Button, FInputChord(EKeys::Escape));
+	}
+
+} // namespace UE::Audio::Insights
+
+#undef LOCTEXT_NAMESPACE
diff --git a/Engine/Plugins/AudioInsights/Source/AudioInsights/Private/SoundPlotsWidgetCommands.h b/Engine/Plugins/AudioInsights/Source/AudioInsights/Private/SoundPlotsWidgetCommands.h
new file mode 100644
--- /dev/null
+++ b/Engine/Plugins/AudioInsights/Source/AudioInsights/Private/SoundPlotsWidgetCommands.h
@@ -0,0 +1,23 @@
+// Copyright Epic Games, Inc. All Rights Reserved.
+#pragma once
+
+#include "Framework/Commands/Commands.h"
+#include "Templates/SharedPointer.h"
+
+class FUICommandInfo;
+
+namespace UE::Audio::Insights
+{
+	class FSoundPlotsWidgetCommands : public TCommands<FSoundPlotsWidgetCommands>
+	{
+	public:
+		FSoundPlotsWidgetCommands();
+
+		virtual void RegisterCommands() override;
+
+		TSharedPtr<const FUICommandInfo> GetResetInspectTimestampCommand() const { return ResetInspectTimestamp; }
+
+	private:
+		TSharedPtr<FUICommandInfo> ResetInspectTimestamp;
+	};
+} // namespace UE::Audio::Insights
diff --git a/Engine/Plugins/AudioInsights/Source/AudioInsights/Private/Views/SAudioCurveView.cpp b/Engine/Plugins/AudioInsights/Source/AudioInsights/Private/Views/SAudioCurveView.cpp
--- a/Engine/Plugins/AudioInsights/Source/AudioInsights/Private/Views/SAudioCurveView.cpp
+++ b/Engine/Plugins/AudioInsights/Source/AudioInsights/Private/Views/SAudioCurveView.cpp
@@ -24,8 +24,11 @@ void SAudioCurveView::Construct( const SAudioCurveView::FArguments& InArgs )
 	YMargin = FMath::Clamp(InArgs._YMargin.Get(), 0.0f, 0.5f);
 	HorizontalAxisIncrement = InArgs._HorizontalAxisIncrement;
 	DesiredSize = InArgs._DesiredSize;
+	AdditionalToolTipText = InArgs._AdditionalToolTipText;
 	YDataRange = FFloatInterval(0.0f, 1.0f);
 	NumHorizontalGridLines = 10;
+	OnScrubPositionChanged = InArgs._OnScrubPositionChanged;
+	OnKeyDownHandler = InArgs._OnKeyDown;
 
 	// set clipping on by default, since the OnPaint function is drawing outside the bounds
 	Clipping = EWidgetClipping::ClipToBounds;
@@ -43,15 +46,53 @@ int32 SAudioCurveView::OnPaint( const FPaintArgs& Args, const FGeometry& Allotte
 	return FMath::Max( NewLayer, SCompoundWidget::OnPaint(Args, AllottedGeometry, MyCullingRect, OutDrawElements, NewLayer, InWidgetStyle, ShouldBeEnabled( bParentEnabled ) ) );
 }
 
+FReply SAudioCurveView::OnPreviewMouseButtonDown(const FGeometry& MyGeometry, const FPointerEvent& MouseEvent)
+{
+	if (MouseEvent.GetEffectingButton() == EKeys::LeftMouseButton)
+	{
+		bIsMouseDragging = true;
+	}
+
+	BroadcastScrubPosition();
+
+	return bIsMouseDragging ? FReply::Handled() : FReply::Unhandled();
+}
+
 FReply SAudioCurveView::OnMouseMove(const FGeometry& InMyGeometry, const FPointerEvent& InMouseEvent)
 {
 	UpdateCurveToolTip(InMyGeometry, InMouseEvent);
+	BroadcastScrubPosition();
+
 	return FReply::Handled();
 }
 
+FReply SAudioCurveView::OnMouseButtonUp(const FGeometry& MyGeometry, const FPointerEvent& MouseEvent)
+{
+	BroadcastScrubPosition();
+	
+	if (MouseEvent.GetEffectingButton() == EKeys::LeftMouseButton)
+	{
+		bIsMouseDragging = false;
+		return FReply::Handled();
+	}
+
+	return FReply::Unhandled();
+}
+
 void SAudioCurveView::OnMouseLeave(const FPointerEvent& MouseEvent)
 {
 	CrosshairPosition.Reset();
+	bIsMouseDragging = false;
+}
+
+FReply SAudioCurveView::OnPreviewKeyDown(const FGeometry& Geometry, const FKeyEvent& KeyEvent)
+{
+	if (OnKeyDownHandler.IsBound())
+	{
+		return OnKeyDownHandler.Execute(Geometry, KeyEvent);
+	}
+
+	return SCompoundWidget::OnPreviewKeyDown(Geometry, KeyEvent);
 }
 
 TSharedRef<SToolTip> SAudioCurveView::CreateCurveTooltip()
@@ -85,9 +126,26 @@ TSharedRef<SToolTip> SAudioCurveView::CreateCurveTooltip()
 				.Font(FCoreStyle::Get().GetFontStyle("ToolTip.LargerFont"))
 				.ColorAndOpacity(FLinearColor::Black)
 			]
+			+ SVerticalBox::Slot()
+			[
+				SNew(STextBlock)
+				.Text(AdditionalToolTipText)
+				.Font(FCoreStyle::Get().GetFontStyle("ToolTip.LargerFont"))
+				.ColorAndOpacity(FLinearColor::Black)
+				.Visibility_Lambda([this]() { return AdditionalToolTipText.IsSet() && !AdditionalToolTipText.Get().IsEmpty() ? EVisibility::Visible : EVisibility::Collapsed; })
+			]
 		];
 }
 
+void SAudioCurveView::BroadcastScrubPosition()
+{
+	if (bIsMouseDragging && OnScrubPositionChanged.IsBound() && CrosshairPosition.IsSet())
+	{
+		const FCurvePoint CrosshairPoint = CrosshairPosition.GetValue();
+		OnScrubPositionChanged.Execute(CrosshairPoint.Key, CrosshairPoint.Value);
+	}
+}
+
 void SAudioCurveView::UpdateCurveToolTip(const FGeometry& InMyGeometry, const FPointerEvent& InMouseEvent)
 {
 	if (!InMyGeometry.IsUnderLocation(InMouseEvent.GetScreenSpacePosition()) || !MetadataPerCurve.IsValid())
@@ -96,6 +154,11 @@ void SAudioCurveView::UpdateCurveToolTip(const FGeometry& InMyGeometry, const FP
 		return;
 	}
 
+	if (!PointDataPerCurve.IsValid())
+	{
+		return;
+	}
+
 	// Mouse position in widget space
 	const FVector2f HitPosition = InMyGeometry.AbsoluteToLocal(InMouseEvent.GetScreenSpacePosition());
 
@@ -119,6 +182,11 @@ void SAudioCurveView::UpdateCurveToolTip(const FGeometry& InMyGeometry, const FP
 		}
 
 		const TArray<FCurvePoint>* CurvePoints = PointDataPerCurve->Find(CurveId);
+		if (CurvePoints == nullptr)
+		{
+			continue;
+		}
+
 		const int32 NumPoints = CurvePoints->Num();
 		if (NumPoints > 0)
 		{
@@ -128,7 +196,7 @@ void SAudioCurveView::UpdateCurveToolTip(const FGeometry& InMyGeometry, const FP
 				const FCurvePoint& Point2 = (*CurvePoints)[i];
 
 				// Find points that contain mouse hit-point x
-				if (Point1.Key >= TargetX && TargetX <= Point2.Key)
+				if (Point1.Key <= TargetX && Point2.Key >= TargetX)
 				{
 					// Choose point with the smallest x delta
 					const float Delta1 = abs(TargetX - Point1.Key);
diff --git a/Engine/Plugins/AudioInsights/Source/AudioInsights/Private/Views/SoundPlotsWidgetView.cpp b/Engine/Plugins/AudioInsights/Source/AudioInsights/Private/Views/SoundPlotsWidgetView.cpp
--- a/Engine/Plugins/AudioInsights/Source/AudioInsights/Private/Views/SoundPlotsWidgetView.cpp
+++ b/Engine/Plugins/AudioInsights/Source/AudioInsights/Private/Views/SoundPlotsWidgetView.cpp
@@ -4,6 +4,7 @@
 #include "AudioInsightsModule.h"
 #include "AudioInsightsStyle.h"
 #include "Features/IModularFeatures.h"
+#include "SoundPlotsWidgetCommands.h"
 #include "Widgets/Input/SComboBox.h"
 #include "Widgets/Input/SNumericEntryBox.h"
 #include "Widgets/Input/SSegmentedControl.h"
@@ -32,6 +33,9 @@ namespace UE::Audio::Insights
 		FTraceModule& AudioInsightsTraceModule = static_cast<FTraceModule&>(FAudioInsightsModule::GetChecked().GetTraceModule());
 		AudioInsightsTraceModule.OnAnalysisStarting.AddRaw(this, &FSoundPlotsWidgetView::OnAnalysisStarting);
 
+		FSoundPlotsWidgetCommands::Register();
+		BindCommands();
+
 #if !WITH_EDITOR
 		FAudioInsightsModule& AudioInsightsModule = FAudioInsightsModule::GetChecked();
 		const TSharedPtr<FAudioInsightsComponent> AudioInsightsComponent = AudioInsightsModule.GetAudioInsightsComponent();
@@ -49,6 +53,8 @@ namespace UE::Audio::Insights
 
 	FSoundPlotsWidgetView::~FSoundPlotsWidgetView()
 	{
+		FSoundPlotsWidgetCommands::Unregister();
+
 		if (FModuleManager::Get().IsModuleLoaded("AudioInsights") && IModularFeatures::Get().IsModularFeatureAvailable(TraceServices::ModuleFeatureName))
 		{
 			FTraceModule& TraceModule = static_cast<FTraceModule&>(FAudioInsightsModule::GetChecked().GetTraceModule());
@@ -326,6 +332,11 @@ namespace UE::Audio::Insights
 #endif // WITH_EDITOR
 	}
 
+	FReply FSoundPlotsWidgetView::OnHandleKeyDown(const FGeometry& Geometry, const FKeyEvent& KeyEvent)
+	{
+		return (CommandList && CommandList->ProcessCommandBindings(KeyEvent)) ? FReply::Handled() : FReply::Unhandled();
+	}
+
 	TSharedRef<SWidget> FSoundPlotsWidgetView::MakePlotsWidget()
 	{
 		if (!PlotWidget.IsValid() || !ComboBox.IsValid())
@@ -470,7 +481,10 @@ namespace UE::Audio::Insights
 		SAssignNew(PlotWidget, SAudioCurveView)
 			.AutoRangeYAxis_Lambda([this]() { return !GetPlotColumnUseCustomYRange(SelectedPlotColumnName); })
 			.ViewRange_Lambda([this]() { return GetViewRange(); })
-			.PixelSnappingMethod(EWidgetPixelSnapping::Disabled);
+			.PixelSnappingMethod(EWidgetPixelSnapping::Disabled)
+			.AdditionalToolTipText_Lambda([this]() { return GetAdditionalTooltipText(); })
+			.OnScrubPositionChanged_Lambda([this](const double TimestampPosition, const float YValue) { SetTimestampFromCurveView(TimestampPosition); })
+			.OnKeyDown_Lambda([this](const FGeometry& Geometry, const FKeyEvent& KeyEvent) { return OnHandleKeyDown(Geometry, KeyEvent); });
 
 		PlotWidget->SetYValueFormattingOptions(*GetPlotColumnNumberFormat(SelectedPlotColumnName));
 		PlotWidget->SetYAxisRange(GetPlotColumnCustomizedOrDefaultYRange(SelectedPlotColumnName));
@@ -508,6 +522,20 @@ namespace UE::Audio::Insights
 			];
 	}
 
+	void FSoundPlotsWidgetView::BindCommands()
+	{
+		CommandList = MakeShared<FUICommandList>();
+
+		const FSoundPlotsWidgetCommands& Commands = FSoundPlotsWidgetCommands::Get();
+		CommandList->MapAction(Commands.GetResetInspectTimestampCommand(), FExecuteAction::CreateLambda([this]() { ResetInspectTimestamp(); }));
+	}
+
+	void FSoundPlotsWidgetView::ResetInspectTimestamp()
+	{
+		FAudioInsightsModule& AudioInsightsModule = FAudioInsightsModule::GetChecked();
+		AudioInsightsModule.GetTimingViewExtender().ResetMessageProcessType();
+	}
+
 	void FSoundPlotsWidgetView::CollectDataRecursive(const TSharedPtr<IDashboardDataTreeViewEntry>& Entry, const TArray<TSharedPtr<IDashboardDataTreeViewEntry>>& SelectedEntries, bool& bOutHasNewMetadata)
 	{
 		using namespace FSoundPlotsWidgetViewPrivate;
@@ -717,6 +745,17 @@ namespace UE::Audio::Insights
 		return FText::GetEmpty();
 	}
 
+	FText FSoundPlotsWidgetView::GetAdditionalTooltipText() const
+	{
+		FAudioInsightsModule& AudioInsightsModule = UE::Audio::Insights::FAudioInsightsModule::GetChecked();
+		if (AudioInsightsModule.GetTimingViewExtender().GetMessageCacheAndProcessingStatus() != ECacheAndProcess::Latest)
+		{
+			return LOCTEXT("SoundPlotsWidgetView_ResetTimeControlMethod", "Audio Insights has been paused. Press Esc to resume.");
+		}
+
+		return LOCTEXT("SoundPlotsWidgetView_InspectTimestampHint", "Left Click to inspect timestamp.");
+	}
+
 	TOptional<float> FSoundPlotsWidgetView::GetPlotColumnYRangeFloatMetaData(const FName& ColumnName, const FName& MetaDataKey)
 	{
 		if (const FSoundPlotsWidgetView::FPlotColumnInfo* PlotColumnInfo = ColumnInfo.Find(ColumnName))
@@ -793,8 +832,30 @@ namespace UE::Audio::Insights
 		}
 	}
 
+	void FSoundPlotsWidgetView::SetTimestampFromCurveView(const double Timestamp)
+	{
+		// Timestamp from SAudioCurve is a relative one - we need to transform it back to the actual timestamp
+#if WITH_EDITOR
+		const double AbsoluteTimestamp = Timestamp + BeginTimestamp;
+		const TRange<double> AbsolutePlottingRange = TRange<double>(CurrentRange.GetLowerBoundValue() + BeginTimestamp, CurrentRange.GetUpperBoundValue() + BeginTimestamp);
+#else
+		const double AbsoluteTimestamp = Timestamp;
+		const TRange<double> AbsolutePlottingRange = CurrentRange;
+#endif // WITH_EDITOR
+
+		FAudioInsightsModule& AudioInsightsModule = UE::Audio::Insights::FAudioInsightsModule::GetChecked();
+		AudioInsightsModule.GetTimingViewExtender().SetTimeMarker(AbsoluteTimestamp, ESystemControllingTimeMarker::PlotsWidget, AbsolutePlottingRange);
+	}
+
 	TRange<double> FSoundPlotsWidgetView::GetViewRange()
 	{
+		const FAudioInsightsModule& AudioInsightsModule = UE::Audio::Insights::FAudioInsightsModule::GetChecked();
+		TOptional<ESystemControllingTimeMarker> ControllingSystem = AudioInsightsModule.GetTimingViewExtender().TryGetSystemControllingTimeMarker();
+		if (ControllingSystem.IsSet() && ControllingSystem.GetValue() == ESystemControllingTimeMarker::PlotsWidget)
+		{
+			return CurrentRange;
+		}
+
 		if (BeginTimestamp == TNumericLimits<double>::Max())
 		{
 			return TRange<double>(0, MaxPlotHistorySeconds);
@@ -823,7 +884,8 @@ namespace UE::Audio::Insights
 		const TSharedPtr<const FAudioInsightsComponent> AudioInsightsComponent = FAudioInsightsModule::GetChecked().GetAudioInsightsComponent();
 		if (AudioInsightsComponent.IsValid() && !AudioInsightsComponent->GetIsLiveSession())
 		{
-			return TRange<double>(CurrentTimestamp - MaxPlotHistorySeconds, CurrentTimestamp);
+			CurrentRange = TRange<double>(CurrentTimestamp - MaxPlotHistorySeconds, CurrentTimestamp);
+			return CurrentRange;
 		}
 		else
 		{
@@ -841,7 +903,8 @@ namespace UE::Audio::Insights
 		}
 #endif // WITH_EDITOR
 
-		return TRange<double>(RangeUpperBound - MaxPlotHistorySeconds, RangeUpperBound);
+		CurrentRange = TRange<double>(RangeUpperBound - MaxPlotHistorySeconds, RangeUpperBound);
+		return CurrentRange;
 	}
 
 	bool FSoundPlotsWidgetView::GetIsPlotEnabledForEntry(const IDashboardDataTreeViewEntry& InEntry) const
diff --git a/Engine/Plugins/AudioInsights/Source/AudioInsights/Public/Views/SAudioCurveView.h b/Engine/Plugins/AudioInsights/Source/AudioInsights/Public/Views/SAudioCurveView.h
--- a/Engine/Plugins/AudioInsights/Source/AudioInsights/Public/Views/SAudioCurveView.h
+++ b/Engine/Plugins/AudioInsights/Source/AudioInsights/Public/Views/SAudioCurveView.h
@@ -2,6 +2,7 @@
 
 #pragma once
 
+#include "Framework/SlateDelegates.h"
 #include "SSimpleTimeSlider.h"
 #include "Widgets/SCompoundWidget.h"
 #include "Widgets/DeclarativeSyntaxSupport.h"
@@ -18,6 +19,8 @@ public:
 	// A curve point is a (double XValue, float YValue) pair
 	using FCurvePoint = TPair<double, float>;
 
+	DECLARE_DELEGATE_TwoParams(FOnScrubPositionChanged, double, float)
+
 	struct FCurveMetadata
 	{
 		FLinearColor CurveColor;
@@ -33,6 +36,7 @@ public:
 		, _GridLineColor(FLinearColor(0.5f, 0.5f, 0.5f, 0.25f))
 		, _AxesLabelColor(FLinearColor::White)
 		, _DesiredSize(FVector2D(100.f,100.f))
+		, _AdditionalToolTipText(FText::GetEmpty())
 	{}
     	/** View X axis range (in value space) */
     	SLATE_ATTRIBUTE(TRange<double>, ViewRange);
@@ -51,6 +55,15 @@ public:
 
 		/** Desired widget size */
 		SLATE_ATTRIBUTE(FVector2D, DesiredSize);
+
+		/** Additional text displayed below the standard tooltip */
+		SLATE_ATTRIBUTE(FText, AdditionalToolTipText);
+
+		/** Called when user clicks/drags on the widget. Only returns when there is a valid crosshair position */
+		SLATE_EVENT(FOnScrubPositionChanged, OnScrubPositionChanged);
+
+		/** Called when user presses a key when the widget is in focus */
+		SLATE_EVENT(FOnKeyDown, OnKeyDown);
 	
 	SLATE_END_ARGS()
 
@@ -84,8 +97,12 @@ protected:
 	// SWidget interface
 	virtual FVector2D ComputeDesiredSize(float) const override;
 	virtual int32 OnPaint( const FPaintArgs& Args, const FGeometry& AllottedGeometry, const FSlateRect& MyCullingRect, FSlateWindowElementList& OutDrawElements, int32 LayerId, const FWidgetStyle& InWidgetStyle, bool bParentEnabled ) const override;
+	virtual FReply OnPreviewMouseButtonDown(const FGeometry& MyGeometry, const FPointerEvent& MouseEvent) override;
 	virtual FReply OnMouseMove(const FGeometry& MyGeometry, const FPointerEvent& MouseEvent) override;
+	virtual FReply OnMouseButtonUp(const FGeometry& MyGeometry, const FPointerEvent& MouseEvent) override;
 	virtual void OnMouseLeave(const FPointerEvent& MouseEvent) override;
+	virtual FReply OnPreviewKeyDown(const FGeometry& Geometry, const FKeyEvent& KeyEvent) override;
+	virtual bool SupportsKeyboardFocus() const override { return true; }
 	
 	int32 PaintCurves(const FGeometry& AllottedGeometry, const FSlateRect& MyCullingRect, FSlateWindowElementList& OutDrawElements, int32 LayerId, const FWidgetStyle& InWidgetStyle, bool bParentEnabled) const;
 	int32 PaintGridLines(const FGeometry& AllottedGeometry, const FSlateRect& MyCullingRect, FSlateWindowElementList& OutDrawElements, int32 LayerId, const FWidgetStyle& InWidgetStyle, bool bParentEnabled, const SSimpleTimeSlider::FScrubRangeToScreen& RangeToScreen) const;
@@ -95,6 +112,8 @@ protected:
 	void UpdateYDataRange();
 	TSharedRef<SToolTip> CreateCurveTooltip();
 
+	void BroadcastScrubPosition();
+
 	TAttribute<TRange<double>> ViewRange;
 	TAttribute<bool> AutoRangeYAxis;
 	TAttribute<float> YMargin;
@@ -102,6 +121,10 @@ protected:
 	TAttribute<FLinearColor> GridLineColor;
 	TAttribute<FLinearColor> AxesLabelColor;
 	TAttribute<FVector2D> DesiredSize;
+	TAttribute<FText> AdditionalToolTipText;
+
+	FOnScrubPositionChanged OnScrubPositionChanged;
+	FOnKeyDown OnKeyDownHandler;
 
 	// Point data and metadata, keyed by curve id
 	TSharedPtr<TMap<int32, TArray<FCurvePoint>>> PointDataPerCurve;
@@ -123,4 +146,5 @@ protected:
 	FSlateFontInfo LabelFont;
 
 	TOptional<FCurvePoint> CrosshairPosition;
+	bool bIsMouseDragging = false;
 };
diff --git a/Engine/Plugins/AudioInsights/Source/AudioInsights/Public/Views/SoundPlotsWidgetView.h b/Engine/Plugins/AudioInsights/Source/AudioInsights/Public/Views/SoundPlotsWidgetView.h
--- a/Engine/Plugins/AudioInsights/Source/AudioInsights/Public/Views/SoundPlotsWidgetView.h
+++ b/Engine/Plugins/AudioInsights/Source/AudioInsights/Public/Views/SoundPlotsWidgetView.h
@@ -3,6 +3,7 @@
 
 #include "AudioInsightsDataSource.h"
 #include "DSP/Dsp.h"
+#include "Framework/Commands/UICommandList.h"
 #include "Templates/Function.h"
 #include "Views/SAudioCurveView.h"
 
@@ -61,8 +62,12 @@ namespace UE::Audio::Insights
 		void OnTimingViewTimeMarkerChanged(double InTimeMarker);
 		void OnTimeControlMethodReset();
 
+		FReply OnHandleKeyDown(const FGeometry& Geometry, const FKeyEvent& KeyEvent);
+
 		TSharedRef<SWidget> MakePlotsWidget();
 		void ConstructPlotsWidgetContents();
+		void BindCommands();
+		void ResetInspectTimestamp();
 
 		void CollectDataRecursive(const TSharedPtr<IDashboardDataTreeViewEntry>& Entry, const TArray<TSharedPtr<IDashboardDataTreeViewEntry>>& SelectedEntries, bool& bOutHasNewMetadata);
 		void ResetPlots();
@@ -75,6 +80,7 @@ namespace UE::Audio::Insights
 		const FNumberFormattingOptions* GetPlotColumnNumberFormat(const FName& ColumnName);
 		FFloatInterval GetPlotColumnCustomizedOrDefaultYRange(const FName& ColumnName);
 		const FText GetPlotColumnDisplayName(const FName& ColumnName);
+		FText GetAdditionalTooltipText() const;
 
 		TOptional<float> GetPlotColumnYRangeFloatMetaData(const FName& ColumnName, const FName& MetaDataKey);
 		TOptional<FFloatInterval> GetPlotColumnCustomizedYRange(const FName& ColumnName);
@@ -82,6 +88,7 @@ namespace UE::Audio::Insights
 		void SetPlotColumnUseCustomYRange(const FName& ColumnName, bool bUseCustomYRange);
 		void SetPlotColumnCustomYRangeMin(const FName& ColumnName, float Value, TOptional<ETextCommit::Type> CommitType = NullOpt);
 		void SetPlotColumnCustomYRangeMax(const FName& ColumnName, float Value, TOptional<ETextCommit::Type> CommitType = NullOpt);
+		void SetTimestampFromCurveView(const double Timestamp);
 
 		TRange<double> GetViewRange();
 
@@ -112,6 +119,8 @@ namespace UE::Audio::Insights
 		double BeginTimestamp = TNumericLimits<double>::Max();
 		double CurrentTimestamp = TNumericLimits<double>::Lowest();
 
+		TRange<double> CurrentRange { TNumericLimits<double>::Lowest(), TNumericLimits<double>::Lowest() };
+
 		enum class EGameState : uint8
 		{
 			Running,
@@ -127,6 +136,8 @@ namespace UE::Audio::Insights
 		TSharedPtr<SComboBox<FName>> ComboBox;
 		FName SelectedPlotColumnName;
 
+		TSharedPtr<FUICommandList> CommandList;
+
 #if !WITH_EDITOR
 		double PreviousTime = 0.0;
 		double CurrentRangeUpperBound = 0.0;
