[Audio Insights] Fix parameter values not updating in the sounds tab when scrubbing in standalone
* When loading in .utrace files into Audio Insights, there is a very low chance (1 in 1000) that two parameter messages will be sent to OnEvent inside SoundTraceProvider out of order (earlier message last)
* Having the TPagedArrays not sorted in timestamp order causes the binary search algorithm to return early, as it detects the "closest" message at the point the two parameter messages were sent out of order
* Sorting the paged arrays after the files have been loaded is one option, but we are concerned about the performance hit on very large trace files.
* This fix drops any new parameter message that is older than the current head of the TPagedArray. This results in a very small amount of dropped data, but in a stream of hundreds of points per sound per second the effect is negligible

diff --git a/Engine/Plugins/AudioInsights/Source/AudioInsights/Private/Providers/SoundTraceProvider.cpp b/Engine/Plugins/AudioInsights/Source/AudioInsights/Private/Providers/SoundTraceProvider.cpp
--- a/Engine/Plugins/AudioInsights/Source/AudioInsights/Private/Providers/SoundTraceProvider.cpp
+++ b/Engine/Plugins/AudioInsights/Source/AudioInsights/Private/Providers/SoundTraceProvider.cpp
@@ -89,6 +89,19 @@ namespace UE::Audio::Insights
 			OutDataRange.Push({ Msg.Timestamp, DataPoint });
 		}
 
+#if !WITH_EDITOR
+		template<typename T>
+		void CacheStandaloneParameterEntry(const T& Msg, TraceServices::TPagedArray<T>& OutCachedMessages)
+		{
+			// When loading in trace files, we very occasionally can get parameter messages out of order
+			// In these situations we drop the one message that was sent out of sync
+			if (OutCachedMessages.Num() == 0 || Msg.Timestamp >= OutCachedMessages.Last().Timestamp)
+			{
+				OutCachedMessages.EmplaceBack(Msg);
+			}
+		}
+#endif // !WITH_EDITOR
+
 		//////////////////////////////////////////////////////////////////////////////////////////////////////////////
 		// Finds an Active Sound entry inside a Category entry
 		void FindActiveSoundEntryInCategory(const TSharedPtr<FSoundDashboardEntry>* CategoryEntry, const uint32 ActiveSoundPlayOrder, TSharedPtr<FSoundDashboardEntry>& OutActiveSoundEntry)
@@ -415,6 +428,8 @@ namespace UE::Audio::Insights
 
 	bool FSoundTraceProvider::ProcessMessages()
 	{
+		using namespace FSoundTraceProviderPrivate;
+
 		// Helper lambdas
 		TSharedPtr<FSoundDashboardEntry> ActiveSoundEntryReturn = nullptr;
 		auto GetActiveSoundEntryFromActiveStartMessage = [this, &ActiveSoundEntryReturn](const FSoundStartMessage& Msg) -> TSharedPtr<FSoundDashboardEntry>*
@@ -586,7 +601,7 @@ namespace UE::Audio::Insights
 #if !WITH_EDITOR
 			if (SessionCachedMessages.IsValid())
 			{
-				SessionCachedMessages->PriorityCachedMessages.EmplaceBack(Msg);
+				CacheStandaloneParameterEntry(Msg, SessionCachedMessages->PriorityCachedMessages);
 			}
 #endif // !WITH_EDITOR
 
@@ -606,7 +621,7 @@ namespace UE::Audio::Insights
 #if !WITH_EDITOR
 			if (SessionCachedMessages.IsValid())
 			{
-				SessionCachedMessages->DistanceCachedMessages.EmplaceBack(Msg);
+				CacheStandaloneParameterEntry(Msg, SessionCachedMessages->DistanceCachedMessages);
 			}
 #endif // !WITH_EDITOR
 
@@ -626,7 +641,7 @@ namespace UE::Audio::Insights
 #if !WITH_EDITOR
 			if (SessionCachedMessages.IsValid())
 			{
-				SessionCachedMessages->DistanceAttenuationCachedMessages.EmplaceBack(Msg);
+				CacheStandaloneParameterEntry(Msg, SessionCachedMessages->DistanceAttenuationCachedMessages);
 			}
 #endif // !WITH_EDITOR
 
@@ -646,7 +661,7 @@ namespace UE::Audio::Insights
 #if !WITH_EDITOR
 			if (SessionCachedMessages.IsValid())
 			{
-				SessionCachedMessages->LPFFreqCachedMessages.EmplaceBack(Msg);
+				CacheStandaloneParameterEntry(Msg, SessionCachedMessages->LPFFreqCachedMessages);
 			}
 #endif // !WITH_EDITOR
 
@@ -666,7 +681,7 @@ namespace UE::Audio::Insights
 #if !WITH_EDITOR
 			if (SessionCachedMessages.IsValid())
 			{
-				SessionCachedMessages->HPFFreqCachedMessages.EmplaceBack(Msg);
+				CacheStandaloneParameterEntry(Msg, SessionCachedMessages->HPFFreqCachedMessages);
 			}
 #endif // !WITH_EDITOR
 
@@ -686,7 +701,7 @@ namespace UE::Audio::Insights
 #if !WITH_EDITOR
 			if (SessionCachedMessages.IsValid())
 			{
-				SessionCachedMessages->AmplitudeCachedMessages.EmplaceBack(Msg);
+				CacheStandaloneParameterEntry(Msg, SessionCachedMessages->AmplitudeCachedMessages);
 			}
 #endif // !WITH_EDITOR
 
@@ -706,7 +721,7 @@ namespace UE::Audio::Insights
 #if !WITH_EDITOR
 			if (SessionCachedMessages.IsValid())
 			{
-				SessionCachedMessages->VolumeCachedMessages.EmplaceBack(Msg);
+				CacheStandaloneParameterEntry(Msg, SessionCachedMessages->VolumeCachedMessages);
 			}
 #endif // !WITH_EDITOR
 
@@ -726,7 +741,7 @@ namespace UE::Audio::Insights
 #if !WITH_EDITOR
 			if (SessionCachedMessages.IsValid())
 			{
-				SessionCachedMessages->PitchCachedMessages.EmplaceBack(Msg);
+				CacheStandaloneParameterEntry(Msg, SessionCachedMessages->PitchCachedMessages);
 			}
 #endif // !WITH_EDITOR
 
@@ -746,7 +761,7 @@ namespace UE::Audio::Insights
 #if !WITH_EDITOR
 			if (SessionCachedMessages.IsValid())
 			{
-				SessionCachedMessages->RelativeRenderCostCachedMessages.EmplaceBack(Msg);
+				CacheStandaloneParameterEntry(Msg, SessionCachedMessages->RelativeRenderCostCachedMessages);
 			}
 #endif // !WITH_EDITOR
 
