[AudioWidgets][Audio Insights] Use MIN_VOLUME_DECIBELS from AudioDefines.h instead of hardcoding -160.0f value.

diff --git a/Engine/Plugins/AudioInsights/Source/AudioInsights/Private/Views/OutputMeterDashboardViewFactory.cpp b/Engine/Plugins/AudioInsights/Source/AudioInsights/Private/Views/OutputMeterDashboardViewFactory.cpp
--- a/Engine/Plugins/AudioInsights/Source/AudioInsights/Private/Views/OutputMeterDashboardViewFactory.cpp
+++ b/Engine/Plugins/AudioInsights/Source/AudioInsights/Private/Views/OutputMeterDashboardViewFactory.cpp
@@ -1,6 +1,7 @@
 // Copyright Epic Games, Inc. All Rights Reserved.
 #include "Views/OutputMeterDashboardViewFactory.h"
 
+#include "AudioDefines.h"
 #include "AudioInsightsStyle.h"
 #include "DSP/Dsp.h"
 #include "ISessionServicesModule.h"
@@ -145,8 +146,6 @@ namespace UE::Audio::Insights
 
 		if (SubmixDashboardEntry.IsValid())
 		{
-			constexpr float MinDecibelValue = -160.0f;
-
 			const FAudioMeterInfo& AudioMeterInfo = *SubmixDashboardEntry->AudioMeterInfo;
 			const TArray<float>& EnvelopeValues = AudioMeterInfo.EnvelopeValues;
 
@@ -155,7 +154,7 @@ namespace UE::Audio::Insights
 				MeterChannelInfo.Reserve(EnvelopeValues.Num());
 				for (const float EnvelopeValue : EnvelopeValues)
 				{
-					MeterChannelInfo.Emplace(::Audio::ConvertToDecibels(EnvelopeValue), MinDecibelValue, MinDecibelValue);
+					MeterChannelInfo.Emplace(::Audio::ConvertToDecibels(EnvelopeValue), MIN_VOLUME_DECIBELS, MIN_VOLUME_DECIBELS);
 				}
 			}
 			else
@@ -163,7 +162,7 @@ namespace UE::Audio::Insights
 				MeterChannelInfo.Reserve(AudioMeterInfo.NumChannels);
 				while (MeterChannelInfo.Num() < AudioMeterInfo.NumChannels)
 				{
-					MeterChannelInfo.Emplace(MinDecibelValue, MinDecibelValue, MinDecibelValue);
+					MeterChannelInfo.Emplace(MIN_VOLUME_DECIBELS, MIN_VOLUME_DECIBELS, MIN_VOLUME_DECIBELS);
 				}
 			}
 		}
diff --git a/Engine/Plugins/AudioInsights/Source/AudioInsights/Private/Views/SAudioMeterView.cpp b/Engine/Plugins/AudioInsights/Source/AudioInsights/Private/Views/SAudioMeterView.cpp
--- a/Engine/Plugins/AudioInsights/Source/AudioInsights/Private/Views/SAudioMeterView.cpp
+++ b/Engine/Plugins/AudioInsights/Source/AudioInsights/Private/Views/SAudioMeterView.cpp
@@ -2,6 +2,7 @@
 
 #include "Views/SAudioMeterView.h"
 
+#include "AudioDefines.h"
 #include "AudioInsightsDataSource.h"
 #include "DSP/Dsp.h"
 #include "Providers/AudioMeterProvider.h"
@@ -30,7 +31,7 @@ namespace UE::Audio::Insights
 
 						for (const float EnvelopeValue : EnvelopeValues)
 						{
-							MeterChannelInfo.Emplace(::Audio::ConvertToDecibels(EnvelopeValue), -160.0f, -160.0f);
+							MeterChannelInfo.Emplace(::Audio::ConvertToDecibels(EnvelopeValue), MIN_VOLUME_DECIBELS, MIN_VOLUME_DECIBELS);
 						}
 					}
 					else
@@ -39,7 +40,7 @@ namespace UE::Audio::Insights
 							
 						for (int32 Index = 0; Index < NumChannels; ++Index)
 						{
-							MeterChannelInfo.Emplace(-160.0f, -160.0f, -160.0f);
+							MeterChannelInfo.Emplace(MIN_VOLUME_DECIBELS, MIN_VOLUME_DECIBELS, MIN_VOLUME_DECIBELS);
 						}
 					}
 
diff --git a/Engine/Plugins/AudioInsights/Source/AudioInsights/Public/Messages/SoundTraceMessages.h b/Engine/Plugins/AudioInsights/Source/AudioInsights/Public/Messages/SoundTraceMessages.h
--- a/Engine/Plugins/AudioInsights/Source/AudioInsights/Public/Messages/SoundTraceMessages.h
+++ b/Engine/Plugins/AudioInsights/Source/AudioInsights/Public/Messages/SoundTraceMessages.h
@@ -322,7 +322,7 @@ namespace UE::Audio::Insights
 		float DistanceAttenuationDataPoint = 0.0f;
 		float LPFFreqDataPoint = 20000.0f;
 		float HPFFreqDataPoint = 0.0f;
-		float AmplitudeDataPoint = -160.0f;
+		float AmplitudeDataPoint = MIN_VOLUME_DECIBELS;
 		float VolumeDataPoint = 0.0f;
 		float PitchDataPoint = 1.0f;
 		float RelativeRenderCostDataPoint = 0.0f;
diff --git a/Engine/Plugins/AudioInsights/Source/AudioInsights/Public/Settings/SoundDashboardSettings.h b/Engine/Plugins/AudioInsights/Source/AudioInsights/Public/Settings/SoundDashboardSettings.h
--- a/Engine/Plugins/AudioInsights/Source/AudioInsights/Public/Settings/SoundDashboardSettings.h
+++ b/Engine/Plugins/AudioInsights/Source/AudioInsights/Public/Settings/SoundDashboardSettings.h
@@ -2,6 +2,7 @@
 
 #pragma once
 
+#include "AudioDefines.h"
 #include "Delegates/DelegateCombinations.h"
 #include "Engine/DeveloperSettings.h"
 #include "Settings/VisibleColumnsSettingsMenu.h"
@@ -127,7 +128,7 @@ struct FSoundPlotsDashboardPlotRanges
 	bool bUseCustomAmplitude_dBRange = false;
 
 	UPROPERTY(EditAnywhere, config, Category = "SoundDashboard|CustomPlotRanges", meta = (EditCondition = "bUseCustomAmplitude_dBRange", UIMin = "-160", UIMax = "0", ClampMin = "-160", ClampMax = "24", DisplayName = "Amplitude (dB)"))
-	FFloatInterval Amplitude_dB = { -160.0f, 0.0f };
+	FFloatInterval Amplitude_dB = { MIN_VOLUME_DECIBELS, 0.0f };
 
 
 	UPROPERTY(EditAnywhere, config, Category = "SoundDashboard|CustomPlotRanges", meta = (InlineEditConditionToggle))
diff --git a/Engine/Plugins/Runtime/AudioWidgets/Source/AudioWidgets/Private/AudioMeter.cpp b/Engine/Plugins/Runtime/AudioWidgets/Source/AudioWidgets/Private/AudioMeter.cpp
--- a/Engine/Plugins/Runtime/AudioWidgets/Source/AudioWidgets/Private/AudioMeter.cpp
+++ b/Engine/Plugins/Runtime/AudioWidgets/Source/AudioWidgets/Private/AudioMeter.cpp
@@ -257,8 +257,8 @@ namespace AudioWidgets
 			Analyzer->StartAnalyzing(InAudioDeviceId, AudioBus.Get());
 		}
 
-		constexpr float DefaultMeterValue = -160.0f;
-		constexpr float DefaultPeakValue = -160.0f;
+		constexpr float DefaultMeterValue = MIN_VOLUME_DECIBELS;
+		constexpr float DefaultPeakValue  = MIN_VOLUME_DECIBELS;
 		ChannelInfo.Init(FMeterChannelInfo{ DefaultMeterValue, DefaultPeakValue }, InNumChannels);
 
 		if (Widget.IsValid())
diff --git a/Engine/Plugins/Runtime/AudioWidgets/Source/AudioWidgets/Private/SAudioInputWidget.cpp b/Engine/Plugins/Runtime/AudioWidgets/Source/AudioWidgets/Private/SAudioInputWidget.cpp
--- a/Engine/Plugins/Runtime/AudioWidgets/Source/AudioWidgets/Private/SAudioInputWidget.cpp
+++ b/Engine/Plugins/Runtime/AudioWidgets/Source/AudioWidgets/Private/SAudioInputWidget.cpp
@@ -80,7 +80,7 @@ const float FVolumeProcessor::GetSliderValueForText(const FVector2D OutputRange,
 	return GetSliderValueFromDb(OutputRange, OutputValue);;
 }
 
-const float FVolumeProcessor::MinDbValue = -160.0f;
+const float FVolumeProcessor::MinDbValue = MIN_VOLUME_DECIBELS;
 const float FVolumeProcessor::MaxDbValue = 770.0f;
 const FVector2D FVolumeProcessor::GetOutputRange(const FVector2D InRange)
 {
diff --git a/Engine/Plugins/Runtime/AudioWidgets/Source/AudioWidgets/Private/SAudioOscilloscopePanelWidget.cpp b/Engine/Plugins/Runtime/AudioWidgets/Source/AudioWidgets/Private/SAudioOscilloscopePanelWidget.cpp
--- a/Engine/Plugins/Runtime/AudioWidgets/Source/AudioWidgets/Private/SAudioOscilloscopePanelWidget.cpp
+++ b/Engine/Plugins/Runtime/AudioWidgets/Source/AudioWidgets/Private/SAudioOscilloscopePanelWidget.cpp
@@ -2,6 +2,7 @@
 
 #include "SAudioOscilloscopePanelWidget.h"
 
+#include "AudioDefines.h"
 #include "AudioWidgetsLog.h"
 #include "DSP/Dsp.h"
 #include "FixedSampledSequenceGridData.h"
@@ -32,14 +33,12 @@ namespace AudioOscilloscopePanelWidgetPrivate
 
 	TFunction<FText(const double)> YAxisValueGeneratorDb = [](const double InLabelValue)
 	{
-		constexpr float MinDecibelNumericLimit = -160.0f;
-
 		const float AbsAmplitude  = FMath::Abs(InLabelValue * 2.0f - 1.0f);
 		const float GridLineValue = Audio::ConvertToDecibels(AbsAmplitude);
 
 		const FNumberFormattingOptions Formatting = FNumberFormattingOptions().SetMaximumFractionalDigits(1);
 
-		const FText GridLineValueText = (GridLineValue <= MinDecibelNumericLimit) ? FText::FromString("-inf") : FText::AsNumber(GridLineValue, &Formatting);
+		const FText GridLineValueText = (GridLineValue <= MIN_VOLUME_DECIBELS) ? FText::FromString("-inf") : FText::AsNumber(GridLineValue, &Formatting);
 
 		return GridLineValueText;
 	};
diff --git a/Engine/Plugins/Runtime/AudioWidgets/Source/AudioWidgets/Private/SAudioRadialSlider.cpp b/Engine/Plugins/Runtime/AudioWidgets/Source/AudioWidgets/Private/SAudioRadialSlider.cpp
--- a/Engine/Plugins/Runtime/AudioWidgets/Source/AudioWidgets/Private/SAudioRadialSlider.cpp
+++ b/Engine/Plugins/Runtime/AudioWidgets/Source/AudioWidgets/Private/SAudioRadialSlider.cpp
@@ -276,7 +276,7 @@ void SAudioRadialSlider::SetSliderThickness(const float Thickness)
 }
 
 // SAudioVolumeRadialSlider
-const float SAudioVolumeRadialSlider::MinDbValue = -160.0f;
+const float SAudioVolumeRadialSlider::MinDbValue = MIN_VOLUME_DECIBELS;
 const float SAudioVolumeRadialSlider::MaxDbValue = 770.0f;
 
 SAudioVolumeRadialSlider::SAudioVolumeRadialSlider()
diff --git a/Engine/Plugins/Runtime/AudioWidgets/Source/AudioWidgets/Private/SAudioSlider.cpp b/Engine/Plugins/Runtime/AudioWidgets/Source/AudioWidgets/Private/SAudioSlider.cpp
--- a/Engine/Plugins/Runtime/AudioWidgets/Source/AudioWidgets/Private/SAudioSlider.cpp
+++ b/Engine/Plugins/Runtime/AudioWidgets/Source/AudioWidgets/Private/SAudioSlider.cpp
@@ -388,7 +388,7 @@ const float SAudioSlider::GetSliderValue(const float OutputValue)
 }
 
 // SAudioVolumeSlider
-const float SAudioVolumeSlider::MinDbValue = -160.0f;
+const float SAudioVolumeSlider::MinDbValue = MIN_VOLUME_DECIBELS;
 const float SAudioVolumeSlider::MaxDbValue = 770.0f;
 SAudioVolumeSlider::SAudioVolumeSlider()
 {
