[Audio Insights] Loctext conflict fix

diff --git a/Engine/Plugins/AudioInsights/Source/AudioInsights/Private/AudioEventLogEditorCommands.cpp b/Engine/Plugins/AudioInsights/Source/AudioInsights/Private/AudioEventLogEditorCommands.cpp
--- a/Engine/Plugins/AudioInsights/Source/AudioInsights/Private/AudioEventLogEditorCommands.cpp
+++ b/Engine/Plugins/AudioInsights/Source/AudioInsights/Private/AudioEventLogEditorCommands.cpp
@@ -18,7 +18,7 @@ namespace UE::Audio::Insights
 	{
 		UI_COMMAND(Browse, "Browse To Asset", "Browses to the selected sound asset in the content browser.", EUserInterfaceActionType::Button, FInputChord(EKeys::B, EModifierKey::Control));
 		UI_COMMAND(Edit, "Edit", "Opens the selected sound for edit.", EUserInterfaceActionType::Button, FInputChord(EKeys::E, EModifierKey::Control));
-		UI_COMMAND(ResetInspectTimestamp, "Reset Inspect Timestamp", "Resumes gathering the latest audio data and auto-scrolling in the Event Log if previously paused.", EUserInterfaceActionType::Button, FInputChord(EKeys::Escape));
+		UI_COMMAND(ResetInspectTimestamp, "Reset Inspect Timestamp Event Log", "Resumes gathering the latest audio data and auto-scrolling in the Event Log if previously paused.", EUserInterfaceActionType::Button, FInputChord(EKeys::Escape));
 
 		UI_COMMAND(AutoStopCachingWhenLastInCache, "When marked for deletion", "Automatically stops caching new data when inspecting an event that has been marked for deletion from the cache.", EUserInterfaceActionType::RadioButton, FInputChord(EKeys::F1, EModifierKey::Control));
 		UI_COMMAND(AutoStopCachingOnInspect, "On inspect", "Automatically stops caching new data when inspecting an event.", EUserInterfaceActionType::RadioButton, FInputChord(EKeys::F2, EModifierKey::Control));
diff --git a/Engine/Plugins/AudioInsights/Source/AudioInsights/Private/SoundPlotsWidgetCommands.cpp b/Engine/Plugins/AudioInsights/Source/AudioInsights/Private/SoundPlotsWidgetCommands.cpp
--- a/Engine/Plugins/AudioInsights/Source/AudioInsights/Private/SoundPlotsWidgetCommands.cpp
+++ b/Engine/Plugins/AudioInsights/Source/AudioInsights/Private/SoundPlotsWidgetCommands.cpp
@@ -16,7 +16,7 @@ namespace UE::Audio::Insights
 
 	void FSoundPlotsWidgetCommands::RegisterCommands()
 	{
-		UI_COMMAND(ResetInspectTimestamp, "Reset Inspect Timestamp", "Resumes gathering the latest audio data and auto-scrolling in the Plots tab if previously paused.", EUserInterfaceActionType::Button, FInputChord(EKeys::Escape));
+		UI_COMMAND(ResetInspectTimestamp, "Reset Inspect Timestamp Plots", "Resumes gathering the latest audio data and auto-scrolling in the Plots tab if previously paused.", EUserInterfaceActionType::Button, FInputChord(EKeys::Escape));
 	}
 
 } // namespace UE::Audio::Insights
