[Audio Insights] Toolbar revamp + option to select between trace monitoring or recording (see JIRA for more details)

diff --git a/Engine/Plugins/AudioInsights/Source/AudioInsights/Private/AudioInsightsDashboardFactory.cpp b/Engine/Plugins/AudioInsights/Source/AudioInsights/Private/AudioInsightsDashboardFactory.cpp
--- a/Engine/Plugins/AudioInsights/Source/AudioInsights/Private/AudioInsightsDashboardFactory.cpp
+++ b/Engine/Plugins/AudioInsights/Source/AudioInsights/Private/AudioInsightsDashboardFactory.cpp
@@ -546,7 +546,7 @@ namespace UE::Audio::Insights
 	FReply FDashboardFactory::ToggleAutoEnableAudioTraces()
 	{
 		IAudioInsightsTraceModule& TraceModule = FAudioInsightsModule::GetChecked().GetTraceModule();
-		TraceModule.StartTraceAnalysis(false);
+		TraceModule.StartTraceAnalysis(false, ETraceMode::Recording);
 
 		return FReply::Handled();
 	}
diff --git a/Engine/Plugins/AudioInsights/Source/AudioInsights/Private/AudioInsightsTraceModule.cpp b/Engine/Plugins/AudioInsights/Source/AudioInsights/Private/AudioInsightsTraceModule.cpp
--- a/Engine/Plugins/AudioInsights/Source/AudioInsights/Private/AudioInsightsTraceModule.cpp
+++ b/Engine/Plugins/AudioInsights/Source/AudioInsights/Private/AudioInsightsTraceModule.cpp
@@ -303,7 +303,7 @@ namespace UE::Audio::Insights
 		OnAnalysisStarting.Broadcast(FPlatformTime::Seconds());
 	}
 
-	void FTraceModule::StartTraceAnalysis(const bool bOnlyTraceAudioChannels)
+	void FTraceModule::StartTraceAnalysis(const bool bInOnlyTraceAudioChannels, const ETraceMode InTraceMode)
 	{
 		if (bTraceAnalysisHasStarted)
 		{
@@ -313,11 +313,12 @@ namespace UE::Audio::Insights
 #if WITH_EDITOR
 		CacheCurrentlyEnabledTraceChannels();
 
-		if (bOnlyTraceAudioChannels)
+		if (bInOnlyTraceAudioChannels)
 		{
 			DisableAllTraceChannels();
 		}
 
+		CurrentTraceMode = InTraceMode;
 		bAudioTraceChannelsEnabled = EnableAudioInsightsTraceChannels();
 
 		IUnrealInsightsModule& UnrealInsightsModule = FModuleManager::LoadModuleChecked<IUnrealInsightsModule>("TraceInsights");
@@ -334,18 +335,26 @@ namespace UE::Audio::Insights
 		}
 		else
 		{
-			// Start analysis with a direct trace
-			const FString DirectTraceHost = "127.0.0.1";
-			const uint32 DirectTracePort  = UnrealInsightsModule.StartAnalysisWithDirectTrace();
-			const FString DirectTraceAddress = FString::Printf(TEXT("%s:%u"), *DirectTraceHost, DirectTracePort);
+			bIsAudioInsightsTrace = true;
 
 			// Clear all buffered data and prevent data from previous recordings from leaking into the new recording
 			FTraceAuxiliary::FOptions Options;
 			Options.bExcludeTail = true;
 
-			bIsAudioInsightsTrace = true;
+			if (CurrentTraceMode == ETraceMode::Recording)
+			{
+				FTraceAuxiliary::Start(FTraceAuxiliary::EConnectionType::Network, TEXT("localhost"), TEXT(""), &Options);
 
-			FTraceAuxiliary::Start(FTraceAuxiliary::EConnectionType::Network, *DirectTraceAddress, TEXT(""), &Options);
+				UnrealInsightsModule.StartAnalysisForLastLiveSession();
+			}
+			else if (CurrentTraceMode == ETraceMode::Monitoring)
+			{
+				const uint32 DirectTracePort = UnrealInsightsModule.StartAnalysisWithDirectTrace();
+
+				const FString DirectTraceHost = FString::Printf(TEXT("127.0.0.1:%u"), DirectTracePort);
+
+				FTraceAuxiliary::Start(FTraceAuxiliary::EConnectionType::Network, *DirectTraceHost, TEXT(""), &Options);
+			}
 		}
 #else
 		bAudioTraceChannelsEnabled = EnableAudioInsightsTraceChannels();
@@ -380,6 +389,7 @@ namespace UE::Audio::Insights
 		}
 
 		bTraceAnalysisHasStarted = false;
+		CurrentTraceMode = ETraceMode::None;
 	}
 #endif // WITH_EDITOR
 
@@ -598,6 +608,11 @@ namespace UE::Audio::Insights
 #endif // WITH_EDITOR
 	}
 
+	ETraceMode FTraceModule::GetTraceMode() const
+	{
+		return CurrentTraceMode;
+	}
+
 	void FTraceModule::GetLoggers(TArray<const TCHAR *>& OutLoggers)
 	{
 		OutLoggers.Add(TEXT("Audio"));
diff --git a/Engine/Plugins/AudioInsights/Source/AudioInsights/Public/AudioInsightsStyle.h b/Engine/Plugins/AudioInsights/Source/AudioInsights/Public/AudioInsightsStyle.h
--- a/Engine/Plugins/AudioInsights/Source/AudioInsights/Public/AudioInsightsStyle.h
+++ b/Engine/Plugins/AudioInsights/Source/AudioInsights/Public/AudioInsightsStyle.h
@@ -182,6 +182,7 @@ namespace UE::Audio::Insights
 			Set("AudioInsights.Icon.Dashboard", new IMAGE_BRUSH_SVG(TEXT("Icons/audio_dashboard"), Icon16));
 			Set("AudioInsights.Icon.Event", new IMAGE_BRUSH_SVG(TEXT("Icons/audio_event"), Icon16));
 			Set("AudioInsights.Icon.Log", new IMAGE_BRUSH_SVG(TEXT("Icons/audio_log"), Icon16));
+			Set("AudioInsights.Icon.OnlyTraceAudioChannels", new IMAGE_BRUSH_SVG(TEXT("Icons/only_trace_audio_channels"), Icon16));
 			Set("AudioInsights.Icon.Open", new IMAGE_BRUSH_SVG(TEXT("Icons/open"), Icon20));
 			Set("AudioInsights.Icon.Pause", new IMAGE_BRUSH_SVG(TEXT("Icons/pause"), Icon20));
 			Set("AudioInsights.Icon.Record.Active", new IMAGE_BRUSH_SVG(TEXT("Icons/record_active"), Icon20));
diff --git a/Engine/Plugins/AudioInsights/Source/AudioInsights/Public/AudioInsightsTraceModule.h b/Engine/Plugins/AudioInsights/Source/AudioInsights/Public/AudioInsightsTraceModule.h
--- a/Engine/Plugins/AudioInsights/Source/AudioInsights/Public/AudioInsightsTraceModule.h
+++ b/Engine/Plugins/AudioInsights/Source/AudioInsights/Public/AudioInsightsTraceModule.h
@@ -16,6 +16,13 @@
 
 namespace UE::Audio::Insights
 {
+	enum class ETraceMode : uint8
+	{
+		Monitoring = 0,
+		Recording,
+		None
+	};
+
 	class FRewindDebuggerAudioInsightsRuntime : public IRewindDebuggerRuntimeExtension
 	{
 	public:
@@ -44,7 +51,7 @@ namespace UE::Audio::Insights
 
 		UE_API virtual void AddTraceProvider(TSharedPtr<FTraceProviderBase> TraceProvider) override;
 
-		UE_API virtual void StartTraceAnalysis(const bool bOnlyTraceAudioChannels) override;
+		UE_API virtual void StartTraceAnalysis(const bool bInOnlyTraceAudioChannels, const ETraceMode InTraceMode) override;
 		UE_API virtual bool IsTraceAnalysisActive() const override;
 		UE_API virtual void StopTraceAnalysis() override;
 		UE_API virtual void OnOnlyTraceAudioChannelsStateChanged(const bool bOnlyTraceAudioChannels) override;
@@ -55,6 +62,8 @@ namespace UE::Audio::Insights
 
 		UE_API virtual void SaveTraceSnapshot();
 
+		UE_API virtual ETraceMode GetTraceMode() const override;
+
 #if !WITH_EDITOR
 		UE_API virtual void InitializeSessionInfo(const TraceServices::FSessionInfo& SessionInfo) override;
 		UE_API virtual void RequestChannelUpdate() override;
@@ -93,6 +102,8 @@ namespace UE::Audio::Insights
 		TArray<FString> ChannelsToRestore;
 		FRewindDebuggerAudioInsightsRuntime RewindDebugger;
 
+		ETraceMode CurrentTraceMode = ETraceMode::None;
+
 		bool bTraceAnalysisHasStarted = false;
 		bool bAudioTraceChannelsEnabled = false;
 		bool bIsAudioInsightsTrace = false;
diff --git a/Engine/Plugins/AudioInsights/Source/AudioInsights/Public/IAudioInsightsTraceModule.h b/Engine/Plugins/AudioInsights/Source/AudioInsights/Public/IAudioInsightsTraceModule.h
--- a/Engine/Plugins/AudioInsights/Source/AudioInsights/Public/IAudioInsightsTraceModule.h
+++ b/Engine/Plugins/AudioInsights/Source/AudioInsights/Public/IAudioInsightsTraceModule.h
@@ -7,6 +7,7 @@
 namespace UE::Audio::Insights
 {
 	class FTraceProviderBase;
+	enum class ETraceMode : uint8;
 } // namespace UE::Audio::Insights
 
 namespace TraceServices
@@ -19,7 +20,7 @@ class IAudioInsightsTraceModule : public TraceServices::IModule
 public:
 	virtual void AddTraceProvider(TSharedPtr<UE::Audio::Insights::FTraceProviderBase> TraceProvider) = 0;
 
-	virtual void StartTraceAnalysis(const bool bOnlyTraceAudioChannels) = 0;
+	virtual void StartTraceAnalysis(const bool bInOnlyTraceAudioChannels, const UE::Audio::Insights::ETraceMode InTraceMode) = 0;
 	virtual bool IsTraceAnalysisActive() const = 0;
 	virtual void StopTraceAnalysis() = 0;
 	
@@ -31,6 +32,8 @@ public:
 
 	virtual void SaveTraceSnapshot() = 0;
 
+	virtual UE::Audio::Insights::ETraceMode GetTraceMode() const = 0;
+
 #if !WITH_EDITOR
 	virtual void InitializeSessionInfo(const TraceServices::FSessionInfo& SessionInfo) = 0;
 	virtual void RequestChannelUpdate() = 0;
diff --git a/Engine/Plugins/AudioInsights/Source/AudioInsightsEditor/AudioInsightsEditor.Build.cs b/Engine/Plugins/AudioInsights/Source/AudioInsightsEditor/AudioInsightsEditor.Build.cs
--- a/Engine/Plugins/AudioInsights/Source/AudioInsightsEditor/AudioInsightsEditor.Build.cs
+++ b/Engine/Plugins/AudioInsights/Source/AudioInsightsEditor/AudioInsightsEditor.Build.cs
@@ -32,7 +32,6 @@ public class AudioInsightsEditor : ModuleRules
 				"SignalProcessing",
 				"Slate",
 				"SlateCore",
-				"ToolMenus",
 				"TraceAnalysis",
 				"UnrealEd",
 				"WorkspaceMenuStructure",
diff --git a/Engine/Plugins/AudioInsights/Source/AudioInsightsEditor/Private/AudioInsightsEditorDashboardFactory.cpp b/Engine/Plugins/AudioInsights/Source/AudioInsightsEditor/Private/AudioInsightsEditorDashboardFactory.cpp
--- a/Engine/Plugins/AudioInsights/Source/AudioInsightsEditor/Private/AudioInsightsEditorDashboardFactory.cpp
+++ b/Engine/Plugins/AudioInsights/Source/AudioInsightsEditor/Private/AudioInsightsEditorDashboardFactory.cpp
@@ -8,7 +8,7 @@
 #include "AudioInsightsEditorModule.h"
 #include "AudioInsightsEditorSettings.h"
 #include "AudioInsightsStyle.h"
-#include "Editor.h"
+#include "AudioInsightsTraceModule.h"
 #include "Engine/World.h"
 #include "Framework/Commands/UICommandList.h"
 #include "Framework/Docking/LayoutService.h"
@@ -16,10 +16,7 @@
 #include "Framework/Notifications/NotificationManager.h"
 #include "Internationalization/Text.h"
 #include "IPropertyTypeCustomization.h"
-#include "Kismet2/DebuggerCommands.h"
 #include "Templates/SharedPointer.h"
-#include "ToolMenu.h"
-#include "ToolMenus.h"
 #include "Widgets/Input/SButton.h"
 #include "Widgets/Input/SCheckBox.h"
 #include "Widgets/Layout/SSeparator.h"
@@ -43,25 +40,28 @@ namespace UE::Audio::Insights
 
 		static const FText ViewResetLayoutText = LOCTEXT("AudioDashboardEditor_View_ResetLayout_Text", "Reset Layout");
 
-		static const FText OnlyTraceAudioChannelsText = LOCTEXT("AudioDashboardEditor_OnlyTraceAudioChannels_Text", "Only trace audio channels:");
 		static const FText OnlyTraceAudioChannelsTooltipText = LOCTEXT("AudioDashboardEditor_OnlyTraceAudioChannels_TooltipText", "Disable all non-command line trace channels apart from Audio, Audio Mixer and Bookmark when starting a live monitoring / live trace. This will reduce the file size of snapshots / trace sessions while using Audio Insights.");
 
-		static const FText StartDirectTraceText = LOCTEXT("AudioDashboardEditor_StartDirectTrace_Text", "Start Monitoring");
-		static const FText StartDirectTraceTooltipText = LOCTEXT("AudioDashboardEditor_StartDirectTrace_TooltipText", "Starts monitoring for audio trace data without saving the data to a .utrace file.");
+		static const FText StartMonitoringText = LOCTEXT("AudioDashboardEditor_StartMonitoring_Text", "Start Monitoring");
+		static const FText StartMonitoringTooltipText = LOCTEXT("AudioDashboardEditor_StartMonitoring_TooltipText", "Starts monitoring for audio trace data without saving the data to a .utrace file.");
 
-		static const FText StopDirectTraceText = LOCTEXT("AudioDashboardEditor_StopDirectTrace_Text", "Stop Monitoring");
-		static const FText StopDirectTraceTooltipText = LOCTEXT("AudioDashboardEditor_StopDirectTrace_TooltipText", "Stops the active monitoring session.");
+		static const FText StopMonitoringText = LOCTEXT("AudioDashboardEditor_StopMonitoring_Text", "Monitoring...");
+		static const FText StopMonitoringTooltipText = LOCTEXT("AudioDashboardEditor_StopMonitoring_TooltipText", "Stops the active monitoring session.");
 
-		static const FText StopLiveTraceText = LOCTEXT("AudioDashboardEditor_StopLiveTrace_Text", "Stop Trace");
-		static const FText StopLiveTraceTooltipText = LOCTEXT("AudioDashboardEditor_StopLiveTrace_TooltipText", "Stops the currently active trace recording.");
+		static const FText StartRecordingText = LOCTEXT("AudioDashboardEditor_StartRecording_Text", "Start Recording");
+		static const FText StartRecordingTooltipText = LOCTEXT("AudioDashboardEditor_StartRecording_TooltipText", "Starts to record audio trace data saving it to a .utrace file.");
+
+		static const FText StopRecordingText = LOCTEXT("AudioDashboardEditor_StopRecording_Text", "Recording...");
+		static const FText StopRecordingTooltipText = LOCTEXT("AudioDashboardEditor_StopRecording_TooltipText", "Stops the currently active trace recording.");
+
+		static const FText RecordingNotificationText = LOCTEXT("AudioDashboardEditor_RecordingButton_NotificationText", "Trace Started");
+		static const FText RecordingNotificationSubText = LOCTEXT("AudioDashboardEditor_RecordingButton_NotificationSubText", "Trace is now active and saving to the following location (file or tracestore):\n");
 
-		static const FText SaveSnapshotText = LOCTEXT("AudioDashboardEditor_SnapshotButton_Text", "Save Snapshot");
 		static const FText SaveSnapshotTooltipText = LOCTEXT("AudioDashboardEditor_SnapshotButton_TooltipText", "Saves a snapshot of the most recent trace data to a .utrace file. By default snapshots are limited to 32MB and can be accessed in Unreal Insights. To specify a custom snapshot size (in MB) launch the editor executable with the -tracetailmb=<size> option.");
 
 		static const FText SaveSnapshotNotificationText = LOCTEXT("AudioDashboardEditor_SnapshotButton_NotificationText", "Audio Insights Snapshot saved.");
 		static const FText SaveSnapshotNotificationSubText = LOCTEXT("AudioDashboardEditor_SnapshotButton_NotificationSubText", "A snapshot .utrace with the most recent events has been saved to your trace server.");
 
-		static const FText SaveBookmarkText = LOCTEXT("AudioDashboardEditor_BookmarkButton_Text", "Bookmark");
 		static const FText SaveBookmarkTooltipText = LOCTEXT("AudioDashboardEditor_BookmarkButton_TooltipText", "Ctrl+M - Creates an 'Audio Insights' bookmark in the .utrace file generated by Save Snapshot or Trace at the timestamp the button was pressed.");
 
 		static const FText BookmarkLabelText = LOCTEXT("AudioDashboardEditor_BookmarkLabelTitle_Text", "Audio Insights Bookmark");
@@ -71,6 +71,28 @@ namespace UE::Audio::Insights
 
 		static const FText WorldFilterText = LOCTEXT("AudioDashboardEditor_WorldFilter_Text", "World Filter:");
 		static const FText WorldFilterTooltipText = LOCTEXT("AudioDashboardEditor_WorldFilter_TooltipText", "Select world(s) to monitor (worlds may share audio output).");
+		
+		void ShowNotification(const FText& InText, const FText& InSubText)
+		{
+			FNotificationInfo Info(InText);
+			Info.SubText = InSubText;
+			Info.bFireAndForget = true;
+			Info.FadeOutDuration = 1.0f;
+			Info.ExpireDuration = 4.0f;
+
+			FSlateNotificationManager::Get().AddNotification(Info);
+		}
+
+		void ShowRecordingNotification()
+		{
+			FString TraceDestinationStr = FTraceAuxiliary::GetTraceDestinationString();
+			if (TraceDestinationStr.IsEmpty())
+			{
+				TraceDestinationStr = TEXT("External Target");
+			}
+			
+			ShowNotification(RecordingNotificationText, FText::Format(FText::FromString(TEXT("{0}{1}")), RecordingNotificationSubText, FText::FromString(TraceDestinationStr)));
+		}
 
 		FText GetDebugNameFromDeviceId(::Audio::FDeviceId InDeviceId)
 		{
@@ -153,7 +175,7 @@ namespace UE::Audio::Insights
 		// If we haven't started a trace, start the analysis with the currently running trace
 		if (!TraceModule.IsAudioInsightsTrace())
 		{
-			TraceModule.StartTraceAnalysis(bOnlyTraceAudioChannels);
+			TraceModule.StartTraceAnalysis(bOnlyTraceAudioChannels, ETraceMode::Recording);
 			BookmarkIndex = 0;
 		}
 	}
@@ -326,7 +348,7 @@ namespace UE::Audio::Insights
 		}
 
 		IAudioInsightsTraceModule& TraceModule = FAudioInsightsEditorModule::GetChecked().GetTraceModule();
-		TraceModule.StartTraceAnalysis(bOnlyTraceAudioChannels);
+		TraceModule.StartTraceAnalysis(bOnlyTraceAudioChannels, ETraceMode::Monitoring);
 
 		const TSharedRef<FTabManager::FLayout> TabLayout = LoadLayoutFromConfig();
 
@@ -518,70 +540,53 @@ namespace UE::Audio::Insights
 	{
 		using namespace EditorDashboardFactoryPrivate;
 
-		static const FName PlayWorldToolBarName = "Kismet.DebuggingViewToolBar";
-		if (!UToolMenus::Get()->IsMenuRegistered(PlayWorldToolBarName))
-		{
-			UToolMenu* ToolBar = UToolMenus::Get()->RegisterMenu(PlayWorldToolBarName, NAME_None, EMultiBoxType::SlimHorizontalToolBar);
-			FToolMenuSection& Section = ToolBar->AddSection("Debug");
-			FPlayWorldCommands::BuildToolbar(Section);
-		}
-
-		static FSlateBrush TransportBackgroundColorBrush;
-		TransportBackgroundColorBrush.TintColor = FSlateColor(FLinearColor(0.018f, 0.018f, 0.018f, 1.0f));
-		TransportBackgroundColorBrush.DrawAs    = ESlateBrushDrawType::Box;
+		static FSlateBrush BackgroundColorBrush;
+		BackgroundColorBrush.TintColor = FSlateColor(FLinearColor(0.018f, 0.018f, 0.018f, 1.0f));
+		BackgroundColorBrush.DrawAs = ESlateBrushDrawType::Box;
 
 		return SNew(SBorder)
-			.BorderImage(&TransportBackgroundColorBrush)
+			.BorderImage(&BackgroundColorBrush)
 			[
 				SNew(SHorizontalBox)
-				// Play World Toolbar
+				// Monitoring button
 				+ SHorizontalBox::Slot()
 				.HAlign(HAlign_Left)
 				.VAlign(VAlign_Center)
 				.AutoWidth()
+				.Padding(3.0f, 0.0f, 5.0f, 0.0f)
 				[
-					SNew(SBorder)
-					.BorderImage(FAppStyle::Get().GetBrush("NoBorder"))
-					[
-						UToolMenus::Get()->GenerateWidget(PlayWorldToolBarName, { FPlayWorldCommands::GlobalPlayWorldActions })
-					]
+					MakeTraceModeButton(ETraceMode::Monitoring)
 				]
-				// Only trace audio channels label
+				// Separator
 				+ SHorizontalBox::Slot()
-				.HAlign(HAlign_Right)
+				.HAlign(HAlign_Left)
 				.VAlign(VAlign_Center)
 				.AutoWidth()
-				.Padding(2.0f, 0.0f)
+				.Padding(0.0f, 0.0f, 0.0f, 0.0f)
 				[
-					SNew(STextBlock)
-					.Text(OnlyTraceAudioChannelsText)
-					.ToolTipText(OnlyTraceAudioChannelsTooltipText)
-					.Font(IPropertyTypeCustomizationUtils::GetRegularFont())
+					SNew(SBox)
+					.HeightOverride(38.0f)
+					[
+						SNew(SSeparator)
+						.Orientation(Orient_Vertical)
+						.Thickness(2.0f)
+					]
 				]
-				// Only trace audio channels checkbox
+				// Recording button
 				+ SHorizontalBox::Slot()
 				.HAlign(HAlign_Left)
 				.VAlign(VAlign_Center)
 				.AutoWidth()
-				.Padding(2.0f, 0.0f)
+				.Padding(5.0f, 0.0f, 5.0f, 0.0f)
 				[
-					SNew(SCheckBox)
-					.ToolTipText(OnlyTraceAudioChannelsTooltipText)
-					.IsChecked_Lambda([this]() { return bOnlyTraceAudioChannels ? ECheckBoxState::Checked : ECheckBoxState::Unchecked; })
-					.OnCheckStateChanged_Lambda([this](ECheckBoxState NewState) 
-					{ 
-						bOnlyTraceAudioChannels = NewState == ECheckBoxState::Checked;
-
-						IAudioInsightsTraceModule& TraceModule = FAudioInsightsEditorModule::GetChecked().GetTraceModule();
-						TraceModule.OnOnlyTraceAudioChannelsStateChanged(bOnlyTraceAudioChannels);
-					})
+					MakeTraceModeButton(ETraceMode::Recording)
 				]
 				// Separator
 				+ SHorizontalBox::Slot()
 				.HAlign(HAlign_Left)
 				.VAlign(VAlign_Center)
 				.AutoWidth()
-				.Padding(12.0f, 0.0f, 0.0f, 0.0f)
+				.Padding(0.0f, 0.0f, 0.0f, 0.0f)
 				[
 					SNew(SBox)
 					.HeightOverride(38.0f)
@@ -591,7 +596,7 @@ namespace UE::Audio::Insights
 						.Thickness(2.0f)
 					]
 				]
-				// Start direct trace button
+				// Save snapshot button
 				+ SHorizontalBox::Slot()
 				.HAlign(HAlign_Left)
 				.VAlign(VAlign_Center)
@@ -600,98 +605,19 @@ namespace UE::Audio::Insights
 				[
 					SNew(SButton)
 					.ButtonStyle(&FAppStyle::Get().GetWidgetStyle<FButtonStyle>("SimpleButton"))
-					.ToolTipText_Lambda([&TraceModule = FAudioInsightsEditorModule::GetChecked().GetTraceModule()]()
-					{
-						if (TraceModule.IsTraceAnalysisActive())
-						{
-							return TraceModule.IsAudioInsightsTrace() ? StopDirectTraceTooltipText : StopLiveTraceTooltipText;
-						}
-
-						return StartDirectTraceTooltipText;
-					})
-					.OnClicked_Lambda([this]()
+					.ToolTipText(SaveSnapshotTooltipText)
+					.IsEnabled_Lambda([this]()
 					{
 						IAudioInsightsTraceModule& TraceModule = FAudioInsightsEditorModule::GetChecked().GetTraceModule();
-
-						if (TraceModule.IsTraceAnalysisActive())
-						{
-							TraceModule.StopTraceAnalysis();
-						}
-						else
-						{
-							TraceModule.StartTraceAnalysis(bOnlyTraceAudioChannels);
-						}
-
-						return FReply::Handled();
+						return TraceModule.IsTraceAnalysisActive();
 					})
-					[
-						SNew(SBox)
-						.HeightOverride(16.0f)
-						[
-							SNew(SHorizontalBox)
-							// Trace record icon
-							+ SHorizontalBox::Slot()
-							.VAlign(VAlign_Center)
-							.AutoWidth()
-							[
-								SNew(SImage)
-								.Image(FSlateStyle::Get().GetBrush("AudioInsights.Icon.Trace"))
-								.ColorAndOpacity_Lambda([&TraceModule = FAudioInsightsEditorModule::GetChecked().GetTraceModule()]()
-				 				{
-									if (TraceModule.IsTraceAnalysisActive())
-									{
-										return TraceModule.IsAudioInsightsTrace() ? FSlateColor(FColor::Green) : FSlateColor(FColor::Red);
-									}
-
-									return FSlateColor::UseForeground();
-				 				})
-							]
-							// Trace record text
-							+ SHorizontalBox::Slot()
-							.VAlign(VAlign_Center)
-							.AutoWidth()
-							.Padding(6.0f, 1.0f, 0.0f, 0.0f)
-							[
-								SNew(STextBlock)
-								.Justification(ETextJustify::Center)
-								.Text_Lambda([&TraceModule = FAudioInsightsEditorModule::GetChecked().GetTraceModule()]()
-				 				{
-									if (TraceModule.IsTraceAnalysisActive())
-									{
-										return TraceModule.IsAudioInsightsTrace() ? StopDirectTraceText : StopLiveTraceText;
-									}
-
-									return StartDirectTraceText;
-				 				})
-								.Font(IPropertyTypeCustomizationUtils::GetRegularFont())
-							]
-						]
-					]
-				]
-				// Save snapshot button
-				+ SHorizontalBox::Slot()
-				.HAlign(HAlign_Left)
-				.VAlign(VAlign_Center)
-				.AutoWidth()
-				.Padding(12.0f, 0.0f, 0.0f, 0.0f)
-				[
-					SNew(SButton)
-					.ButtonStyle(&FAppStyle::Get().GetWidgetStyle<FButtonStyle>("SimpleButton"))
-					.ToolTipText(SaveSnapshotTooltipText)
 					.OnClicked_Lambda([]()
 					{
 						// Save snapshot
 						IAudioInsightsTraceModule& TraceModule = FAudioInsightsEditorModule::GetChecked().GetTraceModule();
 						TraceModule.SaveTraceSnapshot();
 
-						// Show floating notification
-						FNotificationInfo Info(SaveSnapshotNotificationText);
-						Info.SubText = SaveSnapshotNotificationSubText;
-						Info.bFireAndForget  = true;
-						Info.FadeOutDuration = 1.0f;
-						Info.ExpireDuration  = 4.0f;
-
-						FSlateNotificationManager::Get().AddNotification(Info);
+						ShowNotification(SaveSnapshotNotificationText, SaveSnapshotNotificationSubText);
 
 						return FReply::Handled();
 					})
@@ -709,17 +635,6 @@ namespace UE::Audio::Insights
 								.Image(FSlateStyle::Get().GetBrush("AudioInsights.Icon.Trace.Snapshot"))
 								.ColorAndOpacity(FSlateColor::UseForeground())
 							]
-							// Trace snapshot text
-							+ SHorizontalBox::Slot()
-							.VAlign(VAlign_Center)
-							.AutoWidth()
-							.Padding(6.0f, 1.0f, 0.0f, 0.0f)
-							[
-								SNew(STextBlock)
-								.Justification(ETextJustify::Center)
-								.Text(SaveSnapshotText)
-								.Font(IPropertyTypeCustomizationUtils::GetRegularFont())
-							]
 						]
 					]
 				]
@@ -728,7 +643,7 @@ namespace UE::Audio::Insights
 				.HAlign(HAlign_Left)
 				.VAlign(VAlign_Center)
 				.AutoWidth()
-				.Padding(12.0f, 0.0f, 0.0f, 0.0f)
+				.Padding(5.0f, 0.0f, 10.0f, 0.0f)
 				[
 					SNew(SButton)
 					.ButtonStyle(&FAppStyle::Get().GetWidgetStyle<FButtonStyle>("SimpleButton"))
@@ -758,20 +673,67 @@ namespace UE::Audio::Insights
 								.Image(FSlateStyle::Get().GetBrush("AudioInsights.Icon.Trace.Bookmark"))
 								.ColorAndOpacity(FSlateColor::UseForeground())
 							]
-							// Save bookmark text
-							+ SHorizontalBox::Slot()
-							.VAlign(VAlign_Center)
-							.AutoWidth()
-							.Padding(6.0f, 1.0f, 0.0f, 0.0f)
-							[
-								SNew(STextBlock)
-								.Justification(ETextJustify::Center)
-								.Text(SaveBookmarkText)
-								.Font(IPropertyTypeCustomizationUtils::GetRegularFont())
-							]
 						]
 					]
 				]
+				// Separator
+				+ SHorizontalBox::Slot()
+				.HAlign(HAlign_Left)
+				.VAlign(VAlign_Center)
+				.AutoWidth()
+				.Padding(0.0f, 0.0f, 0.0f, 0.0f)
+				[
+					SNew(SBox)
+					.HeightOverride(38.0f)
+					[
+						SNew(SSeparator)
+						.Orientation(Orient_Vertical)
+						.Thickness(2.0f)
+					]
+				]
+				// Only trace audio channels icon
+				+ SHorizontalBox::Slot()
+				.AutoWidth()
+				.VAlign(VAlign_Center)
+				.Padding(20.0f, 0.0f, 0.0f, 0.0f)
+				[
+					SNew(SImage)
+					.Image(FSlateStyle::Get().GetBrush("AudioInsights.Icon.OnlyTraceAudioChannels"))
+					.ToolTipText(OnlyTraceAudioChannelsTooltipText)
+				]
+				// Only trace audio channels checkbox
+				+ SHorizontalBox::Slot()
+				.HAlign(HAlign_Left)
+				.VAlign(VAlign_Center)
+				.AutoWidth()
+				.Padding(8.0f, 0.0f, 18.0f, 0.0f)
+				[
+					SNew(SCheckBox)
+					.ToolTipText(OnlyTraceAudioChannelsTooltipText)
+					.IsChecked_Lambda([this]() { return bOnlyTraceAudioChannels ? ECheckBoxState::Checked : ECheckBoxState::Unchecked; })
+					.OnCheckStateChanged_Lambda([this](ECheckBoxState NewState) 
+					{ 
+						bOnlyTraceAudioChannels = NewState == ECheckBoxState::Checked;
+
+						IAudioInsightsTraceModule& TraceModule = FAudioInsightsEditorModule::GetChecked().GetTraceModule();
+						TraceModule.OnOnlyTraceAudioChannelsStateChanged(bOnlyTraceAudioChannels);
+					})
+				]
+				// Separator
+				+ SHorizontalBox::Slot()
+				.HAlign(HAlign_Left)
+				.VAlign(VAlign_Center)
+				.AutoWidth()
+				.Padding(0.0f, 0.0f, 0.0f, 0.0f)
+				[
+					SNew(SBox)
+					.HeightOverride(38.0f)
+					[
+						SNew(SSeparator)
+						.Orientation(Orient_Vertical)
+						.Thickness(2.0f)
+					]
+				]
 				// World filter label
 				+ SHorizontalBox::Slot()
 				.HAlign(HAlign_Right)
@@ -822,6 +784,138 @@ namespace UE::Audio::Insights
 			];
 	}
 
+	TSharedRef<SWidget> FEditorDashboardFactory::MakeTraceModeButton(const ETraceMode InTraceButtonMode)
+	{
+		using namespace EditorDashboardFactoryPrivate;
+
+		return SNew(SBox)
+			.WidthOverride(130.0f)
+			.HAlign(HAlign_Center)
+			.VAlign(VAlign_Center)
+			[
+				SNew(SButton)
+				.ButtonStyle(&FAppStyle::Get().GetWidgetStyle<FButtonStyle>("SimpleButton"))
+				.ToolTipText_Lambda([InTraceButtonMode]()
+				{
+					const IAudioInsightsTraceModule& TraceModule = FAudioInsightsEditorModule::GetChecked().GetTraceModule();
+
+					const ETraceMode CurrentTraceMode = TraceModule.GetTraceMode();
+					const bool bIsButtonActive = InTraceButtonMode == ETraceMode::Monitoring ? CurrentTraceMode != ETraceMode::None : CurrentTraceMode == InTraceButtonMode;
+					
+					if (TraceModule.IsTraceAnalysisActive() && bIsButtonActive)
+					{
+						return InTraceButtonMode == ETraceMode::Monitoring ? StopMonitoringTooltipText : StopRecordingTooltipText;
+					}
+
+					return InTraceButtonMode == ETraceMode::Monitoring ? StartMonitoringTooltipText : StartRecordingTooltipText;
+				})
+				.OnClicked_Lambda([this, InTraceButtonMode]()
+				{
+					IAudioInsightsTraceModule& TraceModule = FAudioInsightsEditorModule::GetChecked().GetTraceModule();
+
+					if (FTraceAuxiliary::IsConnected())
+					{
+						const ETraceMode TraceMode = TraceModule.GetTraceMode();
+
+						TraceModule.StopTraceAnalysis();
+
+						// Go to Recording mode when Monitoring is active and Record button is pressed
+						if (InTraceButtonMode == ETraceMode::Recording)
+						{
+							// Stopping a trace doesn't occur immediately, we need to check until it actually
+							// stops in order to call StartTraceAnalysis again (and disable the ticker)
+							auto StartTraceAnalysisDeferred = [](const bool bInOnlyTraceAudioChannels, const ETraceMode InTraceMode)
+							{
+								FTSTicker::GetCoreTicker().AddTicker(FTickerDelegate::CreateLambda(
+									[bInOnlyTraceAudioChannels, InTraceMode](float DeltaTime)
+									{
+										if (FTraceAuxiliary::IsConnected())
+										{
+											return true;
+										}
+
+										IAudioInsightsTraceModule& TraceModule = FAudioInsightsEditorModule::GetChecked().GetTraceModule();
+										TraceModule.StartTraceAnalysis(bInOnlyTraceAudioChannels, InTraceMode);
+										ShowRecordingNotification();
+
+										return false;
+									}));
+							};
+
+							if (TraceMode == ETraceMode::Monitoring)
+							{
+								StartTraceAnalysisDeferred(bOnlyTraceAudioChannels, ETraceMode::Recording);
+							}
+						}
+					}
+					else
+					{
+						TraceModule.StartTraceAnalysis(bOnlyTraceAudioChannels, InTraceButtonMode);
+
+						if (InTraceButtonMode == ETraceMode::Recording)
+						{
+							ShowRecordingNotification();
+						}
+					}
+
+					return FReply::Handled();
+				})
+				[
+					SNew(SBox)
+					.HeightOverride(16.0f)
+					[
+						SNew(SHorizontalBox)
+						// Trace button icon
+						+ SHorizontalBox::Slot()
+						.VAlign(VAlign_Center)
+						.AutoWidth()
+						[
+							SNew(SImage)
+							.Image(FSlateStyle::Get().GetBrush("AudioInsights.Icon.Trace"))
+							.ColorAndOpacity_Lambda([InTraceButtonMode]()
+				 			{
+								const IAudioInsightsTraceModule& TraceModule = FAudioInsightsEditorModule::GetChecked().GetTraceModule();
+
+								const ETraceMode CurrentTraceMode = TraceModule.GetTraceMode();
+								const bool bIsButtonActive = InTraceButtonMode == ETraceMode::Monitoring ? CurrentTraceMode != ETraceMode::None : CurrentTraceMode == InTraceButtonMode;
+
+								if (TraceModule.IsTraceAnalysisActive() && bIsButtonActive)
+								{
+									return FSlateColor(InTraceButtonMode == ETraceMode::Monitoring ? FLinearColor::Green : FLinearColor::Red);
+								}
+
+								return FSlateColor::UseForeground();
+				 			})
+						]
+						// Trace button text
+						+ SHorizontalBox::Slot()
+						.VAlign(VAlign_Center)
+						.AutoWidth()
+						.Padding(6.0f, 1.0f, 0.0f, 0.0f)
+						[
+							SNew(STextBlock)
+							.Justification(ETextJustify::Center)
+							.Text_Lambda([InTraceButtonMode]()
+				 			{
+								const IAudioInsightsTraceModule& TraceModule = FAudioInsightsEditorModule::GetChecked().GetTraceModule();
+
+								const ETraceMode CurrentTraceMode = TraceModule.GetTraceMode();
+								const bool bIsButtonActive = InTraceButtonMode == ETraceMode::Monitoring ? CurrentTraceMode != ETraceMode::None	: CurrentTraceMode == InTraceButtonMode;
+
+								if (TraceModule.IsTraceAnalysisActive() && bIsButtonActive)
+								{
+									return InTraceButtonMode == ETraceMode::Monitoring ? StopMonitoringText : StopRecordingText;
+								}
+
+								return InTraceButtonMode == ETraceMode::Monitoring ? StartMonitoringText : StartRecordingText;
+				 			})
+							.Font(IPropertyTypeCustomizationUtils::GetRegularFont())
+						]
+					]
+				]
+			];
+	}
+
 	void FEditorDashboardFactory::SetActiveDeviceId(const ::Audio::FDeviceId InDeviceId)
 	{
 		const TObjectPtr<const UAudioInsightsEditorSettings> AudioInsightsEditorSettings = GetDefault<UAudioInsightsEditorSettings>();
diff --git a/Engine/Plugins/AudioInsights/Source/AudioInsightsEditor/Private/AudioInsightsEditorDashboardFactory.h b/Engine/Plugins/AudioInsights/Source/AudioInsightsEditor/Private/AudioInsightsEditorDashboardFactory.h
--- a/Engine/Plugins/AudioInsights/Source/AudioInsightsEditor/Private/AudioInsightsEditorDashboardFactory.h
+++ b/Engine/Plugins/AudioInsights/Source/AudioInsightsEditor/Private/AudioInsightsEditorDashboardFactory.h
@@ -49,6 +49,8 @@ namespace UE::Audio::Insights
 		TSharedRef<SWidget> MakeMenuBarWidget();
 		TSharedRef<SWidget> MakeMainToolbarWidget();
 
+		TSharedRef<SWidget> MakeTraceModeButton(const ETraceMode InTraceButtonMode);
+
 		void SetActiveDeviceId(const ::Audio::FDeviceId InDeviceId);
 		void InitDelegates();
 		TSharedRef<FTabManager::FLayout> GetDefaultTabLayout();
