[Audio Insights] Fix for parameter data not being available for nested sound types when scrubbing through the cache in editor
* We need to locate dashboard entries for nested sounds via their PlayOrder and ActivePlayOrderID. We were only using the active play order before.
* Use the existing methods for locating entries from parameter messages from the device data map
* Make sure we update aggregate data once we have collected parameter data for this timestamp when scrubbing

diff --git a/Engine/Plugins/AudioInsights/Source/AudioInsights/Private/Providers/SoundTraceProvider.cpp b/Engine/Plugins/AudioInsights/Source/AudioInsights/Private/Providers/SoundTraceProvider.cpp
--- a/Engine/Plugins/AudioInsights/Source/AudioInsights/Private/Providers/SoundTraceProvider.cpp
+++ b/Engine/Plugins/AudioInsights/Source/AudioInsights/Private/Providers/SoundTraceProvider.cpp
@@ -1389,13 +1389,13 @@ namespace UE::Audio::Insights
 				}
 			}
 		}
-		
-		// Regenerate the active sounds play order array
-		ActiveSoundToEntryKeysMap.GetKeys(ActiveSoundPlayOrderArray);
 
 #if WITH_EDITOR
-		CollectParamsForTimestamp(CacheManager, ActiveSoundPlayOrderArray, TimeMarker);
+		CollectParamsForTimestamp(CacheManager, TimeMarker);
 #else
+		// Regenerate the active sounds play order array
+		ActiveSoundToEntryKeysMap.GetKeys(ActiveSoundPlayOrderArray);
+
 		CollectParamsForTimestampStandalone(ActiveSoundPlayOrderArray, TimeMarker);
 #endif //!WITH_EDITOR
 
@@ -1409,33 +1409,16 @@ namespace UE::Audio::Insights
 	}
 
 #if WITH_EDITOR
-	void FSoundTraceProvider::CollectParamsForTimestamp(const FAudioInsightsCacheManager& CacheManager, const TArray<uint32>& ActiveSoundPlayOrderArray, const double TimeMarker)
+	void FSoundTraceProvider::CollectParamsForTimestamp(const FAudioInsightsCacheManager& CacheManager, const double TimeMarker)
 	{
-		TMap<uint32, TSharedPtr<FSoundDashboardEntry>> FoundEntries;
-		FoundEntries.Reserve(ActiveSoundPlayOrderArray.Num());
+		using namespace FSoundTraceProviderPrivate;
 
-		auto GetEntry = [this, &FoundEntries](const uint32 InActiveSoundPlayOrderID)
+		for (const auto& [AudioDeviceID, DeviceData] : DeviceDataMap)
 		{
-			if (FoundEntries.Contains(InActiveSoundPlayOrderID))
-			{
-				return;
-			}
-
-			TSharedPtr<FSoundDashboardEntry> ActiveSoundEntry = nullptr;
-			GetActiveSoundEntryFromIDs(InActiveSoundPlayOrderID, ActiveSoundEntry);
-
-			if (ActiveSoundEntry.IsValid())
+			for (const auto& [EntryType, SoundDashboardEntry] : DeviceData)
 			{
-				constexpr uint32 DataPointsCapacity = 2048u;
-				ActiveSoundEntry->ResetDataBuffers(DataPointsCapacity);
-
-				FoundEntries.Add(InActiveSoundPlayOrderID, ActiveSoundEntry);
+				ResetDataBuffersRecursive(SoundDashboardEntry);
 			}
-		};
-
-		for (const uint32 PlayOrder : ActiveSoundPlayOrderArray)
-		{
-			GetEntry(PlayOrder);
 		}
 
 		enum class ESoundParamMessageType : uint8
@@ -1466,149 +1449,203 @@ namespace UE::Audio::Insights
 
 		// Using ParallelFor to speed-up the cached messages retrieval, using a traditional for loop is unacceptably slower, specially in large traces.
 		ParallelFor(NumMessageTypesToProcess,
-		[&FoundEntries, &CacheManager, TimeMarker, StartTime, EndTime, this](const int32 Index)
+		[&CacheManager, TimeMarker, StartTime, EndTime, this](const int32 Index)
 		{
 			switch (static_cast<ESoundParamMessageType>(Index))
 			{
 			case ESoundParamMessageType::Priority:
-				CacheManager.IterateOverRange<FSoundPriorityMessage>(SoundMessageNames::PriorityParam, StartTime, EndTime, [&FoundEntries, TimeMarker](const FSoundPriorityMessage& Message)
+				CacheManager.IterateOverRange<FSoundPriorityMessage>(SoundMessageNames::PriorityParam, StartTime, EndTime, [this, TimeMarker](const FSoundPriorityMessage& Message)
 				{
-					if (TSharedPtr<FSoundDashboardEntry>* SoundEntry = FoundEntries.Find(Message.ActiveSoundPlayOrder))
+					TSharedPtr<FSoundDashboardEntry> ActiveSoundEntry;
+					GetActiveSoundEntryFromIDs(Message.ActiveSoundPlayOrder, ActiveSoundEntry);
+					
+					if (ActiveSoundEntry.IsValid())
 					{
-						(*SoundEntry)->PriorityDataRange.Push({ Message.Timestamp, Message.Priority });
-						(*SoundEntry)->Timestamp = Message.Timestamp;
-
-						if (Message.Timestamp <= TimeMarker)
+						if (FSoundDashboardEntry* SoundEntry = FindSoundEntryForActiveSound(Message.PlayOrder, &ActiveSoundEntry))
 						{
-							(*SoundEntry)->PriorityDataPoint = Message.Priority;
+							SoundEntry->PriorityDataRange.Push({ Message.Timestamp, Message.Priority });
+							SoundEntry->Timestamp = Message.Timestamp;
+
+							if (Message.Timestamp <= TimeMarker)
+							{
+								SoundEntry->PriorityDataPoint = Message.Priority;
+							}
 						}
 					}
 				});
 				break;
 
 			case ESoundParamMessageType::Distance:
-				CacheManager.IterateOverRange<FSoundDistanceMessage>(SoundMessageNames::DistanceParam, StartTime, EndTime, [&FoundEntries, TimeMarker](const FSoundDistanceMessage& Message)
+				CacheManager.IterateOverRange<FSoundDistanceMessage>(SoundMessageNames::DistanceParam, StartTime, EndTime, [this, TimeMarker](const FSoundDistanceMessage& Message)
 				{
-					if (TSharedPtr<FSoundDashboardEntry>* SoundEntry = FoundEntries.Find(Message.ActiveSoundPlayOrder))
-					{
-						(*SoundEntry)->DistanceDataRange.Push({ Message.Timestamp, Message.Distance });
-						(*SoundEntry)->Timestamp = Message.Timestamp;
+					TSharedPtr<FSoundDashboardEntry> ActiveSoundEntry;
+					GetActiveSoundEntryFromIDs(Message.ActiveSoundPlayOrder, ActiveSoundEntry);
 
-						if (Message.Timestamp <= TimeMarker)
+					if (ActiveSoundEntry.IsValid())
+					{
+						if (FSoundDashboardEntry* SoundEntry = FindSoundEntryForActiveSound(Message.PlayOrder, &ActiveSoundEntry))
 						{
-							(*SoundEntry)->DistanceDataPoint = Message.Distance;
+							SoundEntry->DistanceDataRange.Push({ Message.Timestamp, Message.Distance });
+							SoundEntry->Timestamp = Message.Timestamp;
+
+							if (Message.Timestamp <= TimeMarker)
+							{
+								SoundEntry->DistanceDataPoint = Message.Distance;
+							}
 						}
 					}
 				});
 				break;
 
 			case ESoundParamMessageType::DistanceAttenuation:
-				CacheManager.IterateOverRange<FSoundDistanceAttenuationMessage>(SoundMessageNames::DistanceAttenuationParam, StartTime, EndTime, [&FoundEntries, TimeMarker](const FSoundDistanceAttenuationMessage& Message)
+				CacheManager.IterateOverRange<FSoundDistanceAttenuationMessage>(SoundMessageNames::DistanceAttenuationParam, StartTime, EndTime, [this, TimeMarker](const FSoundDistanceAttenuationMessage& Message)
 				{
-					if (TSharedPtr<FSoundDashboardEntry>* SoundEntry = FoundEntries.Find(Message.ActiveSoundPlayOrder))
-					{
-						(*SoundEntry)->DistanceAttenuationDataRange.Push({ Message.Timestamp, Message.DistanceAttenuation });
-						(*SoundEntry)->Timestamp = Message.Timestamp;
+					TSharedPtr<FSoundDashboardEntry> ActiveSoundEntry;
+					GetActiveSoundEntryFromIDs(Message.ActiveSoundPlayOrder, ActiveSoundEntry);
 
-						if (Message.Timestamp <= TimeMarker)
+					if (ActiveSoundEntry.IsValid())
+					{
+						if (FSoundDashboardEntry* SoundEntry = FindSoundEntryForActiveSound(Message.PlayOrder, &ActiveSoundEntry))
 						{
-							(*SoundEntry)->DistanceAttenuationDataPoint = Message.DistanceAttenuation;
+							SoundEntry->DistanceAttenuationDataRange.Push({ Message.Timestamp, Message.DistanceAttenuation });
+							SoundEntry->Timestamp = Message.Timestamp;
+
+							if (Message.Timestamp <= TimeMarker)
+							{
+								SoundEntry->DistanceAttenuationDataPoint = Message.DistanceAttenuation;
+							}
 						}
 					}
 				});
 				break;
 
 			case ESoundParamMessageType::LPFFreq:
-				CacheManager.IterateOverRange<FSoundLPFFreqMessage>(SoundMessageNames::LPFFreqParam, StartTime, EndTime, [&FoundEntries, TimeMarker](const FSoundLPFFreqMessage& Message)
+				CacheManager.IterateOverRange<FSoundLPFFreqMessage>(SoundMessageNames::LPFFreqParam, StartTime, EndTime, [this, TimeMarker](const FSoundLPFFreqMessage& Message)
 				{
-					if (TSharedPtr<FSoundDashboardEntry>* SoundEntry = FoundEntries.Find(Message.ActiveSoundPlayOrder))
-					{
-						(*SoundEntry)->LPFFreqDataRange.Push({ Message.Timestamp, Message.LPFFrequency });
-						(*SoundEntry)->Timestamp = Message.Timestamp;
+					TSharedPtr<FSoundDashboardEntry> ActiveSoundEntry;
+					GetActiveSoundEntryFromIDs(Message.ActiveSoundPlayOrder, ActiveSoundEntry);
 
-						if (Message.Timestamp <= TimeMarker)
+					if (ActiveSoundEntry.IsValid())
+					{
+						if (FSoundDashboardEntry* SoundEntry = FindSoundEntryForActiveSound(Message.PlayOrder, &ActiveSoundEntry))
 						{
-							(*SoundEntry)->LPFFreqDataPoint = Message.LPFFrequency;
+							SoundEntry->LPFFreqDataRange.Push({ Message.Timestamp, Message.LPFFrequency });
+							SoundEntry->Timestamp = Message.Timestamp;
+
+							if (Message.Timestamp <= TimeMarker)
+							{
+								SoundEntry->LPFFreqDataPoint = Message.LPFFrequency;
+							}
 						}
 					}
 				});
 				break;
 
 			case ESoundParamMessageType::HPFFreq:
-				CacheManager.IterateOverRange<FSoundHPFFreqMessage>(SoundMessageNames::HPFFreqParam, StartTime, EndTime, [&FoundEntries, TimeMarker](const FSoundHPFFreqMessage& Message)
+				CacheManager.IterateOverRange<FSoundHPFFreqMessage>(SoundMessageNames::HPFFreqParam, StartTime, EndTime, [this, TimeMarker](const FSoundHPFFreqMessage& Message)
 				{
-					if (TSharedPtr<FSoundDashboardEntry>* SoundEntry = FoundEntries.Find(Message.ActiveSoundPlayOrder))
-					{
-						(*SoundEntry)->HPFFreqDataRange.Push({ Message.Timestamp, Message.HPFFrequency });
-						(*SoundEntry)->Timestamp = Message.Timestamp;
+					TSharedPtr<FSoundDashboardEntry> ActiveSoundEntry;
+					GetActiveSoundEntryFromIDs(Message.ActiveSoundPlayOrder, ActiveSoundEntry);
 
-						if (Message.Timestamp <= TimeMarker)
+					if (ActiveSoundEntry.IsValid())
+					{
+						if (FSoundDashboardEntry* SoundEntry = FindSoundEntryForActiveSound(Message.PlayOrder, &ActiveSoundEntry))
 						{
-							(*SoundEntry)->HPFFreqDataPoint = Message.HPFFrequency;
+							SoundEntry->HPFFreqDataRange.Push({ Message.Timestamp, Message.HPFFrequency });
+							SoundEntry->Timestamp = Message.Timestamp;
+
+							if (Message.Timestamp <= TimeMarker)
+							{
+								SoundEntry->HPFFreqDataPoint = Message.HPFFrequency;
+							}
 						}
 					}
 				});
 				break;
 
 			case ESoundParamMessageType::Envelope:
-				CacheManager.IterateOverRange<FSoundEnvelopeMessage>(SoundMessageNames::EnvelopeParam, StartTime, EndTime, [&FoundEntries, TimeMarker](const FSoundEnvelopeMessage& Message)
+				CacheManager.IterateOverRange<FSoundEnvelopeMessage>(SoundMessageNames::EnvelopeParam, StartTime, EndTime, [this, TimeMarker](const FSoundEnvelopeMessage& Message)
 				{
-					if (TSharedPtr<FSoundDashboardEntry>* SoundEntry = FoundEntries.Find(Message.ActiveSoundPlayOrder))
-					{
-						(*SoundEntry)->AmplitudeDataRange.Push({ Message.Timestamp, Message.Envelope });
-						(*SoundEntry)->Timestamp = Message.Timestamp;
+					TSharedPtr<FSoundDashboardEntry> ActiveSoundEntry;
+					GetActiveSoundEntryFromIDs(Message.ActiveSoundPlayOrder, ActiveSoundEntry);
 
-						if (Message.Timestamp <= TimeMarker)
+					if (ActiveSoundEntry.IsValid())
+					{
+						if (FSoundDashboardEntry* SoundEntry = FindSoundEntryForActiveSound(Message.PlayOrder, &ActiveSoundEntry))
 						{
-							(*SoundEntry)->AmplitudeDataPoint = Message.Envelope;
+							SoundEntry->AmplitudeDataRange.Push({ Message.Timestamp, Message.Envelope });
+							SoundEntry->Timestamp = Message.Timestamp;
+
+							if (Message.Timestamp <= TimeMarker)
+							{
+								SoundEntry->AmplitudeDataPoint = Message.Envelope;
+							}
 						}
 					}
 				});
 				break;
 
 			case ESoundParamMessageType::Volume:
-				CacheManager.IterateOverRange<FSoundVolumeMessage>(SoundMessageNames::VolumeParam, StartTime, EndTime, [&FoundEntries, TimeMarker](const FSoundVolumeMessage& Message)
+				CacheManager.IterateOverRange<FSoundVolumeMessage>(SoundMessageNames::VolumeParam, StartTime, EndTime, [this, TimeMarker](const FSoundVolumeMessage& Message)
 				{
-					if (TSharedPtr<FSoundDashboardEntry>* SoundEntry = FoundEntries.Find(Message.ActiveSoundPlayOrder))
-					{
-						(*SoundEntry)->VolumeDataRange.Push({ Message.Timestamp, Message.Volume });
-						(*SoundEntry)->Timestamp = Message.Timestamp;
+					TSharedPtr<FSoundDashboardEntry> ActiveSoundEntry;
+					GetActiveSoundEntryFromIDs(Message.ActiveSoundPlayOrder, ActiveSoundEntry);
 
-						if (Message.Timestamp <= TimeMarker)
+					if (ActiveSoundEntry.IsValid())
+					{
+						if (FSoundDashboardEntry* SoundEntry = FindSoundEntryForActiveSound(Message.PlayOrder, &ActiveSoundEntry))
 						{
-							(*SoundEntry)->VolumeDataPoint = Message.Volume;
+							SoundEntry->VolumeDataRange.Push({ Message.Timestamp, Message.Volume });
+							SoundEntry->Timestamp = Message.Timestamp;
+
+							if (Message.Timestamp <= TimeMarker)
+							{
+								SoundEntry->VolumeDataPoint = Message.Volume;
+							}
 						}
 					}
 				});
 				break;
 
 			case ESoundParamMessageType::Pitch:
-				CacheManager.IterateOverRange<FSoundPitchMessage>(SoundMessageNames::PitchParam, StartTime, EndTime, [&FoundEntries, TimeMarker](const FSoundPitchMessage& Message)
+				CacheManager.IterateOverRange<FSoundPitchMessage>(SoundMessageNames::PitchParam, StartTime, EndTime, [this, TimeMarker](const FSoundPitchMessage& Message)
 				{
-					if (TSharedPtr<FSoundDashboardEntry>* SoundEntry = FoundEntries.Find(Message.ActiveSoundPlayOrder))
-					{
-						(*SoundEntry)->PitchDataRange.Push({ Message.Timestamp, Message.Pitch });
-						(*SoundEntry)->Timestamp = Message.Timestamp;
+					TSharedPtr<FSoundDashboardEntry> ActiveSoundEntry;
+					GetActiveSoundEntryFromIDs(Message.ActiveSoundPlayOrder, ActiveSoundEntry);
 
-						if (Message.Timestamp <= TimeMarker)
+					if (ActiveSoundEntry.IsValid())
+					{
+						if (FSoundDashboardEntry* SoundEntry = FindSoundEntryForActiveSound(Message.PlayOrder, &ActiveSoundEntry))
 						{
-							(*SoundEntry)->PitchDataPoint = Message.Pitch;
+							SoundEntry->PitchDataRange.Push({ Message.Timestamp, Message.Pitch });
+							SoundEntry->Timestamp = Message.Timestamp;
+
+							if (Message.Timestamp <= TimeMarker)
+							{
+								SoundEntry->PitchDataPoint = Message.Pitch;
+							}
 						}
 					}
 				});
 				break;
 
 			case ESoundParamMessageType::RelativeRenderCost:
-				CacheManager.IterateOverRange<FSoundRelativeRenderCostMessage>(SoundMessageNames::RelativeRenderCostParam, StartTime, EndTime, [&FoundEntries, TimeMarker](const FSoundRelativeRenderCostMessage& Message)
+				CacheManager.IterateOverRange<FSoundRelativeRenderCostMessage>(SoundMessageNames::RelativeRenderCostParam, StartTime, EndTime, [this, TimeMarker](const FSoundRelativeRenderCostMessage& Message)
 				{
-					if (TSharedPtr<FSoundDashboardEntry>* SoundEntry = FoundEntries.Find(Message.ActiveSoundPlayOrder))
-					{
-						(*SoundEntry)->RelativeRenderCostDataRange.Push({ Message.Timestamp, Message.RelativeRenderCost });
-						(*SoundEntry)->Timestamp = Message.Timestamp;
+					TSharedPtr<FSoundDashboardEntry> ActiveSoundEntry;
+					GetActiveSoundEntryFromIDs(Message.ActiveSoundPlayOrder, ActiveSoundEntry);
 
-						if (Message.Timestamp <= TimeMarker)
+					if (ActiveSoundEntry.IsValid())
+					{
+						if (FSoundDashboardEntry* SoundEntry = FindSoundEntryForActiveSound(Message.PlayOrder, &ActiveSoundEntry))
 						{
-							(*SoundEntry)->RelativeRenderCostDataPoint = Message.RelativeRenderCost;
+							SoundEntry->RelativeRenderCostDataRange.Push({ Message.Timestamp, Message.RelativeRenderCost });
+							SoundEntry->Timestamp = Message.Timestamp;
+
+							if (Message.Timestamp <= TimeMarker)
+							{
+								SoundEntry->RelativeRenderCostDataPoint = Message.RelativeRenderCost;
+							}
 						}
 					}
 				});
@@ -1619,6 +1656,8 @@ namespace UE::Audio::Insights
 			}
 		});
 
+		UpdateAggregateActiveSoundData();
+
 		OnProcessPlotData.Broadcast();
 	}
 #else
@@ -1813,5 +1852,25 @@ namespace UE::Audio::Insights
 		}
 	}
 
+	void FSoundTraceProvider::ResetDataBuffersRecursive(const TSharedPtr<IDashboardDataTreeViewEntry>& Entry)
+	{
+		using namespace FSoundTraceProviderPrivate;
+
+		if (!Entry.IsValid())
+		{
+			return;
+		}
+
+		FSoundDashboardEntry& SoundEntry = CastEntry(*Entry);
+
+		constexpr uint32 DataPointsCapacity = 2048u;
+		SoundEntry.ResetDataBuffers(DataPointsCapacity);
+
+		for (const TSharedPtr<IDashboardDataTreeViewEntry>& Child : SoundEntry.Children)
+		{
+			ResetDataBuffersRecursive(Child);
+		}
+	}
+
 } // namespace UE::Audio::Insights
 #undef LOCTEXT_NAMESPACE
\ No newline at end of file
diff --git a/Engine/Plugins/AudioInsights/Source/AudioInsights/Private/Providers/SoundTraceProvider.h b/Engine/Plugins/AudioInsights/Source/AudioInsights/Private/Providers/SoundTraceProvider.h
--- a/Engine/Plugins/AudioInsights/Source/AudioInsights/Private/Providers/SoundTraceProvider.h
+++ b/Engine/Plugins/AudioInsights/Source/AudioInsights/Private/Providers/SoundTraceProvider.h
@@ -53,12 +53,13 @@ namespace UE::Audio::Insights
 		virtual void OnTimeControlMethodReset() override;
 
 #if WITH_EDITOR
-		void CollectParamsForTimestamp(const FAudioInsightsCacheManager& CacheManager, const TArray<uint32>& ActiveSoundPlayOrderArray, const double TimeMarker);
+		void CollectParamsForTimestamp(const FAudioInsightsCacheManager& CacheManager, const double TimeMarker);
 #else
 		void CollectParamsForTimestampStandalone(const TArray<uint32>& ActiveSoundPlayOrderArray, const double TimeMarker);
 #endif // WITH_EDITOR
 
 		void CacheIsPlottingFlagRecursive(const TSharedPtr<IDashboardDataTreeViewEntry>& Entry);
+		void ResetDataBuffersRecursive(const TSharedPtr<IDashboardDataTreeViewEntry>& Entry);
 
 #if !WITH_EDITOR
 		TUniquePtr<FSoundSessionCachedMessages> SessionCachedMessages;
