[Audio Insights] Ensure we use the correct PlayOrder ID to distinguish between active sounds and wave instances
* Active Sounds and Wave Instances have separate Play Order ID counters , meaning they cannot be used as one unique ID
* Store both IDs inside SoundDashboardEntry and combine to form a new unique ID in GetEntryID
* Fixes parameter data being overwritten when a wave instance and an active sound have the same Play Order
* Fix Wave Instances not enabling their plot if processing is paused and then resumed

diff --git a/Engine/Plugins/AudioInsights/Source/AudioInsights/Private/Messages/SoundTraceMessages.cpp b/Engine/Plugins/AudioInsights/Source/AudioInsights/Private/Messages/SoundTraceMessages.cpp
--- a/Engine/Plugins/AudioInsights/Source/AudioInsights/Private/Messages/SoundTraceMessages.cpp
+++ b/Engine/Plugins/AudioInsights/Source/AudioInsights/Private/Messages/SoundTraceMessages.cpp
@@ -22,13 +22,20 @@ namespace UE::Audio::Insights
 		const FName RelativeRenderCostParam = "RelativeRenderCost";
 	};
 
+	namespace SoundMessageUtils
+	{
+		uint64 GeneratePlayOrderUniqueID(const uint32 ActiveSoundPlayOrder, const uint32 WaveInstancePlayOrder)
+		{
+			return (static_cast<uint64>(ActiveSoundPlayOrder) << 32) + static_cast<uint64>(WaveInstancePlayOrder);
+		}
+	};
+
 	FSoundMessageBase::FSoundMessageBase(const Trace::IAnalyzer::FOnEventContext& InContext)
 	{
 		const Trace::IAnalyzer::FEventData& EventData = InContext.EventData;
 
 		DeviceId = static_cast<::Audio::FDeviceId>(EventData.GetValue<uint32>("DeviceId"));
 		Timestamp = InContext.EventTime.AsSeconds(EventData.GetValue<uint64>("Timestamp"));
-		PlayOrder = EventData.GetValue<uint32>("PlayOrder");
 	}
 
 	FSoundStartMessage::FSoundStartMessage(const Trace::IAnalyzer::FOnEventContext& InContext)
@@ -36,6 +43,7 @@ namespace UE::Audio::Insights
 	{
 		const Trace::IAnalyzer::FEventData& EventData = InContext.EventData;
 		EventData.GetString("AssetPath", Name);
+		ActiveSoundPlayOrder = EventData.GetValue<uint32>("PlayOrder");
 
 		FString SoundClassName;
 
@@ -118,4 +126,4 @@ namespace UE::Audio::Insights
 		PitchDataRange.Reset(DataPointsCapacity);
 		RelativeRenderCostDataRange.Reset(DataPointsCapacity);
 	}
-}
\ No newline at end of file
+}
diff --git a/Engine/Plugins/AudioInsights/Source/AudioInsights/Private/Providers/SoundTraceProvider.cpp b/Engine/Plugins/AudioInsights/Source/AudioInsights/Private/Providers/SoundTraceProvider.cpp
--- a/Engine/Plugins/AudioInsights/Source/AudioInsights/Private/Providers/SoundTraceProvider.cpp
+++ b/Engine/Plugins/AudioInsights/Source/AudioInsights/Private/Providers/SoundTraceProvider.cpp
@@ -79,10 +79,10 @@ namespace UE::Audio::Insights
 
 		void UpdateParameterEntry(const FSoundParameterMessage& Msg, const float DataPoint, float& OutDataPoint, ::Audio::TCircularAudioBuffer<FDataPoint>& OutDataRange, TArray<int32, TInlineAllocator<64>>& OutEntriesWithPoppedDataPoints)
 		{
-			if (!OutEntriesWithPoppedDataPoints.Contains(Msg.PlayOrder))
+			if (!OutEntriesWithPoppedDataPoints.Contains(Msg.WaveInstancePlayOrder))
 			{
 				OutDataRange.Pop(OutDataRange.Num());
-				OutEntriesWithPoppedDataPoints.Add(Msg.PlayOrder);
+				OutEntriesWithPoppedDataPoints.Add(Msg.WaveInstancePlayOrder);
 			}
 
 			OutDataPoint = DataPoint;
@@ -100,6 +100,35 @@ namespace UE::Audio::Insights
 				OutCachedMessages.EmplaceBack(Msg);
 			}
 		}
+
+		struct PlayOrderPair
+		{
+			uint32 ActiveSoundPlayOrderID { INDEX_NONE };
+			uint32 WaveInstancePlayOrderID { INDEX_NONE };
+		};
+
+		void CollectPlayOrderPairsRecursive(const TSharedPtr<IDashboardDataTreeViewEntry>& Entry, TArray<PlayOrderPair>& OutPlayOrderArray)
+		{
+			using namespace FSoundTraceProviderPrivate;
+
+			if (!Entry.IsValid())
+			{
+				return;
+			}
+
+			const FSoundDashboardEntry& SoundEntry = CastEntry(*Entry);
+			const PlayOrderPair PlayOrderPair { SoundEntry.ActiveSoundPlayOrder, SoundEntry.WaveInstancePlayOrder };
+
+			if (PlayOrderPair.WaveInstancePlayOrderID != INDEX_NONE || PlayOrderPair.ActiveSoundPlayOrderID != INDEX_NONE)
+			{
+				OutPlayOrderArray.Add(PlayOrderPair);
+			}
+
+			for (const TSharedPtr<IDashboardDataTreeViewEntry>& Child : SoundEntry.Children)
+			{
+				CollectPlayOrderPairsRecursive(Child, OutPlayOrderArray);
+			}
+		}
 #endif // !WITH_EDITOR
 
 		//////////////////////////////////////////////////////////////////////////////////////////////////////////////
@@ -115,7 +144,7 @@ namespace UE::Audio::Insights
 			TSharedPtr<IDashboardDataTreeViewEntry>* ActiveSoundEntry = CategoryEntry->Get()->Children.FindByPredicate([ActiveSoundPlayOrder](TSharedPtr<IDashboardDataTreeViewEntry> ChildEntry)
 			{
 				FSoundDashboardEntry& NewEntry = CastEntry(*ChildEntry.Get());
-				return NewEntry.PlayOrder == ActiveSoundPlayOrder;
+				return NewEntry.ActiveSoundPlayOrder == ActiveSoundPlayOrder;
 			});
 
 			if (ActiveSoundEntry == nullptr || !ActiveSoundEntry->IsValid())
@@ -128,7 +157,7 @@ namespace UE::Audio::Insights
 
 		//////////////////////////////////////////////////////////////////////////////////////////////////////////////
 		// Will return either the Active Sound entry or the child SoundWave entry depending on the sound's EntryType
-		FSoundDashboardEntry* FindSoundEntryForActiveSound(const uint32 SoundWavePlayOrder, const TSharedPtr<FSoundDashboardEntry>* ActiveSoundEntry)
+		FSoundDashboardEntry* FindWaveEntryInActiveSound(const uint32 SoundWavePlayOrder, const TSharedPtr<FSoundDashboardEntry>* ActiveSoundEntry)
 		{
 			if (ActiveSoundEntry == nullptr || !ActiveSoundEntry->IsValid())
 			{
@@ -148,7 +177,7 @@ namespace UE::Audio::Insights
 				}
 
 				FSoundDashboardEntry* SoundWaveEntryCast = static_cast<FSoundDashboardEntry*>(SoundWaveEntry.Get());
-				if (SoundWaveEntryCast->PlayOrder == SoundWavePlayOrder)
+				if (SoundWaveEntryCast->WaveInstancePlayOrder == SoundWavePlayOrder)
 				{
 					return SoundWaveEntryCast;
 				}
@@ -159,14 +188,21 @@ namespace UE::Audio::Insights
 
 #if !WITH_EDITOR
 		template<typename T>
-		const T* FindClosestMessageToTimestamp(const TraceServices::TPagedArray<T>& InCachedMessages, const double InTimeMarker, const uint32 InPlayOrder)
+		const T* FindClosestMessageToTimestamp(const TraceServices::TPagedArray<T>& InCachedMessages, const double InTimeMarker, const PlayOrderPair InPlayOrderPair)
 		{
 			const int32 ClosestMessageToTimeStampIndex = TraceServices::PagedArrayAlgo::BinarySearchClosestBy(InCachedMessages, InTimeMarker, [](const T& Msg) { return Msg.Timestamp; });
 
 			// Iterate backwards from TimeMarker until we find the matching PlayOrder
 			for (auto It = InCachedMessages.GetIteratorFromItem(ClosestMessageToTimeStampIndex); It; --It)
 			{
-				if (It->ActiveSoundPlayOrder == InPlayOrder)
+				if (InPlayOrderPair.WaveInstancePlayOrderID != INDEX_NONE)
+				{
+					if (It->WaveInstancePlayOrder == InPlayOrderPair.WaveInstancePlayOrderID)
+					{
+						return &(*It);
+					}
+				}
+				else if (It->ActiveSoundPlayOrder == InPlayOrderPair.ActiveSoundPlayOrderID)
 				{
 					return &(*It);
 				}
@@ -473,7 +509,7 @@ namespace UE::Audio::Insights
 		{
 			SoundEntryReturn = nullptr;
 
-			GetActiveSoundEntryFromIDs(Msg.PlayOrder, SoundEntryReturn);
+			GetActiveSoundEntryFromIDs(Msg.ActiveSoundPlayOrder, SoundEntryReturn);
 			return &SoundEntryReturn;
 		};
 
@@ -488,8 +524,9 @@ namespace UE::Audio::Insights
 			ActiveSoundEntryRef.SetName(*Msg.Name);
 			ActiveSoundEntryRef.EntryType = Msg.EntryType;
 			ActiveSoundEntryRef.CategoryName = FText::FromString(FSoundTraceProviderPrivate::GetEntryTypeDisplayName(Msg.EntryType));
+			ActiveSoundEntryRef.ActiveSoundPlayOrder = Msg.ActiveSoundPlayOrder;
 
-			if (PlottingSoundEntries.Contains(ActiveSoundEntryRef.PlayOrder))
+			if (PlottingSoundEntries.Contains(ActiveSoundEntryRef.GetEntryID()))
 			{
 				ActiveSoundEntryRef.bIsPlotActive = true;
 			}
@@ -500,12 +537,7 @@ namespace UE::Audio::Insights
 
 		auto ProcessSoundWaveStartMessage = [this](const FSoundWaveStartMessage& Msg, const TSharedPtr<FSoundDashboardEntry>* OutActiveSoundEntry)
 		{
-			if (OutActiveSoundEntry == nullptr || !OutActiveSoundEntry->IsValid())
-			{
-				return;
-			}
-
-			if (!FSoundTraceProviderPrivate::EntryTypeHasSoundWaveEntries((*OutActiveSoundEntry)->EntryType))
+			if (OutActiveSoundEntry == nullptr || !OutActiveSoundEntry->IsValid() || !FSoundTraceProviderPrivate::EntryTypeHasSoundWaveEntries((*OutActiveSoundEntry)->EntryType))
 			{
 				return;
 			}
@@ -513,7 +545,7 @@ namespace UE::Audio::Insights
 			// Avoid creating duplicate entries
 			const bool bEntryExists = (*OutActiveSoundEntry)->Children.ContainsByPredicate([&Msg](const TSharedPtr<IDashboardDataTreeViewEntry>& ChildEntry)
 			{
-				return ChildEntry.IsValid() && Msg.PlayOrder == ChildEntry->GetEntryID();
+				return ChildEntry.IsValid() && Msg.WaveInstancePlayOrder == CastEntry(*ChildEntry).WaveInstancePlayOrder;
 			});
 
 			if (bEntryExists)
@@ -523,7 +555,8 @@ namespace UE::Audio::Insights
 			
 			TSharedPtr<FSoundDashboardEntry> SoundWaveEntry = MakeShared<FSoundDashboardEntry>();
 			SoundWaveEntry->DeviceId = Msg.DeviceId;
-			SoundWaveEntry->PlayOrder = Msg.PlayOrder;
+			SoundWaveEntry->WaveInstancePlayOrder = Msg.WaveInstancePlayOrder;
+			SoundWaveEntry->ActiveSoundPlayOrder = Msg.ActiveSoundPlayOrder;
 			SoundWaveEntry->Timestamp = Msg.Timestamp;
 			SoundWaveEntry->SetName(*Msg.Name);
 			SoundWaveEntry->EntryType = Msg.EntryType;
@@ -531,6 +564,11 @@ namespace UE::Audio::Insights
 			SoundWaveEntry->PinnedEntryType = (*OutActiveSoundEntry)->PinnedEntryType;
 			SoundWaveEntry->ActorLabel = FText::FromString(*Msg.ActorLabel);
 
+			if (PlottingSoundEntries.Contains(SoundWaveEntry->GetEntryID()))
+			{
+				SoundWaveEntry->bIsPlotActive = true;
+			}
+
 			(*OutActiveSoundEntry)->Children.Add(SoundWaveEntry);
 		};
 
@@ -592,7 +630,7 @@ namespace UE::Audio::Insights
 		ProcessMessageQueue<FSoundPriorityMessage>(TraceMessages.PriorityMessages, GetActiveSoundEntryFromParameterMessage,
 		[this, &EntriesWithPoppedDataPoints](const FSoundPriorityMessage& Msg, const TSharedPtr<FSoundDashboardEntry>* OutActiveSoundEntry)
 		{
-			FSoundDashboardEntry* SoundEntry = FSoundTraceProviderPrivate::FindSoundEntryForActiveSound(Msg.PlayOrder, OutActiveSoundEntry);
+			FSoundDashboardEntry* SoundEntry = FSoundTraceProviderPrivate::FindWaveEntryInActiveSound(Msg.WaveInstancePlayOrder, OutActiveSoundEntry);
 			if (SoundEntry == nullptr)
 			{
 				return;
@@ -612,7 +650,7 @@ namespace UE::Audio::Insights
 		ProcessMessageQueue<FSoundDistanceMessage>(TraceMessages.DistanceMessages, GetActiveSoundEntryFromParameterMessage,
 		[this, &EntriesWithPoppedDataPoints](const FSoundDistanceMessage& Msg, const TSharedPtr<FSoundDashboardEntry>* OutActiveSoundEntry)
 		{
-			FSoundDashboardEntry* SoundEntry = FSoundTraceProviderPrivate::FindSoundEntryForActiveSound(Msg.PlayOrder, OutActiveSoundEntry);
+			FSoundDashboardEntry* SoundEntry = FSoundTraceProviderPrivate::FindWaveEntryInActiveSound(Msg.WaveInstancePlayOrder, OutActiveSoundEntry);
 			if (SoundEntry == nullptr)
 			{
 				return;
@@ -632,7 +670,7 @@ namespace UE::Audio::Insights
 		ProcessMessageQueue<FSoundDistanceAttenuationMessage>(TraceMessages.DistanceAttenuationMessages, GetActiveSoundEntryFromParameterMessage,
 		[this, &EntriesWithPoppedDataPoints](const FSoundDistanceAttenuationMessage& Msg, const TSharedPtr<FSoundDashboardEntry>* OutActiveSoundEntry)
 		{
-			FSoundDashboardEntry* SoundEntry = FSoundTraceProviderPrivate::FindSoundEntryForActiveSound(Msg.PlayOrder, OutActiveSoundEntry);
+			FSoundDashboardEntry* SoundEntry = FSoundTraceProviderPrivate::FindWaveEntryInActiveSound(Msg.WaveInstancePlayOrder, OutActiveSoundEntry);
 			if (SoundEntry == nullptr)
 			{
 				return;
@@ -652,7 +690,7 @@ namespace UE::Audio::Insights
 		ProcessMessageQueue<FSoundLPFFreqMessage>(TraceMessages.LPFFreqMessages, GetActiveSoundEntryFromParameterMessage,
 		[this, &EntriesWithPoppedDataPoints](const FSoundLPFFreqMessage& Msg, const TSharedPtr<FSoundDashboardEntry>* OutActiveSoundEntry)
 		{
-			FSoundDashboardEntry* SoundEntry = FSoundTraceProviderPrivate::FindSoundEntryForActiveSound(Msg.PlayOrder, OutActiveSoundEntry);
+			FSoundDashboardEntry* SoundEntry = FSoundTraceProviderPrivate::FindWaveEntryInActiveSound(Msg.WaveInstancePlayOrder, OutActiveSoundEntry);
 			if (SoundEntry == nullptr)
 			{
 				return;
@@ -672,7 +710,7 @@ namespace UE::Audio::Insights
 		ProcessMessageQueue<FSoundHPFFreqMessage>(TraceMessages.HPFFreqMessages, GetActiveSoundEntryFromParameterMessage,
 		[this, &EntriesWithPoppedDataPoints](const FSoundHPFFreqMessage& Msg, const TSharedPtr<FSoundDashboardEntry>* OutActiveSoundEntry)
 		{
-			FSoundDashboardEntry* SoundEntry = FSoundTraceProviderPrivate::FindSoundEntryForActiveSound(Msg.PlayOrder, OutActiveSoundEntry);
+			FSoundDashboardEntry* SoundEntry = FSoundTraceProviderPrivate::FindWaveEntryInActiveSound(Msg.WaveInstancePlayOrder, OutActiveSoundEntry);
 			if (SoundEntry == nullptr)
 			{
 				return;
@@ -692,7 +730,7 @@ namespace UE::Audio::Insights
 		ProcessMessageQueue<FSoundEnvelopeMessage>(TraceMessages.AmplitudeMessages, GetActiveSoundEntryFromParameterMessage,
 		[this, &EntriesWithPoppedDataPoints](const FSoundEnvelopeMessage& Msg, const TSharedPtr<FSoundDashboardEntry>* OutActiveSoundEntry)
 		{
-			FSoundDashboardEntry* SoundEntry = FSoundTraceProviderPrivate::FindSoundEntryForActiveSound(Msg.PlayOrder, OutActiveSoundEntry);
+			FSoundDashboardEntry* SoundEntry = FSoundTraceProviderPrivate::FindWaveEntryInActiveSound(Msg.WaveInstancePlayOrder, OutActiveSoundEntry);
 			if (SoundEntry == nullptr)
 			{
 				return;
@@ -712,7 +750,7 @@ namespace UE::Audio::Insights
 		ProcessMessageQueue<FSoundVolumeMessage>(TraceMessages.VolumeMessages, GetActiveSoundEntryFromParameterMessage,
 		[this, &EntriesWithPoppedDataPoints](const FSoundVolumeMessage& Msg, const TSharedPtr<FSoundDashboardEntry>* OutActiveSoundEntry)
 		{
-			FSoundDashboardEntry* SoundEntry = FSoundTraceProviderPrivate::FindSoundEntryForActiveSound(Msg.PlayOrder, OutActiveSoundEntry);
+			FSoundDashboardEntry* SoundEntry = FSoundTraceProviderPrivate::FindWaveEntryInActiveSound(Msg.WaveInstancePlayOrder, OutActiveSoundEntry);
 			if (SoundEntry == nullptr)
 			{
 				return;
@@ -732,7 +770,7 @@ namespace UE::Audio::Insights
 		ProcessMessageQueue<FSoundPitchMessage>(TraceMessages.PitchMessages, GetActiveSoundEntryFromParameterMessage,
 		[this, &EntriesWithPoppedDataPoints](const FSoundPitchMessage& Msg, const TSharedPtr<FSoundDashboardEntry>* OutActiveSoundEntry)
 		{
-			FSoundDashboardEntry* SoundEntry = FSoundTraceProviderPrivate::FindSoundEntryForActiveSound(Msg.PlayOrder, OutActiveSoundEntry);
+			FSoundDashboardEntry* SoundEntry = FSoundTraceProviderPrivate::FindWaveEntryInActiveSound(Msg.WaveInstancePlayOrder, OutActiveSoundEntry);
 			if (SoundEntry == nullptr)
 			{
 				return;
@@ -752,7 +790,7 @@ namespace UE::Audio::Insights
 		ProcessMessageQueue<FSoundRelativeRenderCostMessage>(TraceMessages.RelativeRenderCostMessages, GetActiveSoundEntryFromParameterMessage,
 		[this, &EntriesWithPoppedDataPoints](const FSoundRelativeRenderCostMessage& Msg, const TSharedPtr<FSoundDashboardEntry>* OutActiveSoundEntry)
 		{
-			FSoundDashboardEntry* SoundEntry = FSoundTraceProviderPrivate::FindSoundEntryForActiveSound(Msg.PlayOrder, OutActiveSoundEntry);
+			FSoundDashboardEntry* SoundEntry = FSoundTraceProviderPrivate::FindWaveEntryInActiveSound(Msg.WaveInstancePlayOrder, OutActiveSoundEntry);
 			if (SoundEntry == nullptr)
 			{
 				return;
@@ -790,18 +828,27 @@ namespace UE::Audio::Insights
 					{
 						FSoundDashboardEntry& SoundWaveEntry = FSoundTraceProviderPrivate::CastEntry(*ChildEntry.Get());
 						SoundWaveEntry.TimeoutTimestamp = TimeoutTimestamp;
+
+						if (SoundWaveEntry.bIsPlotActive)
+						{
+							PlottingSoundEntries.FindOrAdd(SoundWaveEntry.GetEntryID());
+						}
+						else if (PlottingSoundEntries.Contains(SoundWaveEntry.GetEntryID()))
+						{
+							PlottingSoundEntries.Remove(SoundWaveEntry.GetEntryID());
+						}
 					}
 
-					EntriesTimingOut.Push(SoundMessageIDs{ Msg.DeviceId, Msg.PlayOrder });
+					EntriesTimingOut.Push(SoundMessageIDs{ Msg.DeviceId, Msg.ActiveSoundPlayOrder });
 
 					// Cache the bIsPlotActive flag in case this sound returns in future
 					if ((*OutActiveSoundEntry)->bIsPlotActive)
 					{
-						PlottingSoundEntries.FindOrAdd((*OutActiveSoundEntry)->PlayOrder);
+						PlottingSoundEntries.FindOrAdd((*OutActiveSoundEntry)->GetEntryID());
 					}
-					else if (PlottingSoundEntries.Contains((*OutActiveSoundEntry)->PlayOrder))
+					else if (PlottingSoundEntries.Contains((*OutActiveSoundEntry)->GetEntryID()))
 					{
-						PlottingSoundEntries.Remove((*OutActiveSoundEntry)->PlayOrder);
+						PlottingSoundEntries.Remove((*OutActiveSoundEntry)->GetEntryID());
 					}
 				}
 			}
@@ -809,7 +856,7 @@ namespace UE::Audio::Insights
 			{
 				// Keep track of any messages that have sent stop but no start message
 				// In rare cases these can be processed out of order
-				SoundsStoppedBeforeStart.Add(Msg.PlayOrder);
+				SoundsStoppedBeforeStart.Add(Msg.ActiveSoundPlayOrder);
 			}
 		});
 
@@ -827,7 +874,7 @@ namespace UE::Audio::Insights
 		for (const SoundMessageIDs SoundMessageIDs : EntriesTimingOut)
 		{
 			TSharedPtr<FSoundDashboardEntry> ActiveSoundEntry = nullptr;
-			GetActiveSoundEntryFromIDs(SoundMessageIDs.PlayOrder, ActiveSoundEntry);
+			GetActiveSoundEntryFromIDs(SoundMessageIDs.ActiveSoundPlayOrder, ActiveSoundEntry);
 
 			if (!ActiveSoundEntry.IsValid())
 			{
@@ -861,9 +908,9 @@ namespace UE::Audio::Insights
 		
 		// We cannot always guarentee the order of receiving start and stop messages
 		// If the stop message preceeded the start message, do not create a new entry
-		if (SoundsStoppedBeforeStart.Contains(Msg.PlayOrder))
+		if (SoundsStoppedBeforeStart.Contains(Msg.ActiveSoundPlayOrder))
 		{
-			SoundsStoppedBeforeStart.Remove(Msg.PlayOrder);
+			SoundsStoppedBeforeStart.Remove(Msg.ActiveSoundPlayOrder);
 			return;
 		}
 
@@ -878,34 +925,34 @@ namespace UE::Audio::Insights
 				CategoryEntry->bIsCategory = true;
 			}
 
-			TSharedPtr<IDashboardDataTreeViewEntry>* SoundEntry = CategoryEntry->Children.FindByPredicate([&Msg](TSharedPtr<IDashboardDataTreeViewEntry> ChildEntry)
+			TSharedPtr<IDashboardDataTreeViewEntry>* ActiveSoundEntry = CategoryEntry->Children.FindByPredicate([&Msg](TSharedPtr<IDashboardDataTreeViewEntry> ChildEntry)
 			{
 				FSoundDashboardEntry& NewEntry = FSoundTraceProviderPrivate::CastEntry(*ChildEntry.Get());
-				return NewEntry.PlayOrder == Msg.PlayOrder;
+				return NewEntry.ActiveSoundPlayOrder == Msg.ActiveSoundPlayOrder;
 			});
 
-			if (SoundEntry == nullptr || !SoundEntry->IsValid())
+			if (ActiveSoundEntry == nullptr || !ActiveSoundEntry->IsValid())
 			{
 				TSharedPtr<FSoundDashboardEntry> NewEntry = MakeShared<FSoundDashboardEntry>();
 				NewEntry->DeviceId = Msg.DeviceId;
-				NewEntry->PlayOrder = Msg.PlayOrder;
+				NewEntry->ActiveSoundPlayOrder = Msg.ActiveSoundPlayOrder;
 
 				CategoryEntry->Children.Add(NewEntry);
 
-				ActiveSoundToEntryKeysMap.Add(Msg.PlayOrder, SoundEntryKeys{ Msg.EntryType, Msg.DeviceId });
+				ActiveSoundToEntryKeysMap.Add(Msg.ActiveSoundPlayOrder, SoundEntryKeys{ Msg.EntryType, Msg.DeviceId });
 
 				OutReturnedSoundEntry = NewEntry;
 			}
 			else
 			{
-				TSharedPtr<FSoundDashboardEntry> NewEntry = StaticCastSharedPtr<FSoundDashboardEntry>(*SoundEntry);
+				TSharedPtr<FSoundDashboardEntry> NewEntry = StaticCastSharedPtr<FSoundDashboardEntry>(*ActiveSoundEntry);
 				if (NewEntry->TimeoutTimestamp != INVALID_TIMEOUT)
 				{
 					NewEntry->TimeoutTimestamp = INVALID_TIMEOUT;
 
 					EntriesTimingOut.RemoveAll([&Msg](const SoundMessageIDs& SoundMessageIDs)
 					{
-						return Msg.DeviceId == SoundMessageIDs.DeviceId && Msg.PlayOrder == SoundMessageIDs.PlayOrder;
+						return Msg.DeviceId == SoundMessageIDs.DeviceId && Msg.ActiveSoundPlayOrder == SoundMessageIDs.ActiveSoundPlayOrder;
 					});
 
 					// If this entry has SoundWaves, they will replay with new play order IDs, so we need to clear the old ones out
@@ -931,7 +978,7 @@ namespace UE::Audio::Insights
 	{
 		OutReturnedSoundEntry = nullptr;
 
-		GetActiveSoundEntryFromIDs(Msg.PlayOrder, OutReturnedSoundEntry);
+		GetActiveSoundEntryFromIDs(Msg.ActiveSoundPlayOrder, OutReturnedSoundEntry);
 		if (!OutReturnedSoundEntry.IsValid())
 		{
 			GetOrCreateActiveSoundEntry(Msg, OutReturnedSoundEntry);
@@ -962,7 +1009,7 @@ namespace UE::Audio::Insights
 			return;
 		}
 
-		const uint32 ActiveSoundPlayOrder = OutActiveSoundEntry->PlayOrder;
+		const uint32 ActiveSoundPlayOrder = OutActiveSoundEntry->ActiveSoundPlayOrder;
 		const SoundEntryKeys* SoundEntryKeys = ActiveSoundToEntryKeysMap.Find(ActiveSoundPlayOrder);
 
 		if (SoundEntryKeys == nullptr)
@@ -1146,7 +1193,7 @@ namespace UE::Audio::Insights
 
 			if (!OutActiveSoundEntry.bIsPlotActive)
 			{
-				OutActiveSoundEntry.bIsPlotActive = PlottingSoundEntries.Contains(OutActiveSoundEntry.PlayOrder);
+				OutActiveSoundEntry.bIsPlotActive = PlottingSoundEntries.Contains(OutActiveSoundEntry.GetEntryID());
 			}
 		};
 
@@ -1154,7 +1201,8 @@ namespace UE::Audio::Insights
 		{
 			TSharedPtr<FSoundDashboardEntry> ChildEntry = MakeShared<FSoundDashboardEntry>();
 			ChildEntry->DeviceId = Msg.DeviceId;
-			ChildEntry->PlayOrder = Msg.PlayOrder;
+			ChildEntry->WaveInstancePlayOrder = Msg.WaveInstancePlayOrder;
+			ChildEntry->ActiveSoundPlayOrder = Msg.ActiveSoundPlayOrder;
 			ChildEntry->Timestamp = Msg.Timestamp;
 			ChildEntry->SetName(*Msg.Name);
 			ChildEntry->ActorLabel = FText::FromString(*Msg.ActorLabel);
@@ -1163,7 +1211,7 @@ namespace UE::Audio::Insights
 
 			if (!ChildEntry->bIsPlotActive)
 			{
-				ChildEntry->bIsPlotActive = PlottingSoundEntries.Contains(ChildEntry->PlayOrder);
+				ChildEntry->bIsPlotActive = PlottingSoundEntries.Contains(ChildEntry->GetEntryID());
 			}
 
 			OutParentSoundEntry->Children.Add(ChildEntry);
@@ -1190,7 +1238,7 @@ namespace UE::Audio::Insights
 				return;
 			}
 
-			ActiveSoundIDs.Add(StartCachedMessage.PlayOrder);
+			ActiveSoundIDs.Add(StartCachedMessage.ActiveSoundPlayOrder);
 			InitActiveSoundEntry(StartCachedMessage, *ActiveSoundEntry);
 		};
 
@@ -1216,7 +1264,7 @@ namespace UE::Audio::Insights
 		// Collect all the is alive ping messages registered until this point in time
 		auto ProcessSoundIsAlivePingMessage = [this, &ActiveSoundIDs, &InitActiveSoundEntry](const FSoundIsAlivePingMessage& IsAlivePingCachedMessage)
 		{
-			ActiveSoundIDs.Add(IsAlivePingCachedMessage.PlayOrder);
+			ActiveSoundIDs.Add(IsAlivePingCachedMessage.ActiveSoundPlayOrder);
 
 			TSharedPtr<FSoundDashboardEntry> ActiveSoundEntry = nullptr;
 			CreateAndGetActiveSoundEntryIfNotFound(IsAlivePingCachedMessage, ActiveSoundEntry);
@@ -1257,11 +1305,11 @@ namespace UE::Audio::Insights
 				return;
 			}
 
-			ActiveSoundIDs.Add(SoundWaveStartCachedMessage.PlayOrder);
+			ActiveSoundIDs.Add(SoundWaveStartCachedMessage.ActiveSoundPlayOrder);
 
 			const bool bEntryExists = ParentSoundEntryReturn->Children.ContainsByPredicate([&SoundWaveStartCachedMessage](const TSharedPtr<IDashboardDataTreeViewEntry>& ChildEntry)
 			{
-				return ChildEntry.IsValid() && SoundWaveStartCachedMessage.PlayOrder == ChildEntry->GetEntryID();
+				return ChildEntry.IsValid() && SoundWaveStartCachedMessage.WaveInstancePlayOrder == CastEntry(*ChildEntry).WaveInstancePlayOrder;
 			});
 
 			if (bEntryExists)
@@ -1301,12 +1349,12 @@ namespace UE::Audio::Insights
 				return;
 			}
 
-			ActiveSoundIDs.Add(SoundWaveIsAlivePingCachedMessage.PlayOrder);
+			ActiveSoundIDs.Add(SoundWaveIsAlivePingCachedMessage.ActiveSoundPlayOrder);
 
 			// Avoid creating duplicate entries when receiving the IsAlivePing message
 			const bool bEntryExists = ParentSoundIsAlivePingEntryReturn->Children.ContainsByPredicate([&SoundWaveIsAlivePingCachedMessage](const TSharedPtr<IDashboardDataTreeViewEntry>& ChildEntry)
 			{
-				return ChildEntry.IsValid() && SoundWaveIsAlivePingCachedMessage.PlayOrder == ChildEntry->GetEntryID();
+				return ChildEntry.IsValid() && SoundWaveIsAlivePingCachedMessage.WaveInstancePlayOrder == CastEntry(*ChildEntry).WaveInstancePlayOrder;
 			});
 
 			if (bEntryExists)
@@ -1339,7 +1387,7 @@ namespace UE::Audio::Insights
 		auto ProcessSoundStopMessageMessage = [this](const FSoundStopMessage& StopCachedMessage)
 		{
 			TSharedPtr<FSoundDashboardEntry> ActiveSoundEntry = nullptr;
-			GetActiveSoundEntryFromIDs(StopCachedMessage.PlayOrder, ActiveSoundEntry);
+			GetActiveSoundEntryFromIDs(StopCachedMessage.ActiveSoundPlayOrder, ActiveSoundEntry);
 
 			if (ActiveSoundEntry.IsValid() && ActiveSoundEntry->Timestamp < StopCachedMessage.Timestamp)
 			{
@@ -1393,10 +1441,7 @@ namespace UE::Audio::Insights
 #if WITH_EDITOR
 		CollectParamsForTimestamp(CacheManager, TimeMarker);
 #else
-		// Regenerate the active sounds play order array
-		ActiveSoundToEntryKeysMap.GetKeys(ActiveSoundPlayOrderArray);
-
-		CollectParamsForTimestampStandalone(ActiveSoundPlayOrderArray, TimeMarker);
+		CollectParamsForTimestampStandalone(TimeMarker);
 #endif //!WITH_EDITOR
 
 		// Call parent method to update LastMessageId
@@ -1461,7 +1506,7 @@ namespace UE::Audio::Insights
 					
 					if (ActiveSoundEntry.IsValid())
 					{
-						if (FSoundDashboardEntry* SoundEntry = FindSoundEntryForActiveSound(Message.PlayOrder, &ActiveSoundEntry))
+						if (FSoundDashboardEntry* SoundEntry = FindWaveEntryInActiveSound(Message.WaveInstancePlayOrder, &ActiveSoundEntry))
 						{
 							SoundEntry->PriorityDataRange.Push({ Message.Timestamp, Message.Priority });
 							SoundEntry->Timestamp = Message.Timestamp;
@@ -1483,7 +1528,7 @@ namespace UE::Audio::Insights
 
 					if (ActiveSoundEntry.IsValid())
 					{
-						if (FSoundDashboardEntry* SoundEntry = FindSoundEntryForActiveSound(Message.PlayOrder, &ActiveSoundEntry))
+						if (FSoundDashboardEntry* SoundEntry = FindWaveEntryInActiveSound(Message.WaveInstancePlayOrder, &ActiveSoundEntry))
 						{
 							SoundEntry->DistanceDataRange.Push({ Message.Timestamp, Message.Distance });
 							SoundEntry->Timestamp = Message.Timestamp;
@@ -1505,7 +1550,7 @@ namespace UE::Audio::Insights
 
 					if (ActiveSoundEntry.IsValid())
 					{
-						if (FSoundDashboardEntry* SoundEntry = FindSoundEntryForActiveSound(Message.PlayOrder, &ActiveSoundEntry))
+						if (FSoundDashboardEntry* SoundEntry = FindWaveEntryInActiveSound(Message.WaveInstancePlayOrder, &ActiveSoundEntry))
 						{
 							SoundEntry->DistanceAttenuationDataRange.Push({ Message.Timestamp, Message.DistanceAttenuation });
 							SoundEntry->Timestamp = Message.Timestamp;
@@ -1527,7 +1572,7 @@ namespace UE::Audio::Insights
 
 					if (ActiveSoundEntry.IsValid())
 					{
-						if (FSoundDashboardEntry* SoundEntry = FindSoundEntryForActiveSound(Message.PlayOrder, &ActiveSoundEntry))
+						if (FSoundDashboardEntry* SoundEntry = FindWaveEntryInActiveSound(Message.WaveInstancePlayOrder, &ActiveSoundEntry))
 						{
 							SoundEntry->LPFFreqDataRange.Push({ Message.Timestamp, Message.LPFFrequency });
 							SoundEntry->Timestamp = Message.Timestamp;
@@ -1549,7 +1594,7 @@ namespace UE::Audio::Insights
 
 					if (ActiveSoundEntry.IsValid())
 					{
-						if (FSoundDashboardEntry* SoundEntry = FindSoundEntryForActiveSound(Message.PlayOrder, &ActiveSoundEntry))
+						if (FSoundDashboardEntry* SoundEntry = FindWaveEntryInActiveSound(Message.WaveInstancePlayOrder, &ActiveSoundEntry))
 						{
 							SoundEntry->HPFFreqDataRange.Push({ Message.Timestamp, Message.HPFFrequency });
 							SoundEntry->Timestamp = Message.Timestamp;
@@ -1571,7 +1616,7 @@ namespace UE::Audio::Insights
 
 					if (ActiveSoundEntry.IsValid())
 					{
-						if (FSoundDashboardEntry* SoundEntry = FindSoundEntryForActiveSound(Message.PlayOrder, &ActiveSoundEntry))
+						if (FSoundDashboardEntry* SoundEntry = FindWaveEntryInActiveSound(Message.WaveInstancePlayOrder, &ActiveSoundEntry))
 						{
 							SoundEntry->AmplitudeDataRange.Push({ Message.Timestamp, Message.Envelope });
 							SoundEntry->Timestamp = Message.Timestamp;
@@ -1593,7 +1638,7 @@ namespace UE::Audio::Insights
 
 					if (ActiveSoundEntry.IsValid())
 					{
-						if (FSoundDashboardEntry* SoundEntry = FindSoundEntryForActiveSound(Message.PlayOrder, &ActiveSoundEntry))
+						if (FSoundDashboardEntry* SoundEntry = FindWaveEntryInActiveSound(Message.WaveInstancePlayOrder, &ActiveSoundEntry))
 						{
 							SoundEntry->VolumeDataRange.Push({ Message.Timestamp, Message.Volume });
 							SoundEntry->Timestamp = Message.Timestamp;
@@ -1615,7 +1660,7 @@ namespace UE::Audio::Insights
 
 					if (ActiveSoundEntry.IsValid())
 					{
-						if (FSoundDashboardEntry* SoundEntry = FindSoundEntryForActiveSound(Message.PlayOrder, &ActiveSoundEntry))
+						if (FSoundDashboardEntry* SoundEntry = FindWaveEntryInActiveSound(Message.WaveInstancePlayOrder, &ActiveSoundEntry))
 						{
 							SoundEntry->PitchDataRange.Push({ Message.Timestamp, Message.Pitch });
 							SoundEntry->Timestamp = Message.Timestamp;
@@ -1637,7 +1682,7 @@ namespace UE::Audio::Insights
 
 					if (ActiveSoundEntry.IsValid())
 					{
-						if (FSoundDashboardEntry* SoundEntry = FindSoundEntryForActiveSound(Message.PlayOrder, &ActiveSoundEntry))
+						if (FSoundDashboardEntry* SoundEntry = FindWaveEntryInActiveSound(Message.WaveInstancePlayOrder, &ActiveSoundEntry))
 						{
 							SoundEntry->RelativeRenderCostDataRange.Push({ Message.Timestamp, Message.RelativeRenderCost });
 							SoundEntry->Timestamp = Message.Timestamp;
@@ -1661,13 +1706,20 @@ namespace UE::Audio::Insights
 		OnProcessPlotData.Broadcast();
 	}
 #else
-	void FSoundTraceProvider::CollectParamsForTimestampStandalone(const TArray<uint32>& ActiveSoundPlayOrderArray, const double TimeMarker)
+	void FSoundTraceProvider::CollectParamsForTimestampStandalone(const double TimeMarker)
 	{
 		using namespace FSoundTraceProviderPrivate;
 
 		const FDeviceData* DeviceData = FindFilteredDeviceData();
 		if (DeviceData)
 		{
+			TArray<PlayOrderPair> PlayOrders;
+
+			for (const auto& [EntryType, SoundDashboardEntry] : *DeviceData)
+			{
+				CollectPlayOrderPairsRecursive(SoundDashboardEntry, PlayOrders);
+			}
+
 			// Collect messages (volume, pitch, etc.) from active sounds (based on active sounds's PlayOrder)
 			struct CachedEntryInfo
 			{
@@ -1683,72 +1735,72 @@ namespace UE::Audio::Insights
 			};
 
 			TArray<CachedEntryInfo> CachedEntryInfos;
-			CachedEntryInfos.SetNum(ActiveSoundPlayOrderArray.Num());
+			CachedEntryInfos.SetNum(PlayOrders.Num());
 
 			// Using ParallelFor to speed-up the cached messages retrieval, using a traditional for loop is unacceptably slower, specially in large traces.
-			ParallelFor(ActiveSoundPlayOrderArray.Num(),
-			[&ActiveSoundPlayOrderArray, &CachedEntryInfos, TimeMarker, this](const int32 Index)
+			ParallelFor(PlayOrders.Num(),
+			[&PlayOrders, &CachedEntryInfos, TimeMarker, this](const int32 Index)
 			{
-				const uint32 ActiveSoundPlayOrder = ActiveSoundPlayOrderArray[Index];
+				const PlayOrderPair PlayOrderPair = PlayOrders[Index];
 
 				// Priority
-				const FSoundPriorityMessage* FoundPriorityCachedMessage = FindClosestMessageToTimestamp(SessionCachedMessages->PriorityCachedMessages, TimeMarker, ActiveSoundPlayOrder);
+				const FSoundPriorityMessage* FoundPriorityCachedMessage = FindClosestMessageToTimestamp(SessionCachedMessages->PriorityCachedMessages, TimeMarker, PlayOrderPair);
 				if (FoundPriorityCachedMessage)
 				{
 					CachedEntryInfos[Index].PriorityMessage = *FoundPriorityCachedMessage;
 				}
 
 				// Distance
-				const FSoundDistanceMessage* FoundDistanceCachedMessage = FindClosestMessageToTimestamp(SessionCachedMessages->DistanceCachedMessages, TimeMarker, ActiveSoundPlayOrder);
+				const FSoundDistanceMessage* FoundDistanceCachedMessage = FindClosestMessageToTimestamp(SessionCachedMessages->DistanceCachedMessages, TimeMarker, PlayOrderPair);
 				if (FoundDistanceCachedMessage)
 				{
 					CachedEntryInfos[Index].DistanceMessage = *FoundDistanceCachedMessage;
 				}
 
 				// Distance Attenuation
-				const FSoundDistanceAttenuationMessage* FoundDistanceAttenuationCachedMessage = FindClosestMessageToTimestamp(SessionCachedMessages->DistanceAttenuationCachedMessages, TimeMarker, ActiveSoundPlayOrder);
+				const FSoundDistanceAttenuationMessage* FoundDistanceAttenuationCachedMessage = FindClosestMessageToTimestamp(SessionCachedMessages->DistanceAttenuationCachedMessages, TimeMarker, PlayOrderPair);
 				if (FoundDistanceAttenuationCachedMessage)
 				{
 					CachedEntryInfos[Index].DistanceAttenuationMessage = *FoundDistanceAttenuationCachedMessage;
 				}
 
 				// LPF Freq
-				const FSoundLPFFreqMessage* FoundLPFFreqCachedMessage = FindClosestMessageToTimestamp(SessionCachedMessages->LPFFreqCachedMessages, TimeMarker, ActiveSoundPlayOrder);
+				const FSoundLPFFreqMessage* FoundLPFFreqCachedMessage = FindClosestMessageToTimestamp(SessionCachedMessages->LPFFreqCachedMessages, TimeMarker, PlayOrderPair);
 				if (FoundLPFFreqCachedMessage)
 				{
 					CachedEntryInfos[Index].LPFFreqMessage = *FoundLPFFreqCachedMessage;
 				}
 
 				// HPF Freq
-				const FSoundHPFFreqMessage* FoundHPFFreqCachedMessage = FindClosestMessageToTimestamp(SessionCachedMessages->HPFFreqCachedMessages, TimeMarker, ActiveSoundPlayOrder);
+				const FSoundHPFFreqMessage* FoundHPFFreqCachedMessage = FindClosestMessageToTimestamp(SessionCachedMessages->HPFFreqCachedMessages, TimeMarker, PlayOrderPair);
 				if (FoundHPFFreqCachedMessage)
 				{
 					CachedEntryInfos[Index].HPFFreqMessage = *FoundHPFFreqCachedMessage;
 				}
 
 				// Amplitude
-				const FSoundEnvelopeMessage* FoundAmplitudeCachedMessage = FindClosestMessageToTimestamp(SessionCachedMessages->AmplitudeCachedMessages, TimeMarker, ActiveSoundPlayOrder);
+				const FSoundEnvelopeMessage* FoundAmplitudeCachedMessage = FindClosestMessageToTimestamp(SessionCachedMessages->AmplitudeCachedMessages, TimeMarker, PlayOrderPair);
 				if (FoundAmplitudeCachedMessage)
 				{
 					CachedEntryInfos[Index].AmplitudeMessage = *FoundAmplitudeCachedMessage;
 				}
 
 				// Volume
-				const FSoundVolumeMessage* FoundVolumeCachedMessage = FindClosestMessageToTimestamp(SessionCachedMessages->VolumeCachedMessages, TimeMarker, ActiveSoundPlayOrder);
+				const FSoundVolumeMessage* FoundVolumeCachedMessage = FindClosestMessageToTimestamp(SessionCachedMessages->VolumeCachedMessages, TimeMarker, PlayOrderPair);
 				if (FoundVolumeCachedMessage)
 				{
 					CachedEntryInfos[Index].VolumeMessage = *FoundVolumeCachedMessage;
 				}
 
 				// Pitch
-				const FSoundPitchMessage* FoundPitchCachedMessage = FindClosestMessageToTimestamp(SessionCachedMessages->PitchCachedMessages, TimeMarker, ActiveSoundPlayOrder);
+				const FSoundPitchMessage* FoundPitchCachedMessage = FindClosestMessageToTimestamp(SessionCachedMessages->PitchCachedMessages, TimeMarker, PlayOrderPair);
 				if (FoundPitchCachedMessage)
 				{
 					CachedEntryInfos[Index].PitchMessage = *FoundPitchCachedMessage;
 				}
 
 				// Relative Render Cost
-				const FSoundRelativeRenderCostMessage* FoundRelativeRenderCostCachedMessage = FindClosestMessageToTimestamp(SessionCachedMessages->RelativeRenderCostCachedMessages, TimeMarker, ActiveSoundPlayOrder);
+				const FSoundRelativeRenderCostMessage* FoundRelativeRenderCostCachedMessage = FindClosestMessageToTimestamp(SessionCachedMessages->RelativeRenderCostCachedMessages, TimeMarker, PlayOrderPair);
 				if (FoundRelativeRenderCostCachedMessage)
 				{
 					CachedEntryInfos[Index].RelativeRenderCostMessage = *FoundRelativeRenderCostCachedMessage;
@@ -1771,7 +1823,7 @@ namespace UE::Audio::Insights
 					continue;
 				}
 
-				FSoundDashboardEntry* SoundEntry = FSoundTraceProviderPrivate::FindSoundEntryForActiveSound(CachedEntryInfo.PriorityMessage.GetValue().PlayOrder, &ActiveSoundEntry);
+				FSoundDashboardEntry* SoundEntry = FSoundTraceProviderPrivate::FindWaveEntryInActiveSound(CachedEntryInfo.PriorityMessage.GetValue().WaveInstancePlayOrder, &ActiveSoundEntry);
 				if (SoundEntry == nullptr)
 				{
 					continue;
@@ -1839,11 +1891,11 @@ namespace UE::Audio::Insights
 		// Cache the bIsPlotActive flag in case this sound returns in future
 		if (SoundEntry.bIsPlotActive)
 		{
-			PlottingSoundEntries.FindOrAdd(SoundEntry.PlayOrder);
+			PlottingSoundEntries.FindOrAdd(SoundEntry.GetEntryID());
 		}
-		else if (PlottingSoundEntries.Contains(SoundEntry.PlayOrder))
+		else if (PlottingSoundEntries.Contains(SoundEntry.GetEntryID()))
 		{
-			PlottingSoundEntries.Remove(SoundEntry.PlayOrder);
+			PlottingSoundEntries.Remove(SoundEntry.GetEntryID());
 		}
 
 		for (const TSharedPtr<IDashboardDataTreeViewEntry>& Child : SoundEntry.Children)
diff --git a/Engine/Plugins/AudioInsights/Source/AudioInsights/Private/Providers/SoundTraceProvider.h b/Engine/Plugins/AudioInsights/Source/AudioInsights/Private/Providers/SoundTraceProvider.h
--- a/Engine/Plugins/AudioInsights/Source/AudioInsights/Private/Providers/SoundTraceProvider.h
+++ b/Engine/Plugins/AudioInsights/Source/AudioInsights/Private/Providers/SoundTraceProvider.h
@@ -55,7 +55,8 @@ namespace UE::Audio::Insights
 #if WITH_EDITOR
 		void CollectParamsForTimestamp(const FAudioInsightsCacheManager& CacheManager, const double TimeMarker);
 #else
-		void CollectParamsForTimestampStandalone(const TArray<uint32>& ActiveSoundPlayOrderArray, const double TimeMarker);
+		void CollectParamsForTimestampStandalone(const double TimeMarker);
+		void CollectParameterPlayOrderIDsRecursive(const TSharedPtr<IDashboardDataTreeViewEntry>& Entry, TArray<uint32>& OutPlayOrderArray);
 #endif // WITH_EDITOR
 
 		void CacheIsPlottingFlagRecursive(const TSharedPtr<IDashboardDataTreeViewEntry>& Entry);
@@ -76,13 +77,13 @@ namespace UE::Audio::Insights
 		struct SoundMessageIDs
 		{
 			const ::Audio::FDeviceId DeviceId;
-			const uint32 PlayOrder;
+			const uint32 ActiveSoundPlayOrder;
 		};
 
 		TMap<uint32, SoundEntryKeys> ActiveSoundToEntryKeysMap;
 		TArray<SoundMessageIDs> EntriesTimingOut;
 		TSet<uint32> SoundsStoppedBeforeStart;
-		TSet<uint32> PlottingSoundEntries;
+		TSet<uint64> PlottingSoundEntries;
 
 		double DashboardTimeoutTime = 3.0;
 
diff --git a/Engine/Plugins/AudioInsights/Source/AudioInsights/Private/Views/SAudioCurveView.cpp b/Engine/Plugins/AudioInsights/Source/AudioInsights/Private/Views/SAudioCurveView.cpp
--- a/Engine/Plugins/AudioInsights/Source/AudioInsights/Private/Views/SAudioCurveView.cpp
+++ b/Engine/Plugins/AudioInsights/Source/AudioInsights/Private/Views/SAudioCurveView.cpp
@@ -169,7 +169,7 @@ void SAudioCurveView::UpdateCurveToolTip(const FGeometry& InMyGeometry, const FP
 	const double TargetX = RangeToScreen.LocalXToInput(HitPosition.X);
 
 	// Keep track of closest curve index, closest point, value 
-	int32 ClosestCurveId = INDEX_NONE;
+	uint64 ClosestCurveId = INDEX_NONE;
 	FCurvePoint ClosestPoint;
 	float ClosestDistance = TNumericLimits<float>::Max();
 	float ThickestPlot = 0.0f;
@@ -334,7 +334,7 @@ void SAudioCurveView::UpdateYDataRangeFromTimestampRange(const double InLowerBou
 }
 #endif // !WITH_EDITOR
 
-void SAudioCurveView::SetCurvesPointData(TSharedPtr<TMap<int32, TArray<FCurvePoint>>> InPointDataPerCurve)
+void SAudioCurveView::SetCurvesPointData(TSharedPtr<TMap<uint64, TArray<FCurvePoint>>> InPointDataPerCurve)
 {
 	PointDataPerCurve = InPointDataPerCurve;
 
@@ -344,7 +344,7 @@ void SAudioCurveView::SetCurvesPointData(TSharedPtr<TMap<int32, TArray<FCurvePoi
 	}
 }
 
-void SAudioCurveView::SetCurvesMetadata(TSharedPtr<TMap<int32, FCurveMetadata>> InMetadataPerCurve)
+void SAudioCurveView::SetCurvesMetadata(TSharedPtr<TMap<uint64, FCurveMetadata>> InMetadataPerCurve)
 {
 	MetadataPerCurve = InMetadataPerCurve;
 }
diff --git a/Engine/Plugins/AudioInsights/Source/AudioInsights/Private/Views/SoundDashboardViewFactory.cpp b/Engine/Plugins/AudioInsights/Source/AudioInsights/Private/Views/SoundDashboardViewFactory.cpp
--- a/Engine/Plugins/AudioInsights/Source/AudioInsights/Private/Views/SoundDashboardViewFactory.cpp
+++ b/Engine/Plugins/AudioInsights/Source/AudioInsights/Private/Views/SoundDashboardViewFactory.cpp
@@ -521,7 +521,7 @@ namespace UE::Audio::Insights
 
 		const FName PlayOrderColumnName = "PlayOrder";
 		const FText PlayOrderColumnDisplayName = LOCTEXT("SoundDashboard_PlayOrderColumnDisplayName", "Play Order");
-		const FText PlayOrderColumnTooltip = LOCTEXT("SoundDashboard_PlayOrderColumnTooltip", "The order in which a sound was played.");
+		const FText PlayOrderColumnTooltip = LOCTEXT("SoundDashboard_PlayOrderColumnTooltip", "The order in which a sound was played.\nFirst Num = Active Sound Play Order, Second Num = Wave Instance Play Order.");
 
 		const FName PriorityColumnName = "Priority";
 		const FText PriorityColumnDisplayName = LOCTEXT("SoundDashboard_PriorityColumnDisplayName", "Priority");
@@ -529,7 +529,7 @@ namespace UE::Audio::Insights
 
 		const FName DistanceColumnName = "Distance";
 		const FText DistanceColumnDisplayName = LOCTEXT("SoundDashboard_DistanceColumnDisplayName", "Distance");
-		const FText DistanceColumnTooltip = LOCTEXT("SoundDashboard_DistanceColumnTooltip", "The Distance from the sound source to the listener.");
+		const FText DistanceColumnTooltip = LOCTEXT("SoundDashboard_DistanceColumnTooltip", "The distance from the sound source to the listener.");
 
 		const FName DistanceAttenuationColumnName = "DistanceAttenuation";
 		const FText DistanceAttenuationColumnDisplayName = LOCTEXT("SoundDashboard_DistanceAttenuationColumnDisplayName", "Distance/Occlusion Attenuation");
@@ -1989,7 +1989,12 @@ namespace UE::Audio::Insights
 							return FText::GetEmpty();
 						}
 
-						return FText::AsNumber(SoundDashboardEntry.PlayOrder);
+						if (SoundDashboardEntry.WaveInstancePlayOrder == INDEX_NONE)
+						{
+							return FText::AsNumber(SoundDashboardEntry.ActiveSoundPlayOrder);
+						}
+
+						return FText::Format(LOCTEXT("AudioDashboard_Sounds_WaveInstancePlayOrderFormat", "{0}, {1}"), SoundDashboardEntry.ActiveSoundPlayOrder, SoundDashboardEntry.WaveInstancePlayOrder);
 					}
 				}
 			},
@@ -2176,7 +2181,12 @@ namespace UE::Audio::Insights
 
 		auto SortByPlayOrder = [](const FSoundDashboardEntry& First, const FSoundDashboardEntry& Second)
 		{
-			return First.PlayOrder < Second.PlayOrder;
+			if (First.ActiveSoundPlayOrder == Second.ActiveSoundPlayOrder)
+			{
+				return First.WaveInstancePlayOrder < Second.WaveInstancePlayOrder;
+			}
+
+			return First.ActiveSoundPlayOrder < Second.ActiveSoundPlayOrder;
 		};
 		
 		auto SortByName = [&SortByPlayOrder](const FSoundDashboardEntry& First, const FSoundDashboardEntry& Second)
@@ -3532,7 +3542,7 @@ namespace UE::Audio::Insights
 					return OriginalChildSoundEntry.EntryType == SelectedSoundEntry.EntryType;
 				}
 				
-				return OriginalChildSoundEntry.PlayOrder == SelectedSoundEntry.PlayOrder;
+				return OriginalChildSoundEntry.GetEntryID() == SelectedSoundEntry.GetEntryID();
 			});
 
 			if (FoundMatchingSelectedEntry != nullptr)
@@ -3678,7 +3688,7 @@ namespace UE::Audio::Insights
 						}
 
 						const FSoundDashboardEntry& ExistingPinnedEntry = CastEntry(*PinnedEntry);
-						return ExistingPinnedEntry.PlayOrder == SoundEntry.PlayOrder;
+						return ExistingPinnedEntry.GetEntryID() == SoundEntry.GetEntryID();
 					});
 
 					// If we didn't find this entry already inside the pinned area, add it here
diff --git a/Engine/Plugins/AudioInsights/Source/AudioInsights/Private/Views/SoundPlotsWidgetView.cpp b/Engine/Plugins/AudioInsights/Source/AudioInsights/Private/Views/SoundPlotsWidgetView.cpp
--- a/Engine/Plugins/AudioInsights/Source/AudioInsights/Private/Views/SoundPlotsWidgetView.cpp
+++ b/Engine/Plugins/AudioInsights/Source/AudioInsights/Private/Views/SoundPlotsWidgetView.cpp
@@ -469,7 +469,7 @@ namespace UE::Audio::Insights
 		// Initialize curve data and metadata
 		if (!PlotWidgetMetadataPerCurve.IsValid())
 		{
-			PlotWidgetMetadataPerCurve = MakeShared<TMap<int32, SAudioCurveView::FCurveMetadata>>();
+			PlotWidgetMetadataPerCurve = MakeShared<TMap<uint64, SAudioCurveView::FCurveMetadata>>();
 			for (const FName& ColumnName : ColumnNames)
 			{
 				TSharedPtr<FPointDataPerCurveMap> PointDataPerCurveMap = MakeShared<FPointDataPerCurveMap>();
@@ -545,7 +545,7 @@ namespace UE::Audio::Insights
 			return;
 		}
 
-		const uint32 SourceId = Entry->GetEntryID();
+		const uint64 SourceId = Entry->GetEntryID();
 		const bool bPlotIsEnabled = GetIsPlotEnabledForEntry(*Entry);
 
 		if (SourceId == INDEX_NONE || !bPlotIsEnabled)
diff --git a/Engine/Plugins/AudioInsights/Source/AudioInsights/Public/AudioInsightsDataSource.h b/Engine/Plugins/AudioInsights/Source/AudioInsights/Public/AudioInsightsDataSource.h
--- a/Engine/Plugins/AudioInsights/Source/AudioInsights/Public/AudioInsightsDataSource.h
+++ b/Engine/Plugins/AudioInsights/Source/AudioInsights/Public/AudioInsightsDataSource.h
@@ -22,7 +22,7 @@ namespace UE::Audio::Insights
 		virtual ~IDashboardDataTreeViewEntry() = default;
 
 		virtual bool IsValid() const = 0;
-		virtual int32 GetEntryID() const = 0;
+		virtual uint64 GetEntryID() const = 0;
 		virtual const FText& GetDisplayName() const = 0;
 		virtual const FLinearColor& GetEntryColor() const = 0;
 		virtual bool HasSetInitExpansion() const = 0;
diff --git a/Engine/Plugins/AudioInsights/Source/AudioInsights/Public/Cache/AudioCachedMessageChunk.h b/Engine/Plugins/AudioInsights/Source/AudioInsights/Public/Cache/AudioCachedMessageChunk.h
--- a/Engine/Plugins/AudioInsights/Source/AudioInsights/Public/Cache/AudioCachedMessageChunk.h
+++ b/Engine/Plugins/AudioInsights/Source/AudioInsights/Public/Cache/AudioCachedMessageChunk.h
@@ -39,7 +39,7 @@ namespace UE::Audio::Insights
 		// Will iterate backwards from param Timestamp until the first message of type MessageID is found
 		// Param ID can be overriden to filter the results further
 		template<TIsCacheableMessage T>
-		T* FindClosestMessage(const FName& MessageID, const double Timestamp, const uint32 ID = INDEX_NONE) const
+		T* FindClosestMessage(const FName& MessageID, const double Timestamp, const uint64 ID = INDEX_NONE) const
 		{
 			const TSharedPtr<TCachePagedArray>* MessageCache = Messages.Find(MessageID);
 			if (MessageCache == nullptr || !MessageCache->IsValid() || (*MessageCache)->Num() == 0u)
diff --git a/Engine/Plugins/AudioInsights/Source/AudioInsights/Public/Cache/AudioInsightsCacheManager.h b/Engine/Plugins/AudioInsights/Source/AudioInsights/Public/Cache/AudioInsightsCacheManager.h
--- a/Engine/Plugins/AudioInsights/Source/AudioInsights/Public/Cache/AudioInsightsCacheManager.h
+++ b/Engine/Plugins/AudioInsights/Source/AudioInsights/Public/Cache/AudioInsightsCacheManager.h
@@ -51,7 +51,7 @@ namespace UE::Audio::Insights
 		// @param ID: When overloaded, the search algorithm will only return messages that have a matching ID.
 		//			  If left at INDEX_NONE, the search algorithm will return the closest message irrespective of ID
 		template<TIsCacheableMessage T>
-		const T* FindClosestMessage(const FName& MessageID, const double Timestamp, const uint32 ID = INDEX_NONE) const
+		const T* FindClosestMessage(const FName& MessageID, const double Timestamp, const uint64 ID = INDEX_NONE) const
 		{
 			uint32 ReadIndex = INDEX_NONE;
 			for (uint32 Index = 0u; Index < Cache.Capacity(); ++Index)
@@ -97,7 +97,7 @@ namespace UE::Audio::Insights
 		// @param ID: When overloaded, the search algorithm will only return messages that have a matching ID.
 		//			  If left at INDEX_NONE, the search algorithm will return all matching messages irrespective of ID
 		template<TIsCacheableMessage T>
-		void IterateOverRange(const FName& MessageID, const double Start, const double End, TFunctionRef<void(const T& Message)> OnMessageFunc, const uint32 ID = INDEX_NONE) const
+		void IterateOverRange(const FName& MessageID, const double Start, const double End, TFunctionRef<void(const T& Message)> OnMessageFunc, const uint64 ID = INDEX_NONE) const
 		{
 			uint32 ReadIndex = StartIndex;
 			double LatestChunkStartTimestamp = -1.0;
@@ -167,7 +167,7 @@ namespace UE::Audio::Insights
 		// Iterates over the cache from the start to a chosen timestamp for a specific message ID
 		// Will stop iterating if the end of the cache is reached
 		template<TIsCacheableMessage T>
-		void IterateTo(const FName& MessageID, const double End, TFunctionRef<void(const T& Message)> OnMessageFunc, const uint32 ID = INDEX_NONE) const
+		void IterateTo(const FName& MessageID, const double End, TFunctionRef<void(const T& Message)> OnMessageFunc, const uint64 ID = INDEX_NONE) const
 		{
 			constexpr double Start = 0.0;
 			IterateOverRange<T>(MessageID, Start, End, OnMessageFunc, ID);
@@ -175,14 +175,14 @@ namespace UE::Audio::Insights
 
 		// Iterates over the cache from a chosen timestamp to the end of the cache for a specific message ID
 		template<TIsCacheableMessage T>
-		void IterateFrom(const FName& MessageID, const double Start, TFunctionRef<void(const T& Message)> OnMessageFunc, const uint32 ID = INDEX_NONE) const
+		void IterateFrom(const FName& MessageID, const double Start, TFunctionRef<void(const T& Message)> OnMessageFunc, const uint64 ID = INDEX_NONE) const
 		{
 			IterateOverRange<T>(MessageID, Start, GetCacheEndTimeStamp(), OnMessageFunc, ID);
 		}
 
 		// Iterates over the entire cache for a specific message ID
 		template<TIsCacheableMessage T>
-		void IterateOverAll(const FName& MessageID, TFunctionRef<void(const T& Message)> OnMessageFunc, const uint32 ID = INDEX_NONE) const
+		void IterateOverAll(const FName& MessageID, TFunctionRef<void(const T& Message)> OnMessageFunc, const uint64 ID = INDEX_NONE) const
 		{
 			constexpr double Start = 0.0;
 			IterateOverRange<T>(MessageID, Start, GetCacheEndTimeStamp(), OnMessageFunc, ID);
diff --git a/Engine/Plugins/AudioInsights/Source/AudioInsights/Public/Cache/IAudioCachedMessage.h b/Engine/Plugins/AudioInsights/Source/AudioInsights/Public/Cache/IAudioCachedMessage.h
--- a/Engine/Plugins/AudioInsights/Source/AudioInsights/Public/Cache/IAudioCachedMessage.h
+++ b/Engine/Plugins/AudioInsights/Source/AudioInsights/Public/Cache/IAudioCachedMessage.h
@@ -12,7 +12,7 @@ namespace UE::Audio::Insights
 		virtual ~IAudioCachedMessage() = default;
 
 		virtual uint32 GetSizeOf() const = 0;
-		virtual uint32 GetID() const = 0;
+		virtual uint64 GetID() const = 0;
 		virtual const FName GetMessageName() const = 0;
 
 		double Timestamp = 0.0;
diff --git a/Engine/Plugins/AudioInsights/Source/AudioInsights/Public/Messages/AudioBusTraceMessages.h b/Engine/Plugins/AudioInsights/Source/AudioInsights/Public/Messages/AudioBusTraceMessages.h
--- a/Engine/Plugins/AudioInsights/Source/AudioInsights/Public/Messages/AudioBusTraceMessages.h
+++ b/Engine/Plugins/AudioInsights/Source/AudioInsights/Public/Messages/AudioBusTraceMessages.h
@@ -35,7 +35,7 @@ namespace UE::Audio::Insights
 		FAudioBusMessageBase() = default;
 		FAudioBusMessageBase(const UE::Trace::IAnalyzer::FOnEventContext& InContext);
 
-		virtual uint32 GetID() const override { return AudioBusId; }
+		virtual uint64 GetID() const override { return AudioBusId; }
 
 		::Audio::FDeviceId DeviceId = INDEX_NONE;
 		uint32 AudioBusId = INDEX_NONE;
diff --git a/Engine/Plugins/AudioInsights/Source/AudioInsights/Public/Messages/AudioEventLogTraceMessages.h b/Engine/Plugins/AudioInsights/Source/AudioInsights/Public/Messages/AudioEventLogTraceMessages.h
--- a/Engine/Plugins/AudioInsights/Source/AudioInsights/Public/Messages/AudioEventLogTraceMessages.h
+++ b/Engine/Plugins/AudioInsights/Source/AudioInsights/Public/Messages/AudioEventLogTraceMessages.h
@@ -39,7 +39,7 @@ namespace UE::Audio::Insights
 		FAudioEventLogMessage(const Trace::IAnalyzer::FOnEventContext& InContext);
 
 		virtual uint32 GetSizeOf() const override;
-		virtual uint32 GetID() const override { return PlayOrder; }
+		virtual uint64 GetID() const override { return PlayOrder; }
 		virtual const FName GetMessageName() const override;
 
 		::Audio::FDeviceId DeviceId = INDEX_NONE;
diff --git a/Engine/Plugins/AudioInsights/Source/AudioInsights/Public/Messages/SoundTraceMessages.h b/Engine/Plugins/AudioInsights/Source/AudioInsights/Public/Messages/SoundTraceMessages.h
--- a/Engine/Plugins/AudioInsights/Source/AudioInsights/Public/Messages/SoundTraceMessages.h
+++ b/Engine/Plugins/AudioInsights/Source/AudioInsights/Public/Messages/SoundTraceMessages.h
@@ -57,16 +57,18 @@ namespace UE::Audio::Insights
 		extern const FName RelativeRenderCostParam;
 	};
 
+	namespace SoundMessageUtils
+	{
+		uint64 GeneratePlayOrderUniqueID(const uint32 ActiveSoundPlayOrder, const uint32 WaveInstancePlayOrder);
+	};
+
 	// Trace messages
 	struct FSoundMessageBase : public IAudioCachedMessage
 	{
 		FSoundMessageBase() = default;
 		FSoundMessageBase(const Trace::IAnalyzer::FOnEventContext& InContext);
 
-		virtual uint32 GetID() const override { return PlayOrder; }
-
 		::Audio::FDeviceId DeviceId = INDEX_NONE;
-		uint32 PlayOrder = INDEX_NONE;
 	};
 
 	struct FSoundParameterMessage : public FSoundMessageBase
@@ -76,11 +78,13 @@ namespace UE::Audio::Insights
 		FSoundParameterMessage(const Trace::IAnalyzer::FOnEventContext& InContext)
 			: FSoundMessageBase(InContext)
 		{
+			WaveInstancePlayOrder = InContext.EventData.GetValue<uint32>("PlayOrder");
 			ActiveSoundPlayOrder = InContext.EventData.GetValue<uint32>("ActiveSoundPlayOrder");
 		}
 
-		virtual uint32 GetID() const override { return ActiveSoundPlayOrder; }
+		virtual uint64 GetID() const override { return SoundMessageUtils::GeneratePlayOrderUniqueID(ActiveSoundPlayOrder, WaveInstancePlayOrder); }
 
+		uint32 WaveInstancePlayOrder = INDEX_NONE;
 		uint32 ActiveSoundPlayOrder = INDEX_NONE;
 	};
 
@@ -91,10 +95,12 @@ namespace UE::Audio::Insights
 
 		virtual uint32 GetSizeOf() const override;
 		virtual const FName GetMessageName() const override { return SoundMessageNames::SoundStart; }
+		virtual uint64 GetID() const override { return ActiveSoundPlayOrder; }
 
 		FString Name;
 		ESoundDashboardEntryType EntryType;
 		FString ActorLabel;
+		uint32 ActiveSoundPlayOrder = INDEX_NONE;
 	};
 
 	struct FSoundIsAlivePingMessage : public FSoundStartMessage
@@ -115,13 +121,15 @@ namespace UE::Audio::Insights
 		FSoundWaveStartMessage(const Trace::IAnalyzer::FOnEventContext& InContext)
 			: FSoundStartMessage(InContext)
 		{
+			WaveInstancePlayOrder = InContext.EventData.GetValue<uint32>("PlayOrder");
 			ActiveSoundPlayOrder = InContext.EventData.GetValue<uint32>("ActiveSoundPlayOrder");
 		}
 
 		virtual uint32 GetSizeOf() const override;
 		virtual const FName GetMessageName() const override { return SoundMessageNames::SoundWaveStart; }
+		virtual uint64 GetID() const override { return SoundMessageUtils::GeneratePlayOrderUniqueID(ActiveSoundPlayOrder, WaveInstancePlayOrder); }
 
-		uint32 ActiveSoundPlayOrder;
+		uint32 WaveInstancePlayOrder = INDEX_NONE;
 	};
 
 	struct FSoundWaveIsAlivePingMessage : public FSoundWaveStartMessage
@@ -141,10 +149,14 @@ namespace UE::Audio::Insights
 		FSoundStopMessage(const Trace::IAnalyzer::FOnEventContext& InContext)
 			: FSoundMessageBase(InContext)
 		{
+			ActiveSoundPlayOrder = InContext.EventData.GetValue<uint32>("PlayOrder");
 		}
 
 		virtual uint32 GetSizeOf() const override;
 		virtual const FName GetMessageName() const override { return SoundMessageNames::SoundStop; }
+		virtual uint64 GetID() const override { return ActiveSoundPlayOrder; }
+
+		uint32 ActiveSoundPlayOrder = INDEX_NONE;
 	};
 
 #define DEFINE_SOUND_PARAM_MESSAGE(ClassName, MessageName, ParamName, Type, Default)			\
@@ -278,8 +290,8 @@ namespace UE::Audio::Insights
 			return bIsCategory ? FLinearColor::White : EntryColor;
 		}
 
-		virtual bool IsValid() const override { return PlayOrder != INDEX_NONE; }
-		virtual int32 GetEntryID() const override { return PlayOrder; }
+		virtual bool IsValid() const override { return ActiveSoundPlayOrder != INDEX_NONE; }
+		virtual uint64 GetEntryID() const override { return SoundMessageUtils::GeneratePlayOrderUniqueID(ActiveSoundPlayOrder, WaveInstancePlayOrder); }
 
 		virtual bool HasSetInitExpansion() const override { return bHasSetInitExpansion; }
 		virtual void ResetHasSetInitExpansion() override { bHasSetInitExpansion = true; }
@@ -294,7 +306,8 @@ namespace UE::Audio::Insights
 		};
 
 		::Audio::FDeviceId DeviceId = INDEX_NONE;
-		uint32 PlayOrder = INDEX_NONE;
+		uint32 WaveInstancePlayOrder = INDEX_NONE;
+		uint32 ActiveSoundPlayOrder = INDEX_NONE;
 		double Timestamp = 0.0;
 		double TimeoutTimestamp = INVALID_TIMEOUT;
 		FText CategoryName;
diff --git a/Engine/Plugins/AudioInsights/Source/AudioInsights/Public/Messages/SubmixTraceMessages.h b/Engine/Plugins/AudioInsights/Source/AudioInsights/Public/Messages/SubmixTraceMessages.h
--- a/Engine/Plugins/AudioInsights/Source/AudioInsights/Public/Messages/SubmixTraceMessages.h
+++ b/Engine/Plugins/AudioInsights/Source/AudioInsights/Public/Messages/SubmixTraceMessages.h
@@ -29,7 +29,7 @@ namespace UE::Audio::Insights
 		FSubmixMessageBase() = default;
 		FSubmixMessageBase(const Trace::IAnalyzer::FOnEventContext& InContext);
 
-		virtual uint32 GetID() const override { return SubmixId; };
+		virtual uint64 GetID() const override { return SubmixId; };
 		virtual uint32 GetSizeOf() const override { return sizeof(FSubmixMessageBase); };
 
 		::Audio::FDeviceId DeviceId = INDEX_NONE;
diff --git a/Engine/Plugins/AudioInsights/Source/AudioInsights/Public/Messages/VirtualLoopTraceMessages.h b/Engine/Plugins/AudioInsights/Source/AudioInsights/Public/Messages/VirtualLoopTraceMessages.h
--- a/Engine/Plugins/AudioInsights/Source/AudioInsights/Public/Messages/VirtualLoopTraceMessages.h
+++ b/Engine/Plugins/AudioInsights/Source/AudioInsights/Public/Messages/VirtualLoopTraceMessages.h
@@ -27,7 +27,7 @@ namespace UE::Audio::Insights
 		FVirtualLoopMessageBase() = default;
 		FVirtualLoopMessageBase(const Trace::IAnalyzer::FOnEventContext& InContext);
 
-		virtual uint32 GetID() const override { return PlayOrder; }
+		virtual uint64 GetID() const override { return PlayOrder; }
 
 		::Audio::FDeviceId DeviceId = INDEX_NONE;
 		uint32 PlayOrder = INDEX_NONE;
diff --git a/Engine/Plugins/AudioInsights/Source/AudioInsights/Public/Views/SAudioCurveView.h b/Engine/Plugins/AudioInsights/Source/AudioInsights/Public/Views/SAudioCurveView.h
--- a/Engine/Plugins/AudioInsights/Source/AudioInsights/Public/Views/SAudioCurveView.h
+++ b/Engine/Plugins/AudioInsights/Source/AudioInsights/Public/Views/SAudioCurveView.h
@@ -76,8 +76,8 @@ public:
 
 	void SetYValueFormattingOptions(const FNumberFormattingOptions InValueFormattingOptions);
 	void SetYAxisRange(const FFloatInterval& InRange);
-	void SetCurvesPointData(TSharedPtr<TMap<int32, TArray<FCurvePoint>>> InCurvesPointData);
-	void SetCurvesMetadata(TSharedPtr<TMap<int32, FCurveMetadata>> InMetadataPerCurve);
+	void SetCurvesPointData(TSharedPtr<TMap<uint64, TArray<FCurvePoint>>> InCurvesPointData);
+	void SetCurvesMetadata(TSharedPtr<TMap<uint64, FCurveMetadata>> InMetadataPerCurve);
 
 	FFloatInterval GetYAxisRange() const { return YDataRange; }
 
@@ -127,8 +127,8 @@ protected:
 	FOnKeyDown OnKeyDownHandler;
 
 	// Point data and metadata, keyed by curve id
-	TSharedPtr<TMap<int32, TArray<FCurvePoint>>> PointDataPerCurve;
-	TSharedPtr<TMap<int32, FCurveMetadata>> MetadataPerCurve; 
+	TSharedPtr<TMap<uint64, TArray<FCurvePoint>>> PointDataPerCurve;
+	TSharedPtr<TMap<uint64, FCurveMetadata>> MetadataPerCurve;
 	// Y axis data range in value space
 	FFloatInterval YDataRange;
 
diff --git a/Engine/Plugins/AudioInsights/Source/AudioInsights/Public/Views/SoundPlotsWidgetView.h b/Engine/Plugins/AudioInsights/Source/AudioInsights/Public/Views/SoundPlotsWidgetView.h
--- a/Engine/Plugins/AudioInsights/Source/AudioInsights/Public/Views/SoundPlotsWidgetView.h
+++ b/Engine/Plugins/AudioInsights/Source/AudioInsights/Public/Views/SoundPlotsWidgetView.h
@@ -42,8 +42,8 @@ namespace UE::Audio::Insights
 
 	private:
 		using FPlotCurvePoint = SAudioCurveView::FCurvePoint;
-		using FPointDataPerCurveMap = TMap<int32, TArray<FPlotCurvePoint>>; // Map of source id to data point array 
-		using FLatestTimestampPerCurveMap = TMap<int32, double>; // Map of source id to most recent timestamp associated with any data
+		using FPointDataPerCurveMap = TMap<uint64, TArray<FPlotCurvePoint>>; // Map of source id to data point array 
+		using FLatestTimestampPerCurveMap = TMap<uint64, double>; // Map of source id to most recent timestamp associated with any data
 		using FPlotCurveMetadata = SAudioCurveView::FCurveMetadata;
 
 		void OnAnalysisStarting(const double Timestamp);
@@ -105,7 +105,7 @@ namespace UE::Audio::Insights
 		TMap<FName, TSharedPtr<FLatestTimestampPerCurveMap>> ColumnNameToLatestTimestamp;
 
 		// SourceId to metadata for the corresponding curve
-		TSharedPtr<TMap<int32, FPlotCurveMetadata>> PlotWidgetMetadataPerCurve;
+		TSharedPtr<TMap<uint64, FPlotCurveMetadata>> PlotWidgetMetadataPerCurve;
 
 		// Column names for plot selector widget 
 		TArray<FName> ColumnNames;
