[Audio Insights] Fix for build errors in SoundTraceProvider

diff --git a/Engine/Plugins/AudioInsights/Source/AudioInsights/Private/Providers/SoundTraceProvider.cpp b/Engine/Plugins/AudioInsights/Source/AudioInsights/Private/Providers/SoundTraceProvider.cpp
--- a/Engine/Plugins/AudioInsights/Source/AudioInsights/Private/Providers/SoundTraceProvider.cpp
+++ b/Engine/Plugins/AudioInsights/Source/AudioInsights/Private/Providers/SoundTraceProvider.cpp
@@ -103,8 +103,8 @@ namespace UE::Audio::Insights
 
 		struct PlayOrderPair
 		{
-			uint32 ActiveSoundPlayOrderID { INDEX_NONE };
-			uint32 WaveInstancePlayOrderID { INDEX_NONE };
+			uint32 ActiveSoundPlayOrderID = static_cast<uint32>(INDEX_NONE);
+			uint32 WaveInstancePlayOrderID = static_cast<uint32>(INDEX_NONE);
 		};
 
 		void CollectPlayOrderPairsRecursive(const TSharedPtr<IDashboardDataTreeViewEntry>& Entry, TArray<PlayOrderPair>& OutPlayOrderArray)
