[Audio Insights] Increase the time messages can be inside the analyzer queue before being cleaned up across the board. Avoids losing data when loading in large trace files.

diff --git a/Engine/Plugins/AudioInsights/Source/AudioInsights/Public/Cache/AudioInsightsCacheManager.h b/Engine/Plugins/AudioInsights/Source/AudioInsights/Public/Cache/AudioInsightsCacheManager.h
--- a/Engine/Plugins/AudioInsights/Source/AudioInsights/Public/Cache/AudioInsightsCacheManager.h
+++ b/Engine/Plugins/AudioInsights/Source/AudioInsights/Public/Cache/AudioInsightsCacheManager.h
@@ -199,7 +199,7 @@ namespace UE::Audio::Insights
 		TraceServices::FSlabAllocator Allocator;
 		TCircularBuffer<TUniquePtr<FAudioCachedMessageChunk>> Cache;
 
-		TAnalyzerMessageQueue<TSharedPtr<IAudioCachedMessage>> NewMessageQueue { 5.0 };
+		TAnalyzerMessageQueue<TSharedPtr<IAudioCachedMessage>> NewMessageQueue;
 
 		uint32 WriteIndex = 0u;
 		uint32 StartIndex = 0u;
diff --git a/Engine/Plugins/AudioInsights/Source/AudioInsights/Public/Messages/AnalyzerMessageQueue.h b/Engine/Plugins/AudioInsights/Source/AudioInsights/Public/Messages/AnalyzerMessageQueue.h
--- a/Engine/Plugins/AudioInsights/Source/AudioInsights/Public/Messages/AnalyzerMessageQueue.h
+++ b/Engine/Plugins/AudioInsights/Source/AudioInsights/Public/Messages/AnalyzerMessageQueue.h
@@ -13,8 +13,8 @@ namespace UE::Audio::Insights
 	{
 		static constexpr double MaxHistoryLimitSec = 5.0;
 
-		explicit TAnalyzerMessageQueue(double InHistoryLimitSec)
-			: HistoryLimitSec(FMath::Clamp(InHistoryLimitSec, UE_DOUBLE_KINDA_SMALL_NUMBER, 5.0))
+		explicit TAnalyzerMessageQueue(const double InHistoryLimitSec = 5.0)
+			: HistoryLimitSec(FMath::Clamp(InHistoryLimitSec, UE_DOUBLE_KINDA_SMALL_NUMBER, MaxHistoryLimitSec))
 		{
 		}
 
diff --git a/Engine/Plugins/AudioInsights/Source/AudioInsights/Public/Messages/AudioBusTraceMessages.h b/Engine/Plugins/AudioInsights/Source/AudioInsights/Public/Messages/AudioBusTraceMessages.h
--- a/Engine/Plugins/AudioInsights/Source/AudioInsights/Public/Messages/AudioBusTraceMessages.h
+++ b/Engine/Plugins/AudioInsights/Source/AudioInsights/Public/Messages/AudioBusTraceMessages.h
@@ -98,11 +98,11 @@ namespace UE::Audio::Insights
 
 	class FAudioBusMessages
 	{
-		TAnalyzerMessageQueue<FAudioBusStartMessage> StartMessages{ 0.1 };
-		TAnalyzerMessageQueue<FAudioBusHasActivityMessage> HasActivityMessages{ 0.1 };
-		TAnalyzerMessageQueue<FAudioBusEnvelopeFollowerEnabledMessage> EnvelopeFollowerEnabledMessages{ 0.1 };
-		TAnalyzerMessageQueue<FAudioBusEnvelopeValuesMessage> EnvelopeValuesMessages{ 0.1 };
-		TAnalyzerMessageQueue<FAudioBusStopMessage> StopMessages{ 0.1 };
+		TAnalyzerMessageQueue<FAudioBusStartMessage> StartMessages;
+		TAnalyzerMessageQueue<FAudioBusHasActivityMessage> HasActivityMessages;
+		TAnalyzerMessageQueue<FAudioBusEnvelopeFollowerEnabledMessage> EnvelopeFollowerEnabledMessages;
+		TAnalyzerMessageQueue<FAudioBusEnvelopeValuesMessage> EnvelopeValuesMessages;
+		TAnalyzerMessageQueue<FAudioBusStopMessage> StopMessages;
 
 		friend class FAudioBusTraceProvider;
 	};
diff --git a/Engine/Plugins/AudioInsights/Source/AudioInsights/Public/Messages/AudioEventLogTraceMessages.h b/Engine/Plugins/AudioInsights/Source/AudioInsights/Public/Messages/AudioEventLogTraceMessages.h
--- a/Engine/Plugins/AudioInsights/Source/AudioInsights/Public/Messages/AudioEventLogTraceMessages.h
+++ b/Engine/Plugins/AudioInsights/Source/AudioInsights/Public/Messages/AudioEventLogTraceMessages.h
@@ -83,7 +83,7 @@ namespace UE::Audio::Insights
 
 	class FAudioEventLogMessages
 	{
-		TAnalyzerMessageQueue<FAudioEventLogMessage> AudioEventLogMessages { 5.0 };
+		TAnalyzerMessageQueue<FAudioEventLogMessage> AudioEventLogMessages;
 
 		friend class FAudioEventLogTraceProvider;
 	};
diff --git a/Engine/Plugins/AudioInsights/Source/AudioInsights/Public/Messages/SoundTraceMessages.h b/Engine/Plugins/AudioInsights/Source/AudioInsights/Public/Messages/SoundTraceMessages.h
--- a/Engine/Plugins/AudioInsights/Source/AudioInsights/Public/Messages/SoundTraceMessages.h
+++ b/Engine/Plugins/AudioInsights/Source/AudioInsights/Public/Messages/SoundTraceMessages.h
@@ -191,20 +191,20 @@ namespace UE::Audio::Insights
 	class FSoundMessages
 	{
 	public:
-		TAnalyzerMessageQueue<FSoundStartMessage> ActiveSoundStartMessages { 0.1 };
-		TAnalyzerMessageQueue<FSoundWaveStartMessage> SoundWaveStartMessages{ 0.1 };
-		TAnalyzerMessageQueue<FSoundIsAlivePingMessage> ActiveSoundIsAlivePingMessages{ 0.1 };
-		TAnalyzerMessageQueue<FSoundWaveIsAlivePingMessage> SoundWaveIsAlivePingMessages{ 0.1 };
-		TAnalyzerMessageQueue<FSoundPriorityMessage> PriorityMessages{ 0.1 };
-		TAnalyzerMessageQueue<FSoundDistanceMessage> DistanceMessages { 0.1 };
-		TAnalyzerMessageQueue<FSoundDistanceAttenuationMessage> DistanceAttenuationMessages { 0.1 };
-		TAnalyzerMessageQueue<FSoundHPFFreqMessage> HPFFreqMessages{ 0.1 };
-		TAnalyzerMessageQueue<FSoundLPFFreqMessage> LPFFreqMessages{ 0.1 };
-		TAnalyzerMessageQueue<FSoundEnvelopeMessage> AmplitudeMessages { 0.1 };
-		TAnalyzerMessageQueue<FSoundVolumeMessage> VolumeMessages { 0.1 };
-		TAnalyzerMessageQueue<FSoundPitchMessage> PitchMessages { 0.1 };
-		TAnalyzerMessageQueue<FSoundRelativeRenderCostMessage> RelativeRenderCostMessages{ 0.1 };
-		TAnalyzerMessageQueue<FSoundStopMessage> StopMessages { 0.1 };
+		TAnalyzerMessageQueue<FSoundStartMessage> ActiveSoundStartMessages;
+		TAnalyzerMessageQueue<FSoundWaveStartMessage> SoundWaveStartMessages;
+		TAnalyzerMessageQueue<FSoundIsAlivePingMessage> ActiveSoundIsAlivePingMessages;
+		TAnalyzerMessageQueue<FSoundWaveIsAlivePingMessage> SoundWaveIsAlivePingMessages;
+		TAnalyzerMessageQueue<FSoundPriorityMessage> PriorityMessages;
+		TAnalyzerMessageQueue<FSoundDistanceMessage> DistanceMessages;
+		TAnalyzerMessageQueue<FSoundDistanceAttenuationMessage> DistanceAttenuationMessages;
+		TAnalyzerMessageQueue<FSoundHPFFreqMessage> HPFFreqMessages;
+		TAnalyzerMessageQueue<FSoundLPFFreqMessage> LPFFreqMessages;
+		TAnalyzerMessageQueue<FSoundEnvelopeMessage> AmplitudeMessages;
+		TAnalyzerMessageQueue<FSoundVolumeMessage> VolumeMessages;
+		TAnalyzerMessageQueue<FSoundPitchMessage> PitchMessages;
+		TAnalyzerMessageQueue<FSoundRelativeRenderCostMessage> RelativeRenderCostMessages;
+		TAnalyzerMessageQueue<FSoundStopMessage> StopMessages;
 	};
 
 #if !WITH_EDITOR
diff --git a/Engine/Plugins/AudioInsights/Source/AudioInsights/Public/Messages/SubmixTraceMessages.h b/Engine/Plugins/AudioInsights/Source/AudioInsights/Public/Messages/SubmixTraceMessages.h
--- a/Engine/Plugins/AudioInsights/Source/AudioInsights/Public/Messages/SubmixTraceMessages.h
+++ b/Engine/Plugins/AudioInsights/Source/AudioInsights/Public/Messages/SubmixTraceMessages.h
@@ -94,11 +94,11 @@ namespace UE::Audio::Insights
 
 	class FSubmixMessages
 	{
-		TAnalyzerMessageQueue<FSubmixLoadedMessage> LoadedMessages{ 0.1 };
-		TAnalyzerMessageQueue<FSubmixHasActivityMessage> HasActivityMessages{ 0.1 };
-		TAnalyzerMessageQueue<FSubmixEnvelopeFollowerEnabledMessage> EnvelopeFollowerEnabledMessages{ 0.1 };
-		TAnalyzerMessageQueue<FSubmixEnvelopeValuesMessage> EnvelopeValuesMessages{ 0.1 };
-		TAnalyzerMessageQueue<FSubmixUnloadedMessage> UnloadedMessages{ 0.1 };
+		TAnalyzerMessageQueue<FSubmixLoadedMessage> LoadedMessages;
+		TAnalyzerMessageQueue<FSubmixHasActivityMessage> HasActivityMessages;
+		TAnalyzerMessageQueue<FSubmixEnvelopeFollowerEnabledMessage> EnvelopeFollowerEnabledMessages;
+		TAnalyzerMessageQueue<FSubmixEnvelopeValuesMessage> EnvelopeValuesMessages;
+		TAnalyzerMessageQueue<FSubmixUnloadedMessage> UnloadedMessages;
 
 		friend class FSubmixTraceProvider;
 	};
diff --git a/Engine/Plugins/AudioInsights/Source/AudioInsights/Public/Messages/VirtualLoopTraceMessages.h b/Engine/Plugins/AudioInsights/Source/AudioInsights/Public/Messages/VirtualLoopTraceMessages.h
--- a/Engine/Plugins/AudioInsights/Source/AudioInsights/Public/Messages/VirtualLoopTraceMessages.h
+++ b/Engine/Plugins/AudioInsights/Source/AudioInsights/Public/Messages/VirtualLoopTraceMessages.h
@@ -99,9 +99,9 @@ namespace UE::Audio::Insights
 
 	class FVirtualLoopMessages
 	{
-		TAnalyzerMessageQueue<FVirtualLoopVirtualizeMessage> VirtualizeMessages { 5.0 };
-		TAnalyzerMessageQueue<FVirtualLoopStopOrRealizeMessage> StopOrRealizeMessages { 5.0 };
-		TAnalyzerMessageQueue<FVirtualLoopUpdateMessage> UpdateMessages { 0.1 };
+		TAnalyzerMessageQueue<FVirtualLoopVirtualizeMessage> VirtualizeMessages;
+		TAnalyzerMessageQueue<FVirtualLoopStopOrRealizeMessage> StopOrRealizeMessages;
+		TAnalyzerMessageQueue<FVirtualLoopUpdateMessage> UpdateMessages;
 
 		friend class FVirtualLoopTraceProvider;
 	};
