[Audio Insights] FEditorDashboardFactory::CreateTraceBookmark: use SubText var instead of SaveBookmarkNotificationLiveTraceSubText

diff --git a/Engine/Plugins/AudioInsights/Source/AudioInsightsEditor/Private/AudioInsightsEditorDashboardFactory.cpp b/Engine/Plugins/AudioInsights/Source/AudioInsightsEditor/Private/AudioInsightsEditorDashboardFactory.cpp
--- a/Engine/Plugins/AudioInsights/Source/AudioInsightsEditor/Private/AudioInsightsEditorDashboardFactory.cpp
+++ b/Engine/Plugins/AudioInsights/Source/AudioInsightsEditor/Private/AudioInsightsEditorDashboardFactory.cpp
@@ -197,7 +197,7 @@ namespace UE::Audio::Insights
 			const FText SubText = TraceModule.IsTraceAnalysisActive() ? SaveBookmarkNotificationLiveTraceSubText : SaveBookmarkNotificationDirectTraceSubText;
 
 			static const FTextFormat SubTextFormat = FTextFormat::FromString(TEXT("{0} {1} {2}"));
-			Info.SubText = FText::Format(SubTextFormat, BookmarkLabelText, BookmarkIndex, SaveBookmarkNotificationLiveTraceSubText);
+			Info.SubText = FText::Format(SubTextFormat, BookmarkLabelText, BookmarkIndex, SubText);
 
 			Info.bFireAndForget = true;
 			Info.FadeOutDuration = 1.0f;
