[Audio Insights] Fix audio leakage when "solo filtered" has no sounds to solo in the sounds tab

diff --git a/Engine/Plugins/AudioInsights/Source/AudioInsights/Private/Views/SoundDashboardViewFactory.cpp b/Engine/Plugins/AudioInsights/Source/AudioInsights/Private/Views/SoundDashboardViewFactory.cpp
--- a/Engine/Plugins/AudioInsights/Source/AudioInsights/Private/Views/SoundDashboardViewFactory.cpp
+++ b/Engine/Plugins/AudioInsights/Source/AudioInsights/Private/Views/SoundDashboardViewFactory.cpp
@@ -2920,19 +2920,10 @@ namespace UE::Audio::Insights
 
 		// If you have filtered out everything but are soloing to listen things that do match the filter,
 		// we mute the rest of the mix at this moment. 
-		if (!GetSearchFilterText().ToString().IsEmpty() && Mode == EMuteSoloMode::Solo && !bSomethingIsVisible)
+		const bool bShouldMuteAll = !GetSearchFilterText().IsEmpty() && Mode == EMuteSoloMode::Solo && !bSomethingIsVisible;
+		if (bShouldMuteAll)
 		{
-			for (TSharedPtr<IDashboardDataTreeViewEntry>& Entry : DataViewEntries)
-			{
-				FSoundDashboardEntry& CategoryEntry = FSoundDashboardViewFactoryPrivate::CastEntry(*Entry);
-				if (CategoryEntry.bIsCategory)
-				{
-					for (TSharedPtr<IDashboardDataTreeViewEntry>& EntryChild : Entry->Children)
-					{
-						SetMuteSolo(*EntryChild, EMuteSoloMode::Mute, true /*bInOnOff*/);
-					}
-				}
-			}
+			AudioDebugger.SetSoloSoundCue(NAME_None, true /*bInOnOff*/);
 		}
 #endif //ENABLE_AUDIO_DEBUG
 	}
