[Audio Insights] Stop the plots tab time ruler from running when not monitoring

diff --git a/Engine/Plugins/AudioInsights/Source/AudioInsights/Private/Views/SoundPlotsWidgetView.cpp b/Engine/Plugins/AudioInsights/Source/AudioInsights/Private/Views/SoundPlotsWidgetView.cpp
--- a/Engine/Plugins/AudioInsights/Source/AudioInsights/Private/Views/SoundPlotsWidgetView.cpp
+++ b/Engine/Plugins/AudioInsights/Source/AudioInsights/Private/Views/SoundPlotsWidgetView.cpp
@@ -864,12 +864,16 @@ namespace UE::Audio::Insights
 		double RangeUpperBound = 0.0;
 
 #if WITH_EDITOR
+		const bool bAnyMessageReceived = CurrentTimestamp != TNumericLimits<double>::Lowest();
+		if (!bAnyMessageReceived && !bIsTraceActive)
+		{
+			return CurrentRange;
+		}
+
 		const double CurrentTime = FPlatformTime::Seconds() - GStartTime;
 
 		double TimestampsDiff = 0.0;
 
-		const bool bAnyMessageReceived = CurrentTimestamp != TNumericLimits<double>::Lowest();
-
 		if (GameState != EGameState::Paused && GameState != EGameState::ControlledByTimemarker && bAnyMessageReceived)
 		{
 			const double RelativeCurrentTime = CurrentTime - BeginTimestamp;
