[Audio Insights] Fix entries not being removed from the sounds tab after stopping when scrubbing
* Make sure we locate the active sound of a wave instance correctly when scrubbing
* Keep track of wave instance IDs when scrubbing and remove entries that have been stopped. Fixes wave instances building up when virtualizing/realizing sounds.

diff --git a/Engine/Plugins/AudioInsights/Source/AudioInsights/Private/Providers/SoundTraceProvider.cpp b/Engine/Plugins/AudioInsights/Source/AudioInsights/Private/Providers/SoundTraceProvider.cpp
--- a/Engine/Plugins/AudioInsights/Source/AudioInsights/Private/Providers/SoundTraceProvider.cpp
+++ b/Engine/Plugins/AudioInsights/Source/AudioInsights/Private/Providers/SoundTraceProvider.cpp
@@ -1225,10 +1225,11 @@ namespace UE::Audio::Insights
 			}
 		}
 
-		TSet<uint32> ActiveSoundIDs;
+		TSet<uint32> AliveActiveSoundIDs;
+		TSet<uint32> AliveWaveInstanceIDs;
 
 		// Collect all the start messages registered until this point in time
-		auto ProcessStartMessage = [this, &ActiveSoundIDs, &InitActiveSoundEntry](const FSoundStartMessage& StartCachedMessage)
+		auto ProcessStartMessage = [this, &AliveActiveSoundIDs, &InitActiveSoundEntry](const FSoundStartMessage& StartCachedMessage)
 		{
 			TSharedPtr<FSoundDashboardEntry> ActiveSoundEntry = nullptr;
 			GetOrCreateActiveSoundEntry(StartCachedMessage, ActiveSoundEntry);
@@ -1238,7 +1239,7 @@ namespace UE::Audio::Insights
 				return;
 			}
 
-			ActiveSoundIDs.Add(StartCachedMessage.ActiveSoundPlayOrder);
+			AliveActiveSoundIDs.Add(StartCachedMessage.ActiveSoundPlayOrder);
 			InitActiveSoundEntry(StartCachedMessage, *ActiveSoundEntry);
 		};
 
@@ -1262,9 +1263,9 @@ namespace UE::Audio::Insights
 			
 
 		// Collect all the is alive ping messages registered until this point in time
-		auto ProcessSoundIsAlivePingMessage = [this, &ActiveSoundIDs, &InitActiveSoundEntry](const FSoundIsAlivePingMessage& IsAlivePingCachedMessage)
+		auto ProcessSoundIsAlivePingMessage = [this, &AliveActiveSoundIDs, &InitActiveSoundEntry](const FSoundIsAlivePingMessage& IsAlivePingCachedMessage)
 		{
-			ActiveSoundIDs.Add(IsAlivePingCachedMessage.ActiveSoundPlayOrder);
+			AliveActiveSoundIDs.Add(IsAlivePingCachedMessage.ActiveSoundPlayOrder);
 
 			TSharedPtr<FSoundDashboardEntry> ActiveSoundEntry = nullptr;
 			CreateAndGetActiveSoundEntryIfNotFound(IsAlivePingCachedMessage, ActiveSoundEntry);
@@ -1296,16 +1297,18 @@ namespace UE::Audio::Insights
 
 		////////////////////////////////////////
 		// Add all soundwave start messages registered until this point in time 
-		TSharedPtr<FSoundDashboardEntry> ParentSoundEntryReturn = nullptr;
-		auto ProcessSoundWaveStartMessage = [this, &ActiveSoundIDs, &ParentSoundEntryReturn, &CreateChildEntry](const FSoundWaveStartMessage& SoundWaveStartCachedMessage)
+		auto ProcessSoundWaveStartMessage = [this, &AliveActiveSoundIDs, &AliveWaveInstanceIDs, &CreateChildEntry](const FSoundWaveStartMessage& SoundWaveStartCachedMessage)
 		{
+			TSharedPtr<FSoundDashboardEntry> ParentSoundEntryReturn = nullptr;
+			GetActiveSoundEntryFromIDs(SoundWaveStartCachedMessage.ActiveSoundPlayOrder, ParentSoundEntryReturn);
+
 			if (ParentSoundEntryReturn == nullptr || !ParentSoundEntryReturn->IsValid() || !FSoundTraceProviderPrivate::EntryTypeHasSoundWaveEntries(ParentSoundEntryReturn->EntryType))
 			{
-				GetActiveSoundEntryFromIDs(SoundWaveStartCachedMessage.ActiveSoundPlayOrder, ParentSoundEntryReturn);
 				return;
 			}
 
-			ActiveSoundIDs.Add(SoundWaveStartCachedMessage.ActiveSoundPlayOrder);
+			AliveActiveSoundIDs.Add(SoundWaveStartCachedMessage.ActiveSoundPlayOrder);
+			AliveWaveInstanceIDs.Add(SoundWaveStartCachedMessage.WaveInstancePlayOrder);
 
 			const bool bEntryExists = ParentSoundEntryReturn->Children.ContainsByPredicate([&SoundWaveStartCachedMessage](const TSharedPtr<IDashboardDataTreeViewEntry>& ChildEntry)
 			{
@@ -1339,9 +1342,9 @@ namespace UE::Audio::Insights
 
 		////////////////////////////////////////
 		// Add all soundwave ping messages registered until this point in time 
-		TSharedPtr<FSoundDashboardEntry> ParentSoundIsAlivePingEntryReturn = nullptr;
-		auto ProcessSoundWaveIsAlivePingMessage = [this, &ActiveSoundIDs, &ParentSoundIsAlivePingEntryReturn, &CreateChildEntry](const FSoundWaveIsAlivePingMessage& SoundWaveIsAlivePingCachedMessage)
+		auto ProcessSoundWaveIsAlivePingMessage = [this, &AliveActiveSoundIDs, &AliveWaveInstanceIDs, &CreateChildEntry](const FSoundWaveIsAlivePingMessage& SoundWaveIsAlivePingCachedMessage)
 		{
+			TSharedPtr<FSoundDashboardEntry> ParentSoundIsAlivePingEntryReturn = nullptr;
 			GetActiveSoundEntryFromIDs(SoundWaveIsAlivePingCachedMessage.ActiveSoundPlayOrder, ParentSoundIsAlivePingEntryReturn);
 
 			if (ParentSoundIsAlivePingEntryReturn == nullptr || !ParentSoundIsAlivePingEntryReturn->IsValid() || !FSoundTraceProviderPrivate::EntryTypeHasSoundWaveEntries(ParentSoundIsAlivePingEntryReturn->EntryType))
@@ -1349,7 +1352,8 @@ namespace UE::Audio::Insights
 				return;
 			}
 
-			ActiveSoundIDs.Add(SoundWaveIsAlivePingCachedMessage.ActiveSoundPlayOrder);
+			AliveActiveSoundIDs.Add(SoundWaveIsAlivePingCachedMessage.ActiveSoundPlayOrder);
+			AliveWaveInstanceIDs.Add(SoundWaveIsAlivePingCachedMessage.WaveInstancePlayOrder);
 
 			// Avoid creating duplicate entries when receiving the IsAlivePing message
 			const bool bEntryExists = ParentSoundIsAlivePingEntryReturn->Children.ContainsByPredicate([&SoundWaveIsAlivePingCachedMessage](const TSharedPtr<IDashboardDataTreeViewEntry>& ChildEntry)
@@ -1384,20 +1388,51 @@ namespace UE::Audio::Insights
 
 		// Selectively remove start messages collected in the step above by knowing which sounds were stopped.
 		// With this we will know what are the active sounds at this point in time.
-		auto ProcessSoundStopMessageMessage = [this](const FSoundStopMessage& StopCachedMessage)
+		auto ProcessSoundStopMessageMessage = [this, &AliveWaveInstanceIDs](const FSoundStopMessage& StopCachedMessage)
 		{
 			TSharedPtr<FSoundDashboardEntry> ActiveSoundEntry = nullptr;
 			GetActiveSoundEntryFromIDs(StopCachedMessage.ActiveSoundPlayOrder, ActiveSoundEntry);
 
-			if (ActiveSoundEntry.IsValid() && ActiveSoundEntry->Timestamp < StopCachedMessage.Timestamp)
+			if (ActiveSoundEntry.IsValid())
 			{
-				if (ActiveSoundEntry->bIsPlotActive)
+				if (ActiveSoundEntry->Timestamp < StopCachedMessage.Timestamp)
 				{
-					ActiveSoundEntry->bForceKeepEntryAlive = true;
+					if (ActiveSoundEntry->bIsPlotActive)
+					{
+						ActiveSoundEntry->bForceKeepEntryAlive = true;
+					}
+					else
+					{
+						// This active sound has been stopped, clean up any child Wave Instance play orders
+						for (auto Iter = ActiveSoundEntry->Children.CreateConstIterator(); Iter; ++Iter)
+						{
+							const TSharedPtr<IDashboardDataTreeViewEntry>& ChildEntry = *Iter;
+							if (ChildEntry.IsValid())
+							{
+								AliveWaveInstanceIDs.Remove(CastEntry(*ChildEntry).WaveInstancePlayOrder);
+							}
+						}
+
+						RemoveActiveSoundEntry(ActiveSoundEntry);
+					}
 				}
 				else
 				{
-					RemoveActiveSoundEntry(ActiveSoundEntry);
+					// This active sound is alive, but may have virtualized/realized
+					// Remove any old wave instances that are no longer alive
+					for (auto Iter = ActiveSoundEntry->Children.CreateIterator(); Iter; ++Iter)
+					{
+						const TSharedPtr<IDashboardDataTreeViewEntry> ChildEntry = *Iter;
+						if (!ChildEntry.IsValid())
+						{
+							Iter.RemoveCurrent();
+						}
+						else if (CastEntry(*ChildEntry).Timestamp < StopCachedMessage.Timestamp)
+						{
+							AliveWaveInstanceIDs.Remove(CastEntry(*ChildEntry).WaveInstancePlayOrder);
+							Iter.RemoveCurrent();
+						}
+					}
 				}
 			}
 		};
@@ -1426,15 +1461,40 @@ namespace UE::Audio::Insights
 		for (int ReverseIndex = ActiveSoundPlayOrderArray.Num() - 1; ReverseIndex >= 0; --ReverseIndex)
 		{
 			const uint32 ActiveSoundPlayOrder = ActiveSoundPlayOrderArray[ReverseIndex];
-			if (!ActiveSoundIDs.Contains(ActiveSoundPlayOrder))
+
+			TSharedPtr<FSoundDashboardEntry> ActiveSoundEntry;
+			GetActiveSoundEntryFromIDs(ActiveSoundPlayOrder, ActiveSoundEntry);
+
+			if (!ActiveSoundEntry.IsValid())
 			{
-				TSharedPtr<FSoundDashboardEntry> ActiveSoundEntry;
-				GetActiveSoundEntryFromIDs(ActiveSoundPlayOrder, ActiveSoundEntry);
+				continue;
+			}
 
-				if (ActiveSoundEntry.IsValid())
+			if (AliveActiveSoundIDs.Contains(ActiveSoundPlayOrder))
+			{
+				// This active sound is alive - check it's wave instances
+				for (auto Iter = ActiveSoundEntry->Children.CreateIterator(); Iter; ++Iter)
+				{
+					const TSharedPtr<IDashboardDataTreeViewEntry> ChildEntry = *Iter;
+					if (!ChildEntry.IsValid() || !AliveWaveInstanceIDs.Contains(CastEntry(*ChildEntry).WaveInstancePlayOrder))
+					{
+						Iter.RemoveCurrent();
+					}
+				}
+			}
+			else
+			{
+				// This active sound is not alive - remove it and clean up any wave instances
+				for (auto Iter = ActiveSoundEntry->Children.CreateConstIterator(); Iter; ++Iter)
 				{
-					RemoveActiveSoundEntry(ActiveSoundEntry);
+					const TSharedPtr<IDashboardDataTreeViewEntry> ChildEntry = *Iter;
+					if (ChildEntry.IsValid())
+					{
+						AliveWaveInstanceIDs.Remove(CastEntry(*ChildEntry).WaveInstancePlayOrder);
+					}
 				}
+
+				RemoveActiveSoundEntry(ActiveSoundEntry);
 			}
 		}
 
