[Audio Insights] Fix sound cue active sounds not showing plot data when enabling plots on them in the sounds tab
* Sound Cues show aggregate data of the wave instances it contains, so we count these as categories of multiple sounds when plotting - i.e. if a sound cue is selected to plot, we individually plot the child wave instances
* Make sure the "IsPlottingEnabled" flag is reset if a child wave instance diables plotting. This puts the behaviour in line with Mute/Solo.

diff --git a/Engine/Plugins/AudioInsights/Source/AudioInsights/Private/Views/SoundDashboardViewFactory.cpp b/Engine/Plugins/AudioInsights/Source/AudioInsights/Private/Views/SoundDashboardViewFactory.cpp
--- a/Engine/Plugins/AudioInsights/Source/AudioInsights/Private/Views/SoundDashboardViewFactory.cpp
+++ b/Engine/Plugins/AudioInsights/Source/AudioInsights/Private/Views/SoundDashboardViewFactory.cpp
@@ -3085,6 +3085,8 @@ namespace UE::Audio::Insights
 	{
 		using namespace FSoundDashboardViewFactoryPrivate;
 
+		TSet<uint32> ActiveSoundsToggledOff;
+
 		for (const TSharedPtr<IDashboardDataTreeViewEntry>& Entry : InEntries)
 		{
 			if (!Entry.IsValid())
@@ -3093,21 +3095,27 @@ namespace UE::Audio::Insights
 			}
 
 			FSoundDashboardEntry& SoundEntry = CastEntry(*Entry);
-			if (SoundEntry.bIsCategory)
+			if (SoundEntry.bIsCategory && bIgnoreCategories)
 			{
-				if (bIgnoreCategories)
-				{
-					continue;
-				}
-
-				SetShowPlotRecursive(Entry, bShowPlot);
+				continue;
 			}
 			else
 			{
-				SetShowPlot(Entry, bShowPlot);
+				SetShowPlotRecursive(Entry, bShowPlot);
+
+				if (!bShowPlot && !SoundEntry.bIsCategory)
+				{
+					ActiveSoundsToggledOff.Add(SoundEntry.ActiveSoundPlayOrder);
+				}
 			}
 		}
 
+		// Clean up any parent Is Plotting states if a child has been disabled
+		if (!bShowPlot && ActiveSoundsToggledOff.Num() > 0)
+		{
+			CleanUpParentPlotEnabledStates(ActiveSoundsToggledOff);
+		}
+
 		if (FilteredEntriesListView.IsValid())
 		{
 			OnProcessPlotData.Broadcast(DataViewEntries, FilteredEntriesListView->GetSelectedItems());
@@ -3165,6 +3173,65 @@ namespace UE::Audio::Insights
 		}
 	}
 
+	void FSoundDashboardViewFactory::CleanUpParentPlotEnabledStates(const TSet<uint32>& ActiveSoundsToggledOff)
+	{
+		using namespace FSoundDashboardViewFactoryPrivate;
+
+		for (const TSharedPtr<IDashboardDataTreeViewEntry>& CategoryEntry : FullTree)
+		{
+			if (!CategoryEntry.IsValid())
+			{
+				continue;
+			}
+
+			const FSoundDashboardEntry& CategorySoundEntry = CastEntry(*CategoryEntry);
+
+			if (CategorySoundEntry.EntryType == ESoundDashboardEntryType::Pinned)
+			{
+				for (const TSharedPtr<IDashboardDataTreeViewEntry>& PinnedCategoryEntry : CategoryEntry->Children)
+				{
+					if (!PinnedCategoryEntry.IsValid())
+					{
+						continue;
+					}
+
+					for (const TSharedPtr<IDashboardDataTreeViewEntry>& ChildEntry : PinnedCategoryEntry->Children)
+					{
+						if (!ChildEntry.IsValid())
+						{
+							continue;
+						}
+
+						const FSoundDashboardEntry& ActiveSoundEntry = CastEntry(*ChildEntry);
+						if (ActiveSoundsToggledOff.Contains(ActiveSoundEntry.ActiveSoundPlayOrder))
+						{
+							SetShowPlot(CategoryEntry, false);
+							SetShowPlot(PinnedCategoryEntry, false);
+							SetShowPlot(ChildEntry, false);
+						}
+					}
+				}
+			}
+			else
+			{
+				for (const TSharedPtr<IDashboardDataTreeViewEntry>& ChildEntry : CategoryEntry->Children)
+				{
+					if (!ChildEntry.IsValid())
+					{
+						continue;
+					}
+
+					const FSoundDashboardEntry& ActiveSoundEntry = CastEntry(*ChildEntry);
+					if (ActiveSoundsToggledOff.Contains(ActiveSoundEntry.ActiveSoundPlayOrder))
+					{
+						SetShowPlot(CategoryEntry, false);
+						SetShowPlot(ChildEntry, false);
+					}
+				}
+			}
+		}
+	}
+
 	void FSoundDashboardViewFactory::ProcessPlotData()
 	{
 #if WITH_EDITOR
diff --git a/Engine/Plugins/AudioInsights/Source/AudioInsights/Public/Views/SoundDashboardViewFactory.h b/Engine/Plugins/AudioInsights/Source/AudioInsights/Public/Views/SoundDashboardViewFactory.h
--- a/Engine/Plugins/AudioInsights/Source/AudioInsights/Public/Views/SoundDashboardViewFactory.h
+++ b/Engine/Plugins/AudioInsights/Source/AudioInsights/Public/Views/SoundDashboardViewFactory.h
@@ -201,6 +201,7 @@ namespace UE::Audio::Insights
 		UE_API void ToggleShowPlotEntries(const TArray<TSharedPtr<IDashboardDataTreeViewEntry>>& InEntries, const bool bShowPlot, const bool bIgnoreCategories = false);
 		UE_API void SetShowPlotRecursive(const TSharedPtr<IDashboardDataTreeViewEntry>& Entry, const bool bShowPlot);
 		UE_API void SetShowPlot(const TSharedPtr<IDashboardDataTreeViewEntry>& Entry, const bool bShowPlot);
+		void CleanUpParentPlotEnabledStates(const TSet<uint32>& ActiveSoundsToggledOff);
 
 		void ProcessPlotData();
 
