[Audio Insights] Fix for pausing/resuming trace causing hanging sounds across Audio Insights
* Listen out for when trace channels are enabled/disabled. If our channels have been enabled, request a refresh
* Each dashboard handles refresh slightly differently to rebuild the state of the tab
* Added IsAlive pings for virtual loops - helps when rebuilding the list and will fix missing loops inside snapshots too.

diff --git a/Engine/Plugins/AudioInsights/Source/AudioInsights/Private/AudioInsightsTraceProviderBase.cpp b/Engine/Plugins/AudioInsights/Source/AudioInsights/Private/AudioInsightsTraceProviderBase.cpp
--- a/Engine/Plugins/AudioInsights/Source/AudioInsights/Private/AudioInsightsTraceProviderBase.cpp
+++ b/Engine/Plugins/AudioInsights/Source/AudioInsights/Private/AudioInsightsTraceProviderBase.cpp
@@ -41,6 +41,10 @@ namespace UE::Audio::Insights
 	void FTraceProviderBase::FTraceAnalyzerBase::OnAnalysisBegin(const FOnAnalysisContext& Context)
 	{
 		Provider->Reset();
+
+		UE::Trace::IAnalyzer::FInterfaceBuilder& Builder = Context.InterfaceBuilder;
+		Builder.RouteEvent(RouteID_ChannelAnnounce, "Trace", "ChannelAnnounce");
+		Builder.RouteEvent(RouteID_ChannelToggle, "Trace", "ChannelToggle");
 	}
 
 	bool FTraceProviderBase::FTraceAnalyzerBase::OnEvent(uint16 RouteId, EStyle Style, const FOnEventContext& Context)
@@ -52,7 +56,84 @@ namespace UE::Audio::Insights
 		}
 #endif // WITH_EDITOR
 
-		return OnHandleEvent(RouteId, Style, Context);
+		switch (RouteId)
+		{
+			case RouteID_ChannelAnnounce:
+			{
+				return HandleChannelAnnounceEvent(RouteId, Style, Context);
+				break;
+			}
+			case RouteID_ChannelToggle:
+				return HandleChannelToggleEvent(RouteId, Style, Context);
+				break;
+			default:
+			{
+				return OnHandleEvent(RouteId, Style, Context);
+			}
+		}
+	}
+
+	bool FTraceProviderBase::FTraceAnalyzerBase::HandleChannelAnnounceEvent(uint16 RouteId, EStyle Style, const FOnEventContext& Context)
+	{
+		if (AudioChannelID.IsSet() && AudioMixerChannelID.IsSet())
+		{
+			return OnEventSuccess(RouteId, Style, Context);
+		}
+
+		const Trace::IAnalyzer::FEventData& EventData = Context.EventData;
+
+		FString ChannelName;
+		EventData.GetString("Name", ChannelName);
+
+		if (!AudioChannelID.IsSet() && ChannelName.Compare("Audio") == 0)
+		{
+			AudioChannelID = EventData.GetValue<uint32>("Id");
+		}
+		else if (!AudioMixerChannelID.IsSet() && ChannelName.Compare("AudioMixer") == 0)
+		{
+			AudioMixerChannelID = EventData.GetValue<uint32>("Id");
+		}
+
+		return OnEventSuccess(RouteId, Style, Context);
+	}
+
+	bool FTraceProviderBase::FTraceAnalyzerBase::HandleChannelToggleEvent(uint16 RouteId, EStyle Style, const FOnEventContext& Context)
+	{
+		if (!AudioChannelID.IsSet() || !AudioMixerChannelID.IsSet())
+		{
+			return OnEventSuccess(RouteId, Style, Context);
+		}
+
+		const bool bPrevAudioChannelIsEnabled = bAudioChannelIsEnabled;
+		const bool bPrevAudioMixerChannelIsEnabled = bAudioMixerChannelIsEnabled;
+
+		const Trace::IAnalyzer::FEventData& EventData = Context.EventData;
+		const uint32 ToggledChannelID = EventData.GetValue<uint32>("Id");
+		if (ToggledChannelID == AudioChannelID.GetValue())
+		{
+			bAudioChannelIsEnabled = EventData.GetValue<bool>("IsEnabled");
+		}
+		else if (ToggledChannelID == AudioMixerChannelID.GetValue())
+		{
+			bAudioMixerChannelIsEnabled = EventData.GetValue<bool>("IsEnabled");
+		}
+
+		const bool bChannelStateHasBeenChanged = bAudioChannelIsEnabled != bPrevAudioChannelIsEnabled || bAudioMixerChannelIsEnabled != bPrevAudioMixerChannelIsEnabled;
+		const bool bProviderRequiresRefresh = bAudioChannelIsEnabled && bAudioMixerChannelIsEnabled && bChannelStateHasBeenChanged;
+												
+		if (bProviderRequiresRefresh)
+		{
+			if (IsInGameThread())
+			{
+				Provider->OnTraceChannelsEnabled();
+			}
+			else
+			{
+				ExecuteOnGameThread(TEXT("FTraceProviderBase::FTraceAnalyzerBase::::OnTraceChannelsEnabled"), [this] { Provider->OnTraceChannelsEnabled(); });
+			}
+		}
+
+		return OnEventSuccess(RouteId, Style, Context);
 	}
 
 	bool FTraceProviderBase::FTraceAnalyzerBase::OnEventSuccess(uint16 RouteId, EStyle Style, const FOnEventContext& Context)
diff --git a/Engine/Plugins/AudioInsights/Source/AudioInsights/Private/Providers/AudioBusTraceProvider.cpp b/Engine/Plugins/AudioInsights/Source/AudioInsights/Private/Providers/AudioBusTraceProvider.cpp
--- a/Engine/Plugins/AudioInsights/Source/AudioInsights/Private/Providers/AudioBusTraceProvider.cpp
+++ b/Engine/Plugins/AudioInsights/Source/AudioInsights/Private/Providers/AudioBusTraceProvider.cpp
@@ -193,6 +193,16 @@ namespace UE::Audio::Insights
 		return new FAudioBusTraceAnalyzer(AsShared(), InSession);
 	}
 
+	void FAudioBusTraceProvider::OnTraceChannelsEnabled()
+	{
+		FTraceProviderBase::OnTraceChannelsEnabled();
+
+#if WITH_EDITOR
+		AudioBusAssetProvider.RequestEntriesUpdate();
+		OnAudioBusListUpdated.ExecuteIfBound();
+#endif // WITH_EDITOR
+	}
+
 	bool FAudioBusTraceProvider::ProcessMessages()
 	{
 		// Helper lambdas
diff --git a/Engine/Plugins/AudioInsights/Source/AudioInsights/Private/Providers/AudioBusTraceProvider.h b/Engine/Plugins/AudioInsights/Source/AudioInsights/Private/Providers/AudioBusTraceProvider.h
--- a/Engine/Plugins/AudioInsights/Source/AudioInsights/Private/Providers/AudioBusTraceProvider.h
+++ b/Engine/Plugins/AudioInsights/Source/AudioInsights/Private/Providers/AudioBusTraceProvider.h
@@ -38,6 +38,9 @@ namespace UE::Audio::Insights
 		DECLARE_DELEGATE(FOnTimeMarkerUpdated);
 		FOnTimeMarkerUpdated OnTimeMarkerUpdated;
 
+	protected:
+		virtual void OnTraceChannelsEnabled() override;
+
 	private:
 		virtual bool ProcessMessages() override;
 
diff --git a/Engine/Plugins/AudioInsights/Source/AudioInsights/Private/Providers/AudioEventLogTraceProvider.h b/Engine/Plugins/AudioInsights/Source/AudioInsights/Private/Providers/AudioEventLogTraceProvider.h
--- a/Engine/Plugins/AudioInsights/Source/AudioInsights/Private/Providers/AudioEventLogTraceProvider.h
+++ b/Engine/Plugins/AudioInsights/Source/AudioInsights/Private/Providers/AudioEventLogTraceProvider.h
@@ -22,6 +22,7 @@ namespace UE::Audio::Insights
 
 	protected:
 		virtual bool IsMessageProcessingPaused() const override;
+		virtual void OnTraceChannelsEnabled() override { /* Do Nothing */ }
 
 	private:
 		virtual bool ProcessMessages() override;
diff --git a/Engine/Plugins/AudioInsights/Source/AudioInsights/Private/Providers/SubmixTraceProvider.cpp b/Engine/Plugins/AudioInsights/Source/AudioInsights/Private/Providers/SubmixTraceProvider.cpp
--- a/Engine/Plugins/AudioInsights/Source/AudioInsights/Private/Providers/SubmixTraceProvider.cpp
+++ b/Engine/Plugins/AudioInsights/Source/AudioInsights/Private/Providers/SubmixTraceProvider.cpp
@@ -200,6 +200,17 @@ namespace UE::Audio::Insights
 		return new FSubmixTraceAnalyzer(AsShared(), InSession);
 	}
 
+	void FSubmixTraceProvider::OnTraceChannelsEnabled()
+	{
+		DeviceDataMap.Empty();
+		FTraceProviderBase::OnTraceChannelsEnabled();
+
+#if WITH_EDITOR
+		SubmixAssetProvider.RequestEntriesUpdate();
+		OnSubmixListUpdated.ExecuteIfBound();
+#endif // WITH_EDITOR
+	}
+
 	bool FSubmixTraceProvider::ProcessMessages()
 	{
 		// Helper lambdas
diff --git a/Engine/Plugins/AudioInsights/Source/AudioInsights/Private/Providers/SubmixTraceProvider.h b/Engine/Plugins/AudioInsights/Source/AudioInsights/Private/Providers/SubmixTraceProvider.h
--- a/Engine/Plugins/AudioInsights/Source/AudioInsights/Private/Providers/SubmixTraceProvider.h
+++ b/Engine/Plugins/AudioInsights/Source/AudioInsights/Private/Providers/SubmixTraceProvider.h
@@ -38,6 +38,9 @@ namespace UE::Audio::Insights
 		DECLARE_DELEGATE(FOnTimeMarkerUpdated);
 		FOnTimeMarkerUpdated OnTimeMarkerUpdated;
 
+	protected:
+		virtual void OnTraceChannelsEnabled() override;
+
 	private:
 		virtual bool ProcessMessages() override;
 
diff --git a/Engine/Plugins/AudioInsights/Source/AudioInsights/Private/Providers/VirtualLoopTraceProvider.cpp b/Engine/Plugins/AudioInsights/Source/AudioInsights/Private/Providers/VirtualLoopTraceProvider.cpp
--- a/Engine/Plugins/AudioInsights/Source/AudioInsights/Private/Providers/VirtualLoopTraceProvider.cpp
+++ b/Engine/Plugins/AudioInsights/Source/AudioInsights/Private/Providers/VirtualLoopTraceProvider.cpp
@@ -304,6 +304,7 @@ namespace UE::Audio::Insights
 				Builder.RouteEvent(RouteId_Stop, "Audio", "VirtualLoopStopOrRealize");
 				Builder.RouteEvent(RouteId_Update, "Audio", "VirtualLoopUpdate");
 				Builder.RouteEvent(RouteId_Virtualize, "Audio", "VirtualLoopVirtualize");
+				Builder.RouteEvent(RouteId_IsStillVirtualizedPing, "Audio", "VirtualLoopIsVirtualizedPing");
 			}
 
 			virtual bool OnHandleEvent(uint16 RouteId, EStyle Style, const FOnEventContext& Context) override
@@ -314,6 +315,7 @@ namespace UE::Audio::Insights
 				switch (RouteId)
 				{
 					case RouteId_Virtualize:
+					case RouteId_IsStillVirtualizedPing:
 					{
 #if WITH_EDITOR
 						CacheMessage<FVirtualLoopVirtualizeMessage>(Context, Messages.VirtualizeMessages);
@@ -364,7 +366,8 @@ namespace UE::Audio::Insights
 			{
 				RouteId_Virtualize,
 				RouteId_Update,
-				RouteId_Stop
+				RouteId_Stop,
+				RouteId_IsStillVirtualizedPing
 			};
 
 			TraceServices::IAnalysisSession& Session;
diff --git a/Engine/Plugins/AudioInsights/Source/AudioInsights/Public/AudioInsightsTraceProviderBase.h b/Engine/Plugins/AudioInsights/Source/AudioInsights/Public/AudioInsightsTraceProviderBase.h
--- a/Engine/Plugins/AudioInsights/Source/AudioInsights/Public/AudioInsightsTraceProviderBase.h
+++ b/Engine/Plugins/AudioInsights/Source/AudioInsights/Public/AudioInsightsTraceProviderBase.h
@@ -41,6 +41,12 @@ namespace UE::Audio::Insights
 			LastMessageId = 0;
 		}
 
+		virtual void OnTraceChannelsEnabled()
+		{
+			LastUpdateId = 0;
+			LastMessageId = 0;
+		}
+
 		virtual bool ProcessMessages()
 		{
 			LastUpdateId = LastMessageId;
@@ -133,7 +139,22 @@ namespace UE::Audio::Insights
 		private:
 			AUDIOINSIGHTS_API virtual bool OnEvent(uint16 RouteId, EStyle Style, const FOnEventContext& Context) override;
 
+			bool HandleChannelAnnounceEvent(uint16 RouteId, EStyle Style, const FOnEventContext& Context);
+			bool HandleChannelToggleEvent(uint16 RouteId, EStyle Style, const FOnEventContext& Context);
+
 			TSharedRef<FTraceProviderBase> Provider;
+
+			enum : uint16
+			{
+				RouteID_ChannelAnnounce = TNumericLimits<uint16>::Max() - 1u,
+				RouteID_ChannelToggle = TNumericLimits<uint16>::Max()
+			};
+
+			TOptional<uint32> AudioChannelID;
+			TOptional<uint32> AudioMixerChannelID;
+
+			bool bAudioChannelIsEnabled = false;
+			bool bAudioMixerChannelIsEnabled = false;
 		};
 
 		uint64 GetLastMessageId() const
@@ -237,6 +258,12 @@ namespace UE::Audio::Insights
 			FTraceProviderBase::Reset();
 		}
 
+		virtual void OnTraceChannelsEnabled() override
+		{
+			DeviceDataMap.Empty();
+			FTraceProviderBase::OnTraceChannelsEnabled();
+		}
+
 		template <typename TMsgType>
 		void ProcessMessageQueue(
 			TAnalyzerMessageQueue<TMsgType>& InQueue,
diff --git a/Engine/Source/Runtime/Engine/Private/AudioVirtualLoop.cpp b/Engine/Source/Runtime/Engine/Private/AudioVirtualLoop.cpp
--- a/Engine/Source/Runtime/Engine/Private/AudioVirtualLoop.cpp
+++ b/Engine/Source/Runtime/Engine/Private/AudioVirtualLoop.cpp
@@ -51,6 +51,14 @@ UE_TRACE_EVENT_BEGIN(Audio, VirtualLoopVirtualize)
 	UE_TRACE_EVENT_FIELD(UE::Trace::WideString, Name)
 UE_TRACE_EVENT_END()
 
+UE_TRACE_EVENT_BEGIN(Audio, VirtualLoopIsVirtualizedPing)
+	UE_TRACE_EVENT_FIELD(uint32, DeviceId)
+	UE_TRACE_EVENT_FIELD(uint64, Timestamp)
+	UE_TRACE_EVENT_FIELD(uint32, PlayOrder)
+	UE_TRACE_EVENT_FIELD(uint64, ComponentId)
+	UE_TRACE_EVENT_FIELD(UE::Trace::WideString, Name)
+UE_TRACE_EVENT_END()
+
 UE_TRACE_EVENT_BEGIN(Audio, VirtualLoopUpdate)
 	UE_TRACE_EVENT_FIELD(uint32, DeviceId)
 	UE_TRACE_EVENT_FIELD(uint64, Timestamp)
@@ -298,6 +306,10 @@ bool FAudioVirtualLoop::Update(float DeltaTime, bool bForceUpdate)
 	ActiveSound->PlaybackTime += DeltaTimePitchCorrected;
 	TimeVirtualized += DeltaTimePitchCorrected;
 
+#if UE_AUDIO_PROFILERTRACE_ENABLED
+	IsVirtualizedPingTimeSec += DeltaTime;
+#endif
+
 	// Update playback time estimates for each wave instance for Seek Restart, factoring in pitch.
 	// This set only has elements if the ActiveSound was set to Seek Restart upon initial virtualization.
 	for (TPair<UPTRINT, float>& Pair : EffectivePlaybackTimes)
@@ -370,6 +382,19 @@ bool FAudioVirtualLoop::Update(float DeltaTime, bool bForceUpdate)
 			const FTransform& Transform = ActiveSound->Transform;
 			const FVector& Location = Transform.GetLocation();
 			const FRotator Rotator = Transform.GetRotation().Rotator();
+
+			if (IsVirtualizedPingTimeSec >= IsVirtualizedPingWaitTimeSec)
+			{
+				UE_TRACE_LOG(Audio, VirtualLoopIsVirtualizedPing, AudioChannel)
+					<< VirtualLoopIsVirtualizedPing.DeviceId(static_cast<uint32>(AudioDevice->DeviceID))
+					<< VirtualLoopIsVirtualizedPing.Timestamp(FPlatformTime::Cycles64())
+					<< VirtualLoopIsVirtualizedPing.PlayOrder(ActiveSound->GetPlayOrder())
+					<< VirtualLoopIsVirtualizedPing.ComponentId(ActiveSound->GetAudioComponentID())
+					<< VirtualLoopIsVirtualizedPing.Name(Sound ? *Sound->GetPathName() : TEXT("N/A"));
+
+				IsVirtualizedPingTimeSec = 0.0f;
+			}
+
 			UE_TRACE_LOG(Audio, VirtualLoopUpdate, AudioChannel)
 				<< VirtualLoopUpdate.DeviceId(static_cast<uint32>(AudioDevice->DeviceID))
 				<< VirtualLoopUpdate.Timestamp(FPlatformTime::Cycles64())
diff --git a/Engine/Source/Runtime/Engine/Public/AudioVirtualLoop.h b/Engine/Source/Runtime/Engine/Public/AudioVirtualLoop.h
--- a/Engine/Source/Runtime/Engine/Public/AudioVirtualLoop.h
+++ b/Engine/Source/Runtime/Engine/Public/AudioVirtualLoop.h
@@ -16,6 +16,11 @@ private:
 	float TimeVirtualized;
 	float UpdateInterval;
 
+#if UE_AUDIO_PROFILERTRACE_ENABLED
+	float IsVirtualizedPingTimeSec = 0.0;
+	static constexpr float IsVirtualizedPingWaitTimeSec = 1.0;
+#endif // UE_AUDIO_PROFILERTRACE_ENABLED
+
 	FActiveSound* ActiveSound;
 
 	/**
