[Audio Insights] Fix plots not drawing over 5s when inspecting .utrace files
* Standalone was using the circular buffers inside the Sound Entries to store all plot data received when loading in trace files. These circular buffers are limited in size and we shouldn't really be attempting to store all plot data for a sound inside them at once.
* Move standalone closer to editor version by pulling out plot data from the existing SessionCacheMessages (as opposed to the cache manager) when the time marker moves.

diff --git a/Engine/Plugins/AudioInsights/Source/AudioInsights/Private/AudioInsightsTimingViewExtender.cpp b/Engine/Plugins/AudioInsights/Source/AudioInsights/Private/AudioInsightsTimingViewExtender.cpp
--- a/Engine/Plugins/AudioInsights/Source/AudioInsights/Private/AudioInsightsTimingViewExtender.cpp
+++ b/Engine/Plugins/AudioInsights/Source/AudioInsights/Private/AudioInsightsTimingViewExtender.cpp
@@ -59,7 +59,7 @@ namespace UE::Audio::Insights
 		}
 		else
 		{
-			PlottingRange = TRange<double>(Timestamp - 5.0, Timestamp);
+			PlottingRange = TRange<double>(Timestamp - MaxPlottingHistorySeconds, Timestamp);
 		}
 
 		// If you set the time marker on the TimingView, it will Broadcast OnTimingViewTimeMarkerChanged. Otherwise we gotta do it ourselves.
@@ -91,6 +91,7 @@ namespace UE::Audio::Insights
 			if (!bUserInputDetected)
 			{
 				SystemControllingTime = ESystemControllingTimeMarker::External;
+				PlottingRange = TRange<double>(InTimeMarker - MaxPlottingHistorySeconds, InTimeMarker);
 			}
 
 			OnTimingViewTimeMarkerChanged.Broadcast(InTimeMarker);
diff --git a/Engine/Plugins/AudioInsights/Source/AudioInsights/Private/AudioInsightsTimingViewExtender.h b/Engine/Plugins/AudioInsights/Source/AudioInsights/Private/AudioInsightsTimingViewExtender.h
--- a/Engine/Plugins/AudioInsights/Source/AudioInsights/Private/AudioInsightsTimingViewExtender.h
+++ b/Engine/Plugins/AudioInsights/Source/AudioInsights/Private/AudioInsightsTimingViewExtender.h
@@ -59,6 +59,10 @@ namespace UE::Audio::Insights
 		DECLARE_MULTICAST_DELEGATE(FOnTimeControlMethodReset);
 		FOnTimeControlMethodReset OnTimeControlMethodReset;
 
+		// Move to AudioInsightsConstants
+		static constexpr double MaxPlottingHistorySeconds = 5.0;
+		static constexpr double PlottingMarginSeconds = 0.2;
+
 	private:
 		void OnTimeMarkerChanged(UE::Insights::Timing::ETimeChangedFlags InFlags, double InTimeMarker);
 
@@ -67,7 +71,7 @@ namespace UE::Audio::Insights
 
 		ECacheAndProcess CacheAndProcessMessageStatus = ECacheAndProcess::Latest;
 		TOptional<ESystemControllingTimeMarker> SystemControllingTime;
-		TRange<double> PlottingRange { 0.0, 5.0 };
+		TRange<double> PlottingRange { 0.0, MaxPlottingHistorySeconds };
 		double TraceDurationSeconds = 0.0;
 		bool bUserInputDetected = false;
 	};
diff --git a/Engine/Plugins/AudioInsights/Source/AudioInsights/Private/Providers/SoundTraceProvider.cpp b/Engine/Plugins/AudioInsights/Source/AudioInsights/Private/Providers/SoundTraceProvider.cpp
--- a/Engine/Plugins/AudioInsights/Source/AudioInsights/Private/Providers/SoundTraceProvider.cpp
+++ b/Engine/Plugins/AudioInsights/Source/AudioInsights/Private/Providers/SoundTraceProvider.cpp
@@ -188,27 +188,29 @@ namespace UE::Audio::Insights
 
 #if !WITH_EDITOR
 		template<typename T>
-		const T* FindClosestMessageToTimestamp(const TraceServices::TPagedArray<T>& InCachedMessages, const double InTimeMarker, const PlayOrderPair InPlayOrderPair)
+		void IterateOverRange(const TraceServices::TPagedArray<T>& InCachedMessages, const double Start, const double End, TFunctionRef<void(const T& Message)> OnMessageFunc)
 		{
-			const int32 ClosestMessageToTimeStampIndex = TraceServices::PagedArrayAlgo::BinarySearchClosestBy(InCachedMessages, InTimeMarker, [](const T& Msg) { return Msg.Timestamp; });
+			if (InCachedMessages.Num() == 0)
+			{
+				return;
+			}
 
-			// Iterate backwards from TimeMarker until we find the matching PlayOrder
-			for (auto It = InCachedMessages.GetIteratorFromItem(ClosestMessageToTimeStampIndex); It; --It)
+			const int32 ClosestMessageToTimeStampIndex = TraceServices::PagedArrayAlgo::BinarySearchClosestBy(InCachedMessages, Start, [](const T& Msg) { return Msg.Timestamp; });
+			
+			for (auto It = InCachedMessages.GetIteratorFromItem(ClosestMessageToTimeStampIndex); It != InCachedMessages.end(); ++It)
 			{
-				if (InPlayOrderPair.WaveInstancePlayOrderID != INDEX_NONE)
+				if (It->Timestamp > End)
 				{
-					if (It->WaveInstancePlayOrder == InPlayOrderPair.WaveInstancePlayOrderID)
-					{
-						return &(*It);
-					}
+					break;
 				}
-				else if (It->ActiveSoundPlayOrder == InPlayOrderPair.ActiveSoundPlayOrderID)
+
+				if (It->Timestamp < Start)
 				{
-					return &(*It);
+					continue;
 				}
-			}
 
-			return nullptr;
+				OnMessageFunc(*It);
+			}
 		}
 #endif // !WITH_EDITOR
 	}
@@ -1501,8 +1503,8 @@ namespace UE::Audio::Insights
 #if WITH_EDITOR
 		CollectParamsForTimestamp(CacheManager, TimeMarker);
 #else
-		CollectParamsForTimestampStandalone(TimeMarker);
-#endif //!WITH_EDITOR
+		CollectParamsForTimestamp(TimeMarker);
+#endif // WITH_EDITOR
 
 		// Call parent method to update LastMessageId
 		FTraceProviderBase::OnTimingViewTimeMarkerChanged(TimeMarker);
@@ -1515,6 +1517,9 @@ namespace UE::Audio::Insights
 
 #if WITH_EDITOR
 	void FSoundTraceProvider::CollectParamsForTimestamp(const FAudioInsightsCacheManager& CacheManager, const double TimeMarker)
+#else
+	void FSoundTraceProvider::CollectParamsForTimestamp(const double TimeMarker)
+#endif // WITH_EDITOR
 	{
 		using namespace FSoundTraceProviderPrivate;
 
@@ -1548,18 +1553,25 @@ namespace UE::Audio::Insights
 		const TRange<double> PlottingRange = AudioInsightsModule.GetTimingViewExtender().GetPlottingRange();
 
 		// Add a small margin to either side to avoid gaps at the edges
-		constexpr double Margin = 0.2;
-		const double StartTime = FMath::Max(0.0f, PlottingRange.GetLowerBoundValue() - Margin);
-		const double EndTime = FMath::Max(0.0f, PlottingRange.GetUpperBoundValue() + Margin);
+		const double StartTime = FMath::Max(0.0f, PlottingRange.GetLowerBoundValue() - FAudioInsightsTimingViewExtender::PlottingMarginSeconds);
+		const double EndTime = FMath::Max(0.0f, PlottingRange.GetUpperBoundValue() + FAudioInsightsTimingViewExtender::PlottingMarginSeconds);
 
 		// Using ParallelFor to speed-up the cached messages retrieval, using a traditional for loop is unacceptably slower, specially in large traces.
 		ParallelFor(NumMessageTypesToProcess,
+#if WITH_EDITOR
 		[&CacheManager, TimeMarker, StartTime, EndTime, this](const int32 Index)
+#else
+		[TimeMarker, StartTime, EndTime, this](const int32 Index)
+#endif // WITH_EDITOR
 		{
 			switch (static_cast<ESoundParamMessageType>(Index))
 			{
 			case ESoundParamMessageType::Priority:
+#if WITH_EDITOR
 				CacheManager.IterateOverRange<FSoundPriorityMessage>(SoundMessageNames::PriorityParam, StartTime, EndTime, [this, TimeMarker](const FSoundPriorityMessage& Message)
+#else
+				IterateOverRange<FSoundPriorityMessage>(SessionCachedMessages->PriorityCachedMessages, StartTime, EndTime, [this, TimeMarker](const FSoundPriorityMessage& Message)
+#endif // WITH_EDITOR
 				{
 					TSharedPtr<FSoundDashboardEntry> ActiveSoundEntry;
 					GetActiveSoundEntryFromIDs(Message.ActiveSoundPlayOrder, ActiveSoundEntry);
@@ -1581,7 +1593,11 @@ namespace UE::Audio::Insights
 				break;
 
 			case ESoundParamMessageType::Distance:
+#if WITH_EDITOR
 				CacheManager.IterateOverRange<FSoundDistanceMessage>(SoundMessageNames::DistanceParam, StartTime, EndTime, [this, TimeMarker](const FSoundDistanceMessage& Message)
+#else
+				IterateOverRange<FSoundDistanceMessage>(SessionCachedMessages->DistanceCachedMessages, StartTime, EndTime, [this, TimeMarker](const FSoundDistanceMessage& Message)
+#endif // WITH_EDITOR
 				{
 					TSharedPtr<FSoundDashboardEntry> ActiveSoundEntry;
 					GetActiveSoundEntryFromIDs(Message.ActiveSoundPlayOrder, ActiveSoundEntry);
@@ -1603,7 +1619,11 @@ namespace UE::Audio::Insights
 				break;
 
 			case ESoundParamMessageType::DistanceAttenuation:
+#if WITH_EDITOR
 				CacheManager.IterateOverRange<FSoundDistanceAttenuationMessage>(SoundMessageNames::DistanceAttenuationParam, StartTime, EndTime, [this, TimeMarker](const FSoundDistanceAttenuationMessage& Message)
+#else
+				IterateOverRange<FSoundDistanceAttenuationMessage>(SessionCachedMessages->DistanceAttenuationCachedMessages, StartTime, EndTime, [this, TimeMarker](const FSoundDistanceAttenuationMessage& Message)
+#endif // WITH_EDITOR
 				{
 					TSharedPtr<FSoundDashboardEntry> ActiveSoundEntry;
 					GetActiveSoundEntryFromIDs(Message.ActiveSoundPlayOrder, ActiveSoundEntry);
@@ -1625,7 +1645,11 @@ namespace UE::Audio::Insights
 				break;
 
 			case ESoundParamMessageType::LPFFreq:
+#if WITH_EDITOR
 				CacheManager.IterateOverRange<FSoundLPFFreqMessage>(SoundMessageNames::LPFFreqParam, StartTime, EndTime, [this, TimeMarker](const FSoundLPFFreqMessage& Message)
+#else
+				IterateOverRange<FSoundLPFFreqMessage>(SessionCachedMessages->LPFFreqCachedMessages, StartTime, EndTime, [this, TimeMarker](const FSoundLPFFreqMessage& Message)
+#endif // WITH_EDITOR
 				{
 					TSharedPtr<FSoundDashboardEntry> ActiveSoundEntry;
 					GetActiveSoundEntryFromIDs(Message.ActiveSoundPlayOrder, ActiveSoundEntry);
@@ -1647,7 +1671,11 @@ namespace UE::Audio::Insights
 				break;
 
 			case ESoundParamMessageType::HPFFreq:
+#if WITH_EDITOR
 				CacheManager.IterateOverRange<FSoundHPFFreqMessage>(SoundMessageNames::HPFFreqParam, StartTime, EndTime, [this, TimeMarker](const FSoundHPFFreqMessage& Message)
+#else
+				IterateOverRange<FSoundHPFFreqMessage>(SessionCachedMessages->HPFFreqCachedMessages, StartTime, EndTime, [this, TimeMarker](const FSoundHPFFreqMessage& Message)
+#endif // WITH_EDITOR
 				{
 					TSharedPtr<FSoundDashboardEntry> ActiveSoundEntry;
 					GetActiveSoundEntryFromIDs(Message.ActiveSoundPlayOrder, ActiveSoundEntry);
@@ -1669,7 +1697,11 @@ namespace UE::Audio::Insights
 				break;
 
 			case ESoundParamMessageType::Envelope:
+#if WITH_EDITOR
 				CacheManager.IterateOverRange<FSoundEnvelopeMessage>(SoundMessageNames::EnvelopeParam, StartTime, EndTime, [this, TimeMarker](const FSoundEnvelopeMessage& Message)
+#else
+				IterateOverRange<FSoundEnvelopeMessage>(SessionCachedMessages->AmplitudeCachedMessages, StartTime, EndTime, [this, TimeMarker](const FSoundEnvelopeMessage& Message)
+#endif // WITH_EDITOR
 				{
 					TSharedPtr<FSoundDashboardEntry> ActiveSoundEntry;
 					GetActiveSoundEntryFromIDs(Message.ActiveSoundPlayOrder, ActiveSoundEntry);
@@ -1691,7 +1723,11 @@ namespace UE::Audio::Insights
 				break;
 
 			case ESoundParamMessageType::Volume:
+#if WITH_EDITOR
 				CacheManager.IterateOverRange<FSoundVolumeMessage>(SoundMessageNames::VolumeParam, StartTime, EndTime, [this, TimeMarker](const FSoundVolumeMessage& Message)
+#else
+				IterateOverRange<FSoundVolumeMessage>(SessionCachedMessages->VolumeCachedMessages, StartTime, EndTime, [this, TimeMarker](const FSoundVolumeMessage& Message)
+#endif // WITH_EDITOR
 				{
 					TSharedPtr<FSoundDashboardEntry> ActiveSoundEntry;
 					GetActiveSoundEntryFromIDs(Message.ActiveSoundPlayOrder, ActiveSoundEntry);
@@ -1713,7 +1749,11 @@ namespace UE::Audio::Insights
 				break;
 
 			case ESoundParamMessageType::Pitch:
+#if WITH_EDITOR
 				CacheManager.IterateOverRange<FSoundPitchMessage>(SoundMessageNames::PitchParam, StartTime, EndTime, [this, TimeMarker](const FSoundPitchMessage& Message)
+#else
+				IterateOverRange<FSoundPitchMessage>(SessionCachedMessages->PitchCachedMessages, StartTime, EndTime, [this, TimeMarker](const FSoundPitchMessage& Message)
+#endif // WITH_EDITOR
 				{
 					TSharedPtr<FSoundDashboardEntry> ActiveSoundEntry;
 					GetActiveSoundEntryFromIDs(Message.ActiveSoundPlayOrder, ActiveSoundEntry);
@@ -1735,7 +1775,11 @@ namespace UE::Audio::Insights
 				break;
 
 			case ESoundParamMessageType::RelativeRenderCost:
+#if WITH_EDITOR
 				CacheManager.IterateOverRange<FSoundRelativeRenderCostMessage>(SoundMessageNames::RelativeRenderCostParam, StartTime, EndTime, [this, TimeMarker](const FSoundRelativeRenderCostMessage& Message)
+#else
+				IterateOverRange<FSoundRelativeRenderCostMessage>(SessionCachedMessages->RelativeRenderCostCachedMessages, StartTime, EndTime, [this, TimeMarker](const FSoundRelativeRenderCostMessage& Message)
+#endif // WITH_EDITOR
 				{
 					TSharedPtr<FSoundDashboardEntry> ActiveSoundEntry;
 					GetActiveSoundEntryFromIDs(Message.ActiveSoundPlayOrder, ActiveSoundEntry);
@@ -1765,177 +1809,6 @@ namespace UE::Audio::Insights
 
 		OnProcessPlotData.Broadcast();
 	}
-#else
-	void FSoundTraceProvider::CollectParamsForTimestampStandalone(const double TimeMarker)
-	{
-		using namespace FSoundTraceProviderPrivate;
-
-		const FDeviceData* DeviceData = FindFilteredDeviceData();
-		if (DeviceData)
-		{
-			TArray<PlayOrderPair> PlayOrders;
-
-			for (const auto& [EntryType, SoundDashboardEntry] : *DeviceData)
-			{
-				CollectPlayOrderPairsRecursive(SoundDashboardEntry, PlayOrders);
-			}
-
-			// Collect messages (volume, pitch, etc.) from active sounds (based on active sounds's PlayOrder)
-			struct CachedEntryInfo
-			{
-				TOptional<FSoundPriorityMessage> PriorityMessage;
-				TOptional<FSoundDistanceMessage> DistanceMessage;
-				TOptional<FSoundDistanceAttenuationMessage> DistanceAttenuationMessage;
-				TOptional<FSoundLPFFreqMessage> LPFFreqMessage;
-				TOptional<FSoundHPFFreqMessage> HPFFreqMessage;
-				TOptional<FSoundEnvelopeMessage> AmplitudeMessage;
-				TOptional<FSoundVolumeMessage> VolumeMessage;
-				TOptional<FSoundPitchMessage> PitchMessage;
-				TOptional<FSoundRelativeRenderCostMessage> RelativeRenderCostMessage;
-			};
-
-			TArray<CachedEntryInfo> CachedEntryInfos;
-			CachedEntryInfos.SetNum(PlayOrders.Num());
-
-			// Using ParallelFor to speed-up the cached messages retrieval, using a traditional for loop is unacceptably slower, specially in large traces.
-			ParallelFor(PlayOrders.Num(),
-			[&PlayOrders, &CachedEntryInfos, TimeMarker, this](const int32 Index)
-			{
-				const PlayOrderPair PlayOrderPair = PlayOrders[Index];
-
-				// Priority
-				const FSoundPriorityMessage* FoundPriorityCachedMessage = FindClosestMessageToTimestamp(SessionCachedMessages->PriorityCachedMessages, TimeMarker, PlayOrderPair);
-				if (FoundPriorityCachedMessage)
-				{
-					CachedEntryInfos[Index].PriorityMessage = *FoundPriorityCachedMessage;
-				}
-
-				// Distance
-				const FSoundDistanceMessage* FoundDistanceCachedMessage = FindClosestMessageToTimestamp(SessionCachedMessages->DistanceCachedMessages, TimeMarker, PlayOrderPair);
-				if (FoundDistanceCachedMessage)
-				{
-					CachedEntryInfos[Index].DistanceMessage = *FoundDistanceCachedMessage;
-				}
-
-				// Distance Attenuation
-				const FSoundDistanceAttenuationMessage* FoundDistanceAttenuationCachedMessage = FindClosestMessageToTimestamp(SessionCachedMessages->DistanceAttenuationCachedMessages, TimeMarker, PlayOrderPair);
-				if (FoundDistanceAttenuationCachedMessage)
-				{
-					CachedEntryInfos[Index].DistanceAttenuationMessage = *FoundDistanceAttenuationCachedMessage;
-				}
-
-				// LPF Freq
-				const FSoundLPFFreqMessage* FoundLPFFreqCachedMessage = FindClosestMessageToTimestamp(SessionCachedMessages->LPFFreqCachedMessages, TimeMarker, PlayOrderPair);
-				if (FoundLPFFreqCachedMessage)
-				{
-					CachedEntryInfos[Index].LPFFreqMessage = *FoundLPFFreqCachedMessage;
-				}
-
-				// HPF Freq
-				const FSoundHPFFreqMessage* FoundHPFFreqCachedMessage = FindClosestMessageToTimestamp(SessionCachedMessages->HPFFreqCachedMessages, TimeMarker, PlayOrderPair);
-				if (FoundHPFFreqCachedMessage)
-				{
-					CachedEntryInfos[Index].HPFFreqMessage = *FoundHPFFreqCachedMessage;
-				}
-
-				// Amplitude
-				const FSoundEnvelopeMessage* FoundAmplitudeCachedMessage = FindClosestMessageToTimestamp(SessionCachedMessages->AmplitudeCachedMessages, TimeMarker, PlayOrderPair);
-				if (FoundAmplitudeCachedMessage)
-				{
-					CachedEntryInfos[Index].AmplitudeMessage = *FoundAmplitudeCachedMessage;
-				}
-
-				// Volume
-				const FSoundVolumeMessage* FoundVolumeCachedMessage = FindClosestMessageToTimestamp(SessionCachedMessages->VolumeCachedMessages, TimeMarker, PlayOrderPair);
-				if (FoundVolumeCachedMessage)
-				{
-					CachedEntryInfos[Index].VolumeMessage = *FoundVolumeCachedMessage;
-				}
-
-				// Pitch
-				const FSoundPitchMessage* FoundPitchCachedMessage = FindClosestMessageToTimestamp(SessionCachedMessages->PitchCachedMessages, TimeMarker, PlayOrderPair);
-				if (FoundPitchCachedMessage)
-				{
-					CachedEntryInfos[Index].PitchMessage = *FoundPitchCachedMessage;
-				}
-
-				// Relative Render Cost
-				const FSoundRelativeRenderCostMessage* FoundRelativeRenderCostCachedMessage = FindClosestMessageToTimestamp(SessionCachedMessages->RelativeRenderCostCachedMessages, TimeMarker, PlayOrderPair);
-				if (FoundRelativeRenderCostCachedMessage)
-				{
-					CachedEntryInfos[Index].RelativeRenderCostMessage = *FoundRelativeRenderCostCachedMessage;
-				}
-			});
-
-			// Update the device entries with the collected info
-			for (const CachedEntryInfo& CachedEntryInfo : CachedEntryInfos)
-			{
-				if (!CachedEntryInfo.PriorityMessage.IsSet())
-				{
-					continue;
-				}
-
-				TSharedPtr<FSoundDashboardEntry> ActiveSoundEntry = nullptr;
-				GetActiveSoundEntryFromIDs(CachedEntryInfo.PriorityMessage.GetValue().ActiveSoundPlayOrder, ActiveSoundEntry);
-				
-				if (!ActiveSoundEntry.IsValid())
-				{
-					continue;
-				}
-
-				FSoundDashboardEntry* SoundEntry = FSoundTraceProviderPrivate::FindWaveEntryInActiveSound(CachedEntryInfo.PriorityMessage.GetValue().WaveInstancePlayOrder, &ActiveSoundEntry);
-				if (SoundEntry == nullptr)
-				{
-					continue;
-				}
-
-				SoundEntry->PriorityDataPoint = CachedEntryInfo.PriorityMessage.GetValue().Priority;
-
-				if (CachedEntryInfo.DistanceMessage.IsSet())
-				{
-					SoundEntry->DistanceDataPoint = CachedEntryInfo.DistanceMessage.GetValue().Distance;
-				}
-
-				if (CachedEntryInfo.DistanceAttenuationMessage.IsSet())
-				{
-					SoundEntry->DistanceAttenuationDataPoint = CachedEntryInfo.DistanceAttenuationMessage.GetValue().DistanceAttenuation;
-				}
-
-				if (CachedEntryInfo.LPFFreqMessage.IsSet())
-				{
-					SoundEntry->LPFFreqDataPoint = CachedEntryInfo.LPFFreqMessage.GetValue().LPFFrequency;
-				}
-
-				if (CachedEntryInfo.HPFFreqMessage.IsSet())
-				{
-					SoundEntry->HPFFreqDataPoint = CachedEntryInfo.HPFFreqMessage.GetValue().HPFFrequency;
-				}
-
-				if (CachedEntryInfo.AmplitudeMessage.IsSet())
-				{
-					SoundEntry->AmplitudeDataPoint = CachedEntryInfo.AmplitudeMessage.GetValue().Envelope;
-				}
-
-				if (CachedEntryInfo.VolumeMessage.IsSet())
-				{
-					SoundEntry->VolumeDataPoint = CachedEntryInfo.VolumeMessage.GetValue().Volume;
-				}
-
-				if (CachedEntryInfo.PitchMessage.IsSet())
-				{
-					SoundEntry->PitchDataPoint = CachedEntryInfo.PitchMessage.GetValue().Pitch;
-				}
-
-				if (CachedEntryInfo.RelativeRenderCostMessage.IsSet())
-				{
-					SoundEntry->RelativeRenderCostDataPoint = CachedEntryInfo.RelativeRenderCostMessage.GetValue().RelativeRenderCost;
-				}
-
-				SoundEntry->Timestamp = CachedEntryInfo.PriorityMessage.GetValue().Timestamp;
-			}
-		}
-	}
-#endif // WITH_EDITOR
 
 	void FSoundTraceProvider::CacheIsPlottingFlagRecursive(const TSharedPtr<IDashboardDataTreeViewEntry>& Entry)
 	{
diff --git a/Engine/Plugins/AudioInsights/Source/AudioInsights/Private/Providers/SoundTraceProvider.h b/Engine/Plugins/AudioInsights/Source/AudioInsights/Private/Providers/SoundTraceProvider.h
--- a/Engine/Plugins/AudioInsights/Source/AudioInsights/Private/Providers/SoundTraceProvider.h
+++ b/Engine/Plugins/AudioInsights/Source/AudioInsights/Private/Providers/SoundTraceProvider.h
@@ -55,8 +55,7 @@ namespace UE::Audio::Insights
 #if WITH_EDITOR
 		void CollectParamsForTimestamp(const FAudioInsightsCacheManager& CacheManager, const double TimeMarker);
 #else
-		void CollectParamsForTimestampStandalone(const double TimeMarker);
-		void CollectParameterPlayOrderIDsRecursive(const TSharedPtr<IDashboardDataTreeViewEntry>& Entry, TArray<uint32>& OutPlayOrderArray);
+		void CollectParamsForTimestamp(const double TimeMarker);
 #endif // WITH_EDITOR
 
 		void CacheIsPlottingFlagRecursive(const TSharedPtr<IDashboardDataTreeViewEntry>& Entry);
diff --git a/Engine/Plugins/AudioInsights/Source/AudioInsights/Private/Views/SoundPlotsWidgetView.cpp b/Engine/Plugins/AudioInsights/Source/AudioInsights/Private/Views/SoundPlotsWidgetView.cpp
--- a/Engine/Plugins/AudioInsights/Source/AudioInsights/Private/Views/SoundPlotsWidgetView.cpp
+++ b/Engine/Plugins/AudioInsights/Source/AudioInsights/Private/Views/SoundPlotsWidgetView.cpp
@@ -109,13 +109,7 @@ namespace UE::Audio::Insights
 
 		if (!PlotWidgetMetadataPerCurve.IsValid())
 		{
-#if WITH_EDITOR
 			return;
-#else
-			// In standalone, we still want to process any plot data while the .utrace file is loading in
-			// Ensure we create the plot widget and start collecting data from the point we start receiving data
-			ConstructPlotsWidgetContents();
-#endif // WITH_EDITOR
 		}
 
 		// Process new data
@@ -132,15 +126,12 @@ namespace UE::Audio::Insights
 		}
 
 		// Remove old points and set curve data for each widget
-		const double PlotDrawMinLimitTimestamp = CurrentTimestamp - BeginTimestamp - (MaxPlotHistorySeconds + 0.2 /* extra grace time to avoid curve cuts being displayed */);
-#if WITH_EDITOR
-		
+		const double PlotDrawMinLimitTimestamp = CurrentTimestamp - BeginTimestamp - (FAudioInsightsTimingViewExtender::MaxPlottingHistorySeconds + FAudioInsightsTimingViewExtender::PlottingMarginSeconds /* extra grace time to avoid curve cuts being displayed */);
+#if !WITH_EDITOR
+		const double PlotDrawMaxLimitTimestamp = CurrentTimestamp + FAudioInsightsTimingViewExtender::PlottingMarginSeconds /* extra grace time to avoid curve cuts being displayed */;
+#endif // !WITH_EDITOR
+
 		const TSharedPtr<FPointDataPerCurveMap> CurveDataMapPtr = *PlotWidgetCurveIdToPointDataMapPerColumn.Find(SelectedPlotColumnName);
-#else
-		// Take a deep copy of the curve data in standalone - we need to retain the original data and filter which points to process
-		const TSharedPtr<FPointDataPerCurveMap> CurveDataMapPtr = MakeShared<FPointDataPerCurveMap>(*(*PlotWidgetCurveIdToPointDataMapPerColumn.Find(SelectedPlotColumnName)).Get());
-		const double PlotDrawMaxLimitTimestamp = CurrentTimestamp + 0.2 /* extra grace time to avoid curve cuts being displayed */;
-#endif // WITH_EDITOR
 
 		if (CurveDataMapPtr.IsValid())
 		{
@@ -576,15 +567,7 @@ namespace UE::Audio::Insights
 			}
 			else if (GameState == EGameState::ControlledByTimemarker)
 			{
-#if WITH_EDITOR
 				DataPoints->Empty(DataPoints->Num());
-#else
-				const TSharedPtr<const FAudioInsightsComponent> AudioInsightsComponent = FAudioInsightsModule::GetChecked().GetAudioInsightsComponent();
-				if (AudioInsightsComponent.IsValid() && AudioInsightsComponent->GetIsLiveSession())
-				{
-					DataPoints->Empty(DataPoints->Num());
-				}
-#endif // WITH_EDITOR
 			}
 
 			auto DataFunc = GetPlotColumnDataFunc(ColumnName);
@@ -858,7 +841,7 @@ namespace UE::Audio::Insights
 
 		if (BeginTimestamp == TNumericLimits<double>::Max())
 		{
-			return TRange<double>(0, MaxPlotHistorySeconds);
+			return TRange<double>(0, FAudioInsightsTimingViewExtender::MaxPlottingHistorySeconds);
 		}
 
 		double RangeUpperBound = 0.0;
@@ -888,7 +871,7 @@ namespace UE::Audio::Insights
 		const TSharedPtr<const FAudioInsightsComponent> AudioInsightsComponent = FAudioInsightsModule::GetChecked().GetAudioInsightsComponent();
 		if (AudioInsightsComponent.IsValid() && !AudioInsightsComponent->GetIsLiveSession())
 		{
-			CurrentRange = TRange<double>(CurrentTimestamp - MaxPlotHistorySeconds, CurrentTimestamp);
+			CurrentRange = TRange<double>(CurrentTimestamp - FAudioInsightsTimingViewExtender::MaxPlottingHistorySeconds, CurrentTimestamp);
 			return CurrentRange;
 		}
 		else
@@ -907,7 +890,7 @@ namespace UE::Audio::Insights
 		}
 #endif // WITH_EDITOR
 
-		CurrentRange = TRange<double>(RangeUpperBound - MaxPlotHistorySeconds, RangeUpperBound);
+		CurrentRange = TRange<double>(RangeUpperBound - FAudioInsightsTimingViewExtender::MaxPlottingHistorySeconds, RangeUpperBound);
 		return CurrentRange;
 	}
 
diff --git a/Engine/Plugins/AudioInsights/Source/AudioInsights/Public/Views/SoundPlotsWidgetView.h b/Engine/Plugins/AudioInsights/Source/AudioInsights/Public/Views/SoundPlotsWidgetView.h
--- a/Engine/Plugins/AudioInsights/Source/AudioInsights/Public/Views/SoundPlotsWidgetView.h
+++ b/Engine/Plugins/AudioInsights/Source/AudioInsights/Public/Views/SoundPlotsWidgetView.h
@@ -94,9 +94,6 @@ namespace UE::Audio::Insights
 
 		bool GetIsPlotEnabledForEntry(const IDashboardDataTreeViewEntry& InEntry) const;
 
-		// Maximum amount of data history kept for plots (in seconds)
-		inline static constexpr double MaxPlotHistorySeconds = 5.0;
-
 		// Curve points per timestamp per source id per column name 
 		TMap<FName, TSharedPtr<FPointDataPerCurveMap>> PlotWidgetCurveIdToPointDataMapPerColumn;
 
