[Audio Insights] Submix dashboard: fix for Dynamic Submix meters not being removed in Audio Meters tab when no longer triggered.

diff --git a/Engine/Plugins/AudioInsights/Source/AudioInsights/Private/Providers/SubmixTraceProvider.cpp b/Engine/Plugins/AudioInsights/Source/AudioInsights/Private/Providers/SubmixTraceProvider.cpp
--- a/Engine/Plugins/AudioInsights/Source/AudioInsights/Private/Providers/SubmixTraceProvider.cpp
+++ b/Engine/Plugins/AudioInsights/Source/AudioInsights/Private/Providers/SubmixTraceProvider.cpp
@@ -256,6 +256,8 @@ namespace UE::Audio::Insights
 				EntryRef.Name = Msg.Name;
 				EntryRef.AudioMeterInfo->NumChannels = Msg.NumChannels;
 				EntryRef.bIsMainSubmix = Msg.bIsMainSubmix;
+
+				OnSubmixLoaded.ExecuteIfBound(Msg.SubmixId);
 			}
 		});
 
@@ -326,7 +328,19 @@ namespace UE::Audio::Insights
 
 			if (OutEntry && (*OutEntry)->Timestamp < Msg.Timestamp)
 			{
+#if WITH_EDITOR
+				// Turn off submix entry activity, do not remove it since we are in Editor mode.
+				// We keep in the dashboard all registered submix assets from Content Browser
+				FSubmixDashboardEntry& EntryRef = *OutEntry->Get();
+
+				EntryRef.bHasActivity = false;
+
+				TArray<float>& EnvelopeValues = EntryRef.AudioMeterInfo->EnvelopeValues;
+				FMemory::Memzero(EnvelopeValues.GetData(), EnvelopeValues.Num() * sizeof(float));
+#else
+				OnSubmixRemoved.ExecuteIfBound(Msg.SubmixId);
 				RemoveDeviceEntry(Msg.DeviceId, Msg.SubmixId);
+#endif // WITH_EDITOR
 			}
 		});
 
diff --git a/Engine/Plugins/AudioInsights/Source/AudioInsights/Private/Providers/SubmixTraceProvider.h b/Engine/Plugins/AudioInsights/Source/AudioInsights/Private/Providers/SubmixTraceProvider.h
--- a/Engine/Plugins/AudioInsights/Source/AudioInsights/Private/Providers/SubmixTraceProvider.h
+++ b/Engine/Plugins/AudioInsights/Source/AudioInsights/Private/Providers/SubmixTraceProvider.h
@@ -31,6 +31,7 @@ namespace UE::Audio::Insights
 		DECLARE_DELEGATE_OneParam(FOnSubmixChanged, const uint32 /*SubmixId*/);
 		FOnSubmixChanged OnSubmixAdded;
 		FOnSubmixChanged OnSubmixRemoved;
+		FOnSubmixChanged OnSubmixLoaded;
 
 		DECLARE_DELEGATE(FOnSubmixListUpdated);
 		FOnSubmixListUpdated OnSubmixListUpdated;
diff --git a/Engine/Plugins/AudioInsights/Source/AudioInsights/Private/Views/SubmixDashboardViewFactory.cpp b/Engine/Plugins/AudioInsights/Source/AudioInsights/Private/Views/SubmixDashboardViewFactory.cpp
--- a/Engine/Plugins/AudioInsights/Source/AudioInsights/Private/Views/SubmixDashboardViewFactory.cpp
+++ b/Engine/Plugins/AudioInsights/Source/AudioInsights/Private/Views/SubmixDashboardViewFactory.cpp
@@ -49,6 +49,7 @@ namespace UE::Audio::Insights
 
 		SubmixProvider->OnSubmixAdded.BindRaw(this, &FSubmixDashboardViewFactory::HandleOnSubmixAdded);
 		SubmixProvider->OnSubmixRemoved.BindRaw(this, &FSubmixDashboardViewFactory::HandleOnSubmixRemoved);
+		SubmixProvider->OnSubmixLoaded.BindRaw(this, &FSubmixDashboardViewFactory::HandleOnSubmixLoaded);
 		SubmixProvider->OnSubmixListUpdated.BindRaw(this, &FSubmixDashboardViewFactory::RequestListRefresh);
 		SubmixProvider->OnTimeMarkerUpdated.BindRaw(this, &FSubmixDashboardViewFactory::HandleOnTimeMarkerUpdated);
 	}
@@ -240,6 +241,11 @@ namespace UE::Audio::Insights
 			UpdateAudioMeters();
 			bListRefreshed = false;
 		}
+		else if (bShouldReactivateAudioMeters)
+		{
+			ReactivateAudioMeters();
+			bShouldReactivateAudioMeters = false;
+		}
 	}
 
 	TSharedRef<SWidget> FSubmixDashboardViewFactory::MakeWidget(TSharedRef<SDockTab> OwnerTab, const FSpawnTabArgs& SpawnTabArgs)
@@ -473,6 +479,42 @@ namespace UE::Audio::Insights
 		RequestListRefresh();
 	}
 
+	void FSubmixDashboardViewFactory::ReactivateAudioMeters()
+	{
+		using namespace FSubmixDashboardViewFactoryPrivate;
+
+		for (auto It = PendingToReactivateSubmixIds.CreateIterator(); It; ++It)
+		{
+			const TSharedPtr<IDashboardDataViewEntry>* FoundEntry = DataViewEntries.FindByPredicate([&InSubmixId = *It](const TSharedPtr<IDashboardDataViewEntry>& Entry)
+				{
+					const FSubmixDashboardEntry& SubmixDashboardEntry = CastEntry(*Entry);
+
+					return SubmixDashboardEntry.SubmixId == InSubmixId;
+				});
+
+			if (FoundEntry)
+			{
+				const TSharedPtr<IDashboardDataViewEntry> Entry = *FoundEntry;
+				const FSubmixDashboardEntry& SubmixDashboardEntry = CastEntry(*Entry);
+
+				SendSubmixEnvelopeFollowerCVar(true, SubmixDashboardEntry);
+				FAudioMeterProvider::OnAddAudioMeter.ExecuteIfBound(SubmixDashboardEntry.SubmixId, SubmixDashboardEntry.AudioMeterInfo, SubmixDashboardEntry.GetDisplayName(), SubmixMeterTextColor);
+				FAudioMeterProvider::OnUpdateAudioMeterInfo.ExecuteIfBound(SubmixDashboardEntry.SubmixId, SubmixDashboardEntry.AudioMeterInfo, SubmixDashboardEntry.GetDisplayName());
+
+				It.RemoveCurrent();
+			}
+		}
+	}
+
+	void FSubmixDashboardViewFactory::HandleOnSubmixLoaded(const uint32 InSubmixId)
+	{
+		if (CheckedSubmixes.Contains(InSubmixId))
+		{
+			PendingToReactivateSubmixIds.Emplace(InSubmixId);
+			bShouldReactivateAudioMeters = true;
+		}
+	}
+
 	void FSubmixDashboardViewFactory::HandleOnTimeMarkerUpdated()
 	{
 		bShouldRebuildAudioMeters = true;
diff --git a/Engine/Plugins/AudioInsights/Source/AudioInsights/Public/Views/SubmixDashboardViewFactory.h b/Engine/Plugins/AudioInsights/Source/AudioInsights/Public/Views/SubmixDashboardViewFactory.h
--- a/Engine/Plugins/AudioInsights/Source/AudioInsights/Public/Views/SubmixDashboardViewFactory.h
+++ b/Engine/Plugins/AudioInsights/Source/AudioInsights/Public/Views/SubmixDashboardViewFactory.h
@@ -33,11 +33,13 @@ namespace UE::Audio::Insights
 		UE_API void RebuildAudioMeters() const;
 		UE_API bool ShouldUpdateAudioMeter(const FSubmixDashboardEntry& InSubmixDashboardEntry) const;
 		UE_API void UpdateAudioMeters() const;
+		UE_API void ReactivateAudioMeters();
 
 		UE_API virtual void SortTable() override;
 
 		UE_API void HandleOnSubmixAdded(const uint32 InSubmixId);
 		UE_API void HandleOnSubmixRemoved(const uint32 InSubmixId);
+		UE_API void HandleOnSubmixLoaded(const uint32 InSubmixId);
 		UE_API void HandleOnTimeMarkerUpdated();
 		UE_API void HandleOnTimeControlMethodReset();
 
@@ -46,9 +48,11 @@ namespace UE::Audio::Insights
 		TSharedPtr<FSubmixTraceProvider> SubmixProvider;
 
 		TSet<uint32> CheckedSubmixes;
+		TSet<uint32> PendingToReactivateSubmixIds;
 
 		bool bShouldRebuildAudioMeters = false;
 		bool bListRefreshed = false;
+		bool bShouldReactivateAudioMeters = false;
 	};
 } // namespace UE::Audio::Insights
 
