[Audio Insights] Audio Bus dashboard: fix for Audio Bus meter no longer appearing when an Audio Bus is triggered after the second time

diff --git a/Engine/Plugins/AudioInsights/Source/AudioInsights/Private/Providers/AudioBusTraceProvider.cpp b/Engine/Plugins/AudioInsights/Source/AudioInsights/Private/Providers/AudioBusTraceProvider.cpp
--- a/Engine/Plugins/AudioInsights/Source/AudioInsights/Private/Providers/AudioBusTraceProvider.cpp
+++ b/Engine/Plugins/AudioInsights/Source/AudioInsights/Private/Providers/AudioBusTraceProvider.cpp
@@ -260,6 +260,8 @@ namespace UE::Audio::Insights
 				{
 					EntryRef.AudioBusType = Msg.AudioBusType;
 				}
+
+				OnAudioBusStarted.ExecuteIfBound(Msg.AudioBusId);
 			}
 		});
 
@@ -331,7 +333,15 @@ namespace UE::Audio::Insights
 			if (OutEntry && (*OutEntry)->Timestamp < Msg.Timestamp)
 			{
 #if WITH_EDITOR
-				const FAudioBusDashboardEntry& EntryRef = *OutEntry->Get();
+				// Turn off audio bus entry activity, do not remove it since we are in Editor mode.
+				// We keep in the dashboard all registered audio bus assets from Content Browser
+				FAudioBusDashboardEntry& EntryRef = *OutEntry->Get();
+
+				EntryRef.bHasActivity = false;
+
+				TArray<float>& EnvelopeValues = EntryRef.AudioMeterInfo->EnvelopeValues;
+				FMemory::Memzero(EnvelopeValues.GetData(), EnvelopeValues.Num() * sizeof(float));
+
 				if (EntryRef.AudioBusType == EAudioBusType::CodeGenerated)
 #endif  // WITH_EDITOR
 				{
diff --git a/Engine/Plugins/AudioInsights/Source/AudioInsights/Private/Providers/AudioBusTraceProvider.h b/Engine/Plugins/AudioInsights/Source/AudioInsights/Private/Providers/AudioBusTraceProvider.h
--- a/Engine/Plugins/AudioInsights/Source/AudioInsights/Private/Providers/AudioBusTraceProvider.h
+++ b/Engine/Plugins/AudioInsights/Source/AudioInsights/Private/Providers/AudioBusTraceProvider.h
@@ -31,6 +31,7 @@ namespace UE::Audio::Insights
 		DECLARE_DELEGATE_OneParam(FOnAudioBusChanged, const uint32 /*AudioBusId*/);
 		FOnAudioBusChanged OnAudioBusAdded;
 		FOnAudioBusChanged OnAudioBusRemoved;
+		FOnAudioBusChanged OnAudioBusStarted;
 
 		DECLARE_DELEGATE(FOnAudioBusListUpdated);
 		FOnAudioBusListUpdated OnAudioBusListUpdated;
diff --git a/Engine/Plugins/AudioInsights/Source/AudioInsights/Private/Views/AudioBusDashboardViewFactory.cpp b/Engine/Plugins/AudioInsights/Source/AudioInsights/Private/Views/AudioBusDashboardViewFactory.cpp
--- a/Engine/Plugins/AudioInsights/Source/AudioInsights/Private/Views/AudioBusDashboardViewFactory.cpp
+++ b/Engine/Plugins/AudioInsights/Source/AudioInsights/Private/Views/AudioBusDashboardViewFactory.cpp
@@ -66,6 +66,7 @@ namespace UE::Audio::Insights
 
 		AudioBusProvider->OnAudioBusAdded.BindRaw(this, &FAudioBusDashboardViewFactory::HandleOnAudioBusAdded);
 		AudioBusProvider->OnAudioBusRemoved.BindRaw(this, &FAudioBusDashboardViewFactory::HandleOnAudioBusRemoved);
+		AudioBusProvider->OnAudioBusStarted.BindRaw(this, &FAudioBusDashboardViewFactory::HandleOnAudioBusStarted);
 		AudioBusProvider->OnAudioBusListUpdated.BindRaw(this, &FAudioBusDashboardViewFactory::RequestListRefresh);
 		AudioBusProvider->OnTimeMarkerUpdated.BindRaw(this, &FAudioBusDashboardViewFactory::HandleOnTimeMarkerUpdated);
 	}
@@ -252,6 +253,11 @@ namespace UE::Audio::Insights
 			UpdateAudioMeters();
 			bListRefreshed = false;
 		}
+		else if (bShouldReactivateAudioMeters)
+		{
+			ReactivateAudioMeters();
+			bShouldReactivateAudioMeters = false;
+		}
 	}
 
 	TSharedRef<SWidget> FAudioBusDashboardViewFactory::MakeAudioBusTypeFilterWidget()
@@ -552,6 +558,33 @@ namespace UE::Audio::Insights
 		}
 	}
 
+	void FAudioBusDashboardViewFactory::ReactivateAudioMeters()
+	{
+		using namespace FAudioBusDashboardViewFactoryPrivate;
+
+		for (auto It = PendingToReactivateAudioBusIds.CreateIterator(); It; ++It)
+		{
+			const TSharedPtr<IDashboardDataViewEntry>* FoundEntry = DataViewEntries.FindByPredicate([&AudioBusId = *It](const TSharedPtr<IDashboardDataViewEntry>& Entry)
+				{
+					const FAudioBusDashboardEntry& AudioBusDashboardEntry = CastEntry(*Entry);
+
+					return AudioBusDashboardEntry.AudioBusId == AudioBusId;
+				});
+
+			if (FoundEntry)
+			{
+				const TSharedPtr<IDashboardDataViewEntry> Entry = *FoundEntry;
+				const FAudioBusDashboardEntry& AudioBusDashboardEntry = CastEntry(*Entry);
+
+				SendAudioBusEnvelopeFollowerCVar(true, AudioBusDashboardEntry.AudioBusId);
+				FAudioMeterProvider::OnAddAudioMeter.ExecuteIfBound(AudioBusDashboardEntry.AudioBusId, AudioBusDashboardEntry.AudioMeterInfo, AudioBusDashboardEntry.GetDisplayName(), AudioBusMeterTextColor);
+				FAudioMeterProvider::OnUpdateAudioMeterInfo.ExecuteIfBound(AudioBusDashboardEntry.AudioBusId, AudioBusDashboardEntry.AudioMeterInfo, AudioBusDashboardEntry.GetDisplayName());
+
+				It.RemoveCurrent();
+			}
+		}
+	}
+
 	void FAudioBusDashboardViewFactory::HandleOnAudioBusAdded(const uint32 InAudioBusId)
 	{
 		RequestListRefresh();
@@ -566,6 +599,15 @@ namespace UE::Audio::Insights
 		RequestListRefresh();
 	}
 
+	void FAudioBusDashboardViewFactory::HandleOnAudioBusStarted(const uint32 InAudioBusId)
+	{
+		if (CheckedAudioBuses.Contains(InAudioBusId))
+		{
+			PendingToReactivateAudioBusIds.Emplace(InAudioBusId);
+			bShouldReactivateAudioMeters = true;
+		}
+	}
+
 	void FAudioBusDashboardViewFactory::HandleOnTimeMarkerUpdated()
 	{
 		bShouldRebuildAudioMeters = true;
diff --git a/Engine/Plugins/AudioInsights/Source/AudioInsights/Public/Views/AudioBusDashboardViewFactory.h b/Engine/Plugins/AudioInsights/Source/AudioInsights/Public/Views/AudioBusDashboardViewFactory.h
--- a/Engine/Plugins/AudioInsights/Source/AudioInsights/Public/Views/AudioBusDashboardViewFactory.h
+++ b/Engine/Plugins/AudioInsights/Source/AudioInsights/Public/Views/AudioBusDashboardViewFactory.h
@@ -40,6 +40,7 @@ namespace UE::Audio::Insights
 		UE_API void RebuildAudioMeters() const;
 		UE_API bool ShouldUpdateAudioMeter(const FAudioBusDashboardEntry& InAudioBusDashboardEntry) const;
 		UE_API void UpdateAudioMeters() const;
+		UE_API void ReactivateAudioMeters();
 
 		UE_API virtual void SortTable() override;
 
@@ -48,6 +49,8 @@ namespace UE::Audio::Insights
 
 		UE_API void HandleOnAudioBusAdded(const uint32 InAudioBusId);
 		UE_API void HandleOnAudioBusRemoved(const uint32 InAudioBusId);
+		UE_API void HandleOnAudioBusStarted(const uint32 InAudioBusId);
+
 		UE_API void HandleOnTimeMarkerUpdated();
 		UE_API void HandleOnTimeControlMethodReset();
 
@@ -56,9 +59,11 @@ namespace UE::Audio::Insights
 		TSharedPtr<FAudioBusTraceProvider> AudioBusProvider;
 
 		TSet<uint32> CheckedAudioBuses;
+		TSet<uint32> PendingToReactivateAudioBusIds;
 
 		bool bShouldRebuildAudioMeters = false;
 		bool bListRefreshed = false;
+		bool bShouldReactivateAudioMeters = false;
 
 		using FComboboxSelectionItem = TPair<EAudioBusTypeComboboxSelection, FText>;
 		TArray<TSharedPtr<FComboboxSelectionItem>> AudioBusTypes;
