[Audio Insights] AudioInsightsTraceModule: Added additional code path to check if a session is available when doing direct trace in Unreal Insights standalone. This allow us to display a message when audio trace channels are not enabled.

diff --git a/Engine/Plugins/AudioInsights/Source/AudioInsights/Private/AudioInsightsTraceModule.cpp b/Engine/Plugins/AudioInsights/Source/AudioInsights/Private/AudioInsightsTraceModule.cpp
--- a/Engine/Plugins/AudioInsights/Source/AudioInsights/Private/AudioInsightsTraceModule.cpp
+++ b/Engine/Plugins/AudioInsights/Source/AudioInsights/Private/AudioInsightsTraceModule.cpp
@@ -184,23 +184,33 @@ namespace UE::Audio::Insights
 	}
 
 #if !WITH_EDITOR
-	const UE::Trace::FStoreClient::FSessionInfo* FTraceModule::GetCurrentAnalysisSessionInfo() const
+	bool FTraceModule::IsSessionInfoAvailable() const
 	{
 		IUnrealInsightsModule& InsightsModule = FModuleManager::GetModuleChecked<IUnrealInsightsModule>("TraceInsights");
 
 		TSharedPtr<const TraceServices::IAnalysisSession> AnalysisSession = InsightsModule.GetAnalysisSession();
-		if (!AnalysisSession.IsValid())
+		if (AnalysisSession.IsValid())
 		{
-			return nullptr;
-		}
+			TraceServices::FAnalysisSessionReadScope SessionReadScope(*AnalysisSession.Get());
 
-		UE::Trace::FStoreClient* StoreClient = InsightsModule.GetStoreClient();
-		if (StoreClient == nullptr)
-		{
-			return nullptr;
+			UE::Trace::FStoreClient* StoreClient = InsightsModule.GetStoreClient();
+			if (StoreClient)
+			{
+				const UE::Trace::FStoreClient::FSessionInfo* SessionInfo = StoreClient->GetSessionInfoByTraceId(AnalysisSession->GetTraceId());
+				return SessionInfo != nullptr;
+			}
+			// Direct traces have a trace ID of zero and a null StoreClient - check IDiagnosticsProvider to see if we are running a direct trace
+			else if (AnalysisSession->GetTraceId() == 0)
+			{
+				const TraceServices::IDiagnosticsProvider* DiagnosticsProvider = TraceServices::ReadDiagnosticsProvider(*AnalysisSession.Get());
+				if (DiagnosticsProvider)
+				{
+					return DiagnosticsProvider->IsSessionInfoAvailable();
+				}
+			}
 		}
 
-		return StoreClient->GetSessionInfoByTraceId(AnalysisSession->GetTraceId());
+		return false;
 	}
 
 	bool FTraceModule::GetAudioTracesAreEnabled() const
@@ -424,7 +434,7 @@ namespace UE::Audio::Insights
 		return false;
 #else
 		// When attaching to a packaged build, we need to check whether we're connected and the audio channels are enabled
-		return !bTraceAnalysisHasStarted && !bAudioTraceChannelsEnabled && GetCurrentAnalysisSessionInfo() != nullptr && !GetAudioTracesAreEnabled();
+		return !bTraceAnalysisHasStarted && !bAudioTraceChannelsEnabled && IsSessionInfoAvailable() && !GetAudioTracesAreEnabled();
 #endif // WITH_EDITOR
 	}
 
diff --git a/Engine/Plugins/AudioInsights/Source/AudioInsights/Public/AudioInsightsTraceModule.h b/Engine/Plugins/AudioInsights/Source/AudioInsights/Public/AudioInsightsTraceModule.h
--- a/Engine/Plugins/AudioInsights/Source/AudioInsights/Public/AudioInsightsTraceModule.h
+++ b/Engine/Plugins/AudioInsights/Source/AudioInsights/Public/AudioInsightsTraceModule.h
@@ -88,7 +88,7 @@ namespace UE::Audio::Insights
 		UE_API void RestoreCachedChannels() const;
 
 #if !WITH_EDITOR
-		UE_API const UE::Trace::FStoreClient::FSessionInfo* GetCurrentAnalysisSessionInfo() const;
+		UE_API bool IsSessionInfoAvailable() const;
 		UE_API bool GetAudioTracesAreEnabled() const;
 		UE_API void SendDiscoveryRequestToTraceController() const;
 		UE_API void OnSessionsUpdated();
