[Audio Insights] BugHawk fixes for CLs 46514818 and 46514823

diff --git a/Engine/Plugins/AudioInsights/Source/AudioInsights/Private/Views/AudioBusDashboardViewFactory.cpp b/Engine/Plugins/AudioInsights/Source/AudioInsights/Private/Views/AudioBusDashboardViewFactory.cpp
--- a/Engine/Plugins/AudioInsights/Source/AudioInsights/Private/Views/AudioBusDashboardViewFactory.cpp
+++ b/Engine/Plugins/AudioInsights/Source/AudioInsights/Private/Views/AudioBusDashboardViewFactory.cpp
@@ -75,6 +75,7 @@ namespace UE::Audio::Insights
 	{
 		AudioBusProvider->OnAudioBusAdded.Unbind();
 		AudioBusProvider->OnAudioBusRemoved.Unbind();
+		AudioBusProvider->OnAudioBusStarted.Unbind();
 		AudioBusProvider->OnAudioBusListUpdated.Unbind();
 		AudioBusProvider->OnTimeMarkerUpdated.Unbind();
 	}
diff --git a/Engine/Plugins/AudioInsights/Source/AudioInsights/Private/Views/SubmixDashboardViewFactory.cpp b/Engine/Plugins/AudioInsights/Source/AudioInsights/Private/Views/SubmixDashboardViewFactory.cpp
--- a/Engine/Plugins/AudioInsights/Source/AudioInsights/Private/Views/SubmixDashboardViewFactory.cpp
+++ b/Engine/Plugins/AudioInsights/Source/AudioInsights/Private/Views/SubmixDashboardViewFactory.cpp
@@ -58,6 +58,7 @@ namespace UE::Audio::Insights
 	{
 		SubmixProvider->OnSubmixAdded.Unbind();
 		SubmixProvider->OnSubmixRemoved.Unbind();
+		SubmixProvider->OnSubmixLoaded.Unbind();
 		SubmixProvider->OnSubmixListUpdated.Unbind();
 		SubmixProvider->OnTimeMarkerUpdated.Unbind();
 	}
