[Audio Insights] Fix for enable plot button color not matching curve after inspecting a timestamp, and plots not immediately appearing when enabling plots while processing is paused.

diff --git a/Engine/Plugins/AudioInsights/Source/AudioInsights/Private/Providers/SoundTraceProvider.cpp b/Engine/Plugins/AudioInsights/Source/AudioInsights/Private/Providers/SoundTraceProvider.cpp
--- a/Engine/Plugins/AudioInsights/Source/AudioInsights/Private/Providers/SoundTraceProvider.cpp
+++ b/Engine/Plugins/AudioInsights/Source/AudioInsights/Private/Providers/SoundTraceProvider.cpp
@@ -1512,6 +1512,14 @@ namespace UE::Audio::Insights
 
 	void FSoundTraceProvider::OnTimeControlMethodReset()
 	{
+		for (const auto& [AudioDeviceID, DeviceData] : DeviceDataMap)
+		{
+			for (const auto& [EntryType, SoundDashboardEntry] : DeviceData)
+			{
+				CacheIsPlottingFlagRecursive(SoundDashboardEntry);
+			}
+		}
+
 		Reset();
 	}
 
diff --git a/Engine/Plugins/AudioInsights/Source/AudioInsights/Private/Views/SoundDashboardViewFactory.cpp b/Engine/Plugins/AudioInsights/Source/AudioInsights/Private/Views/SoundDashboardViewFactory.cpp
--- a/Engine/Plugins/AudioInsights/Source/AudioInsights/Private/Views/SoundDashboardViewFactory.cpp
+++ b/Engine/Plugins/AudioInsights/Source/AudioInsights/Private/Views/SoundDashboardViewFactory.cpp
@@ -3076,7 +3076,7 @@ namespace UE::Audio::Insights
 
 						return FSlateStyle::Get().GetBrush(IconName);
 					})
-					.ColorAndOpacity(InRowData->GetEntryColor())
+					.ColorAndOpacity_Lambda([InRowData]() { return InRowData->GetEntryColor(); })
 				]
 			];
 	}
@@ -3118,7 +3118,8 @@ namespace UE::Audio::Insights
 
 		if (FilteredEntriesListView.IsValid())
 		{
-			OnProcessPlotData.Broadcast(DataViewEntries, FilteredEntriesListView->GetSelectedItems());
+			// We want to force plot data to process even if the values have not changed when toggling visibility
+			OnProcessPlotData.Broadcast(DataViewEntries, FilteredEntriesListView->GetSelectedItems(), true /*bForceUpdate*/);
 		}
 
 		OnUpdatePlotVisibility.Broadcast(DataViewEntries);
@@ -3237,19 +3238,19 @@ namespace UE::Audio::Insights
 #if WITH_EDITOR
 		if (FilteredEntriesListView.IsValid())
 		{
-			OnProcessPlotData.Broadcast(DataViewEntries, FilteredEntriesListView->GetSelectedItems());
+			OnProcessPlotData.Broadcast(DataViewEntries, FilteredEntriesListView->GetSelectedItems(), false /*bForceUpdate*/);
 		}
 #else
 		if (FilteredEntriesListView.IsValid())
 		{
-			OnProcessPlotData.Broadcast(DataViewEntries, FilteredEntriesListView->GetSelectedItems());
+			OnProcessPlotData.Broadcast(DataViewEntries, FilteredEntriesListView->GetSelectedItems(), false /*bForceUpdate*/);
 		}
 		else
 		{
 			// In standalone, we still want to process entries as the trace file loads in
 			// Send an empty array for selected items and broadcast the data
 			const TArray<TSharedPtr<IDashboardDataTreeViewEntry>> EmptySelectedEntriesContainer;
-			OnProcessPlotData.Broadcast(DataViewEntries, EmptySelectedEntriesContainer);
+			OnProcessPlotData.Broadcast(DataViewEntries, EmptySelectedEntriesContainer, false /*bForceUpdate*/);
 		}
 #endif // WITH_EDITOR
 	}
diff --git a/Engine/Plugins/AudioInsights/Source/AudioInsights/Private/Views/SoundPlotsDashboardViewFactory.cpp b/Engine/Plugins/AudioInsights/Source/AudioInsights/Private/Views/SoundPlotsDashboardViewFactory.cpp
--- a/Engine/Plugins/AudioInsights/Source/AudioInsights/Private/Views/SoundPlotsDashboardViewFactory.cpp
+++ b/Engine/Plugins/AudioInsights/Source/AudioInsights/Private/Views/SoundPlotsDashboardViewFactory.cpp
@@ -58,11 +58,11 @@ namespace UE::Audio::Insights
 		OnUpdatePlotSelection = SoundsDashboard->OnUpdatePlotSelection.AddSP(this, &FSoundPlotsDashboardViewFactory::UpdatePlotSelection);
 	}
 
-	void FSoundPlotsDashboardViewFactory::ProcessPlotData(const TArray<TSharedPtr<IDashboardDataTreeViewEntry>>& DataViewEntries, const TArray<TSharedPtr<IDashboardDataTreeViewEntry>>& SelectedEntries)
+	void FSoundPlotsDashboardViewFactory::ProcessPlotData(const TArray<TSharedPtr<IDashboardDataTreeViewEntry>>& DataViewEntries, const TArray<TSharedPtr<IDashboardDataTreeViewEntry>>& SelectedEntries, const bool bForceUpdate /* = false */)
 	{
 		if (PlotsWidget.IsValid())
 		{
-			PlotsWidget->ProcessPlotData(DataViewEntries, SelectedEntries);
+			PlotsWidget->ProcessPlotData(DataViewEntries, SelectedEntries, bForceUpdate);
 		}
 	}
 
diff --git a/Engine/Plugins/AudioInsights/Source/AudioInsights/Private/Views/SoundPlotsWidgetView.cpp b/Engine/Plugins/AudioInsights/Source/AudioInsights/Private/Views/SoundPlotsWidgetView.cpp
--- a/Engine/Plugins/AudioInsights/Source/AudioInsights/Private/Views/SoundPlotsWidgetView.cpp
+++ b/Engine/Plugins/AudioInsights/Source/AudioInsights/Private/Views/SoundPlotsWidgetView.cpp
@@ -95,21 +95,24 @@ namespace UE::Audio::Insights
 		return MakePlotsWidget();
 	}
 
-	void FSoundPlotsWidgetView::ProcessPlotData(const TArray<TSharedPtr<IDashboardDataTreeViewEntry>>& DataViewEntries, const TArray<TSharedPtr<IDashboardDataTreeViewEntry>>& SelectedEntries)
+	void FSoundPlotsWidgetView::ProcessPlotData(const TArray<TSharedPtr<IDashboardDataTreeViewEntry>>& DataViewEntries, const TArray<TSharedPtr<IDashboardDataTreeViewEntry>>& SelectedEntries, const bool bForceUpdate /* = false */)
 	{
-		if (GameState == EGameState::ControlledByTimemarker && !bTimeMarkerMoved)
+		if (DataViewEntries.IsEmpty() || !PlotWidgetMetadataPerCurve.IsValid())
 		{
 			return;
 		}
 
-		if (DataViewEntries.IsEmpty() || GameState == EGameState::Paused)
+		if (!bForceUpdate)
 		{
-			return;
-		}
+			if (GameState == EGameState::ControlledByTimemarker && !bTimeMarkerMoved)
+			{
+				return;
+			}
 
-		if (!PlotWidgetMetadataPerCurve.IsValid())
-		{
-			return;
+			if (GameState == EGameState::Paused)
+			{
+				return;
+			}
 		}
 
 		// Process new data
@@ -605,7 +608,7 @@ namespace UE::Audio::Insights
 					}
 				}
 
-				if (GameState != EGameState::ControlledByTimemarker)
+				if (GameState != EGameState::ControlledByTimemarker && GameState != EGameState::Paused)
 				{
 					CurrentTimestamp = FMath::Max(CurrentTimestamp, Timestamp);
 				}
@@ -626,6 +629,12 @@ namespace UE::Audio::Insights
 		// Create metadata for this curve if necessary
 		if (SAudioCurveView::FCurveMetadata* MetaData = PlotWidgetMetadataPerCurve->Find(SourceId))
 		{
+			const FLinearColor EntryColor = Entry->GetEntryColor();
+			if (MetaData->CurveColor.A > 0.0f && MetaData->CurveColor != EntryColor)
+			{
+				Entry->SetEntryColor(MetaData->CurveColor);
+			}
+
 			const bool bPlotWasSelected = MetaData->PlotThickness != DefaultPlotThickness;
 			if (bPlotWasSelected != SelectedEntries.Contains(Entry))
 			{
diff --git a/Engine/Plugins/AudioInsights/Source/AudioInsights/Public/AudioInsightsDataSource.h b/Engine/Plugins/AudioInsights/Source/AudioInsights/Public/AudioInsightsDataSource.h
--- a/Engine/Plugins/AudioInsights/Source/AudioInsights/Public/AudioInsightsDataSource.h
+++ b/Engine/Plugins/AudioInsights/Source/AudioInsights/Public/AudioInsightsDataSource.h
@@ -25,6 +25,7 @@ namespace UE::Audio::Insights
 		virtual uint64 GetEntryID() const = 0;
 		virtual const FText& GetDisplayName() const = 0;
 		virtual const FLinearColor& GetEntryColor() const = 0;
+		virtual void SetEntryColor(const FLinearColor& Color) = 0;
 		virtual bool HasSetInitExpansion() const = 0;
 		virtual void ResetHasSetInitExpansion() = 0;
 
diff --git a/Engine/Plugins/AudioInsights/Source/AudioInsights/Public/Messages/SoundTraceMessages.h b/Engine/Plugins/AudioInsights/Source/AudioInsights/Public/Messages/SoundTraceMessages.h
--- a/Engine/Plugins/AudioInsights/Source/AudioInsights/Public/Messages/SoundTraceMessages.h
+++ b/Engine/Plugins/AudioInsights/Source/AudioInsights/Public/Messages/SoundTraceMessages.h
@@ -290,6 +290,11 @@ namespace UE::Audio::Insights
 			return bIsCategory ? FLinearColor::White : EntryColor;
 		}
 
+		virtual void SetEntryColor(const FLinearColor& Color) override
+		{
+			EntryColor = Color;
+		}
+
 		virtual bool IsValid() const override { return ActiveSoundPlayOrder != INDEX_NONE; }
 		virtual uint64 GetEntryID() const override { return SoundMessageUtils::GeneratePlayOrderUniqueID(ActiveSoundPlayOrder, WaveInstancePlayOrder); }
 
@@ -311,7 +316,7 @@ namespace UE::Audio::Insights
 		double Timestamp = 0.0;
 		double TimeoutTimestamp = INVALID_TIMEOUT;
 		FText CategoryName;
-		const FLinearColor EntryColor = FColor::MakeRandomColor();
+		FLinearColor EntryColor = FColor::MakeRandomColor();
 		ESoundDashboardEntryType EntryType = ESoundDashboardEntryType::None;
 		EPinnedEntryType PinnedEntryType = EPinnedEntryType::None;
 		bool bIsCategory = false;
diff --git a/Engine/Plugins/AudioInsights/Source/AudioInsights/Public/Views/SoundDashboardViewFactory.h b/Engine/Plugins/AudioInsights/Source/AudioInsights/Public/Views/SoundDashboardViewFactory.h
--- a/Engine/Plugins/AudioInsights/Source/AudioInsights/Public/Views/SoundDashboardViewFactory.h
+++ b/Engine/Plugins/AudioInsights/Source/AudioInsights/Public/Views/SoundDashboardViewFactory.h
@@ -114,7 +114,7 @@ namespace UE::Audio::Insights
 
 		UE_API TMap<FName, FSoundPlotsWidgetView::FPlotColumnInfo> GetPlotColumnInfo();
 
-		DECLARE_MULTICAST_DELEGATE_TwoParams(FOnProcessPlotData, const TArray<TSharedPtr<IDashboardDataTreeViewEntry>>& /*DataViewEntries*/, const TArray<TSharedPtr<IDashboardDataTreeViewEntry>>& /*SelectedEntries*/);
+		DECLARE_MULTICAST_DELEGATE_ThreeParams(FOnProcessPlotData, const TArray<TSharedPtr<IDashboardDataTreeViewEntry>>& /*DataViewEntries*/, const TArray<TSharedPtr<IDashboardDataTreeViewEntry>>& /*SelectedEntries*/, const bool /*bForceUpdate*/);
 		FOnProcessPlotData OnProcessPlotData;
 
 		DECLARE_MULTICAST_DELEGATE_OneParam(FOnUpdatePlotVisibility, const TArray<TSharedPtr<IDashboardDataTreeViewEntry>>& /*DataViewEntries*/);
diff --git a/Engine/Plugins/AudioInsights/Source/AudioInsights/Public/Views/SoundPlotsDashboardViewFactory.h b/Engine/Plugins/AudioInsights/Source/AudioInsights/Public/Views/SoundPlotsDashboardViewFactory.h
--- a/Engine/Plugins/AudioInsights/Source/AudioInsights/Public/Views/SoundPlotsDashboardViewFactory.h
+++ b/Engine/Plugins/AudioInsights/Source/AudioInsights/Public/Views/SoundPlotsDashboardViewFactory.h
@@ -28,7 +28,7 @@ namespace UE::Audio::Insights
 	private:
 		TSharedPtr<FSoundPlotsWidgetView> PlotsWidget;
 
-		void ProcessPlotData(const TArray<TSharedPtr<IDashboardDataTreeViewEntry>>& DataViewEntries, const TArray<TSharedPtr<IDashboardDataTreeViewEntry>>& SelectedEntries);
+		void ProcessPlotData(const TArray<TSharedPtr<IDashboardDataTreeViewEntry>>& DataViewEntries, const TArray<TSharedPtr<IDashboardDataTreeViewEntry>>& SelectedEntries, const bool bForceUpdate = false);
 
 		void UpdatePlotVisibility(const TArray<TSharedPtr<IDashboardDataTreeViewEntry>>& DataViewEntries);
 		void UpdatePlotSelection(const TArray<TSharedPtr<IDashboardDataTreeViewEntry>>& SelectedEntries);
diff --git a/Engine/Plugins/AudioInsights/Source/AudioInsights/Public/Views/SoundPlotsWidgetView.h b/Engine/Plugins/AudioInsights/Source/AudioInsights/Public/Views/SoundPlotsWidgetView.h
--- a/Engine/Plugins/AudioInsights/Source/AudioInsights/Public/Views/SoundPlotsWidgetView.h
+++ b/Engine/Plugins/AudioInsights/Source/AudioInsights/Public/Views/SoundPlotsWidgetView.h
@@ -35,7 +35,7 @@ namespace UE::Audio::Insights
 		UE_API virtual ~FSoundPlotsWidgetView();
 
 		UE_API TSharedRef<SWidget> MakeWidget();
-		UE_API void ProcessPlotData(const TArray<TSharedPtr<IDashboardDataTreeViewEntry>>& DataViewEntries, const TArray<TSharedPtr<IDashboardDataTreeViewEntry>>& SelectedEntries);
+		UE_API void ProcessPlotData(const TArray<TSharedPtr<IDashboardDataTreeViewEntry>>& DataViewEntries, const TArray<TSharedPtr<IDashboardDataTreeViewEntry>>& SelectedEntries, const bool bForceUpdate = false);
 		UE_API void UpdatePlotVisibility(const TArray<TSharedPtr<IDashboardDataTreeViewEntry>>& DataViewEntries);
 
 		UE_API void UpdatePlotSelection(const TArray<TSharedPtr<IDashboardDataTreeViewEntry>>& SelectedEntries);
