[Audio Insights] AudioInsightsEditorDashboardFactory: track when layout is getting reset so we don't stop/start a trace analysis again

diff --git a/Engine/Plugins/AudioInsights/Source/AudioInsightsEditor/Private/AudioInsightsEditorDashboardFactory.cpp b/Engine/Plugins/AudioInsights/Source/AudioInsightsEditor/Private/AudioInsightsEditorDashboardFactory.cpp
--- a/Engine/Plugins/AudioInsights/Source/AudioInsightsEditor/Private/AudioInsightsEditorDashboardFactory.cpp
+++ b/Engine/Plugins/AudioInsights/Source/AudioInsightsEditor/Private/AudioInsightsEditorDashboardFactory.cpp
@@ -347,8 +347,11 @@ namespace UE::Audio::Insights
 			SetActiveDeviceId(DeviceManager->GetActiveAudioDevice()->DeviceID);
 		}
 
-		IAudioInsightsTraceModule& TraceModule = FAudioInsightsEditorModule::GetChecked().GetTraceModule();
-		TraceModule.StartTraceAnalysis(bOnlyTraceAudioChannels, ETraceMode::Monitoring);
+		if (!bIsLayoutReseting)
+		{
+			IAudioInsightsTraceModule& TraceModule = FAudioInsightsEditorModule::GetChecked().GetTraceModule();
+			TraceModule.StartTraceAnalysis(bOnlyTraceAudioChannels, ETraceMode::Monitoring);
+		}
 
 		const TSharedRef<FTabManager::FLayout> TabLayout = LoadLayoutFromConfig();
 
@@ -370,7 +373,7 @@ namespace UE::Audio::Insights
 			// Stop trace when closing Audio Insights (only if we have started it, we don't want to stop external traces)
 			IAudioInsightsTraceModule& TraceModule = FAudioInsightsEditorModule::GetChecked().GetTraceModule();
 
-			if (TraceModule.IsTraceAnalysisActive() && TraceModule.IsAudioInsightsTrace())
+			if (TraceModule.IsTraceAnalysisActive() && TraceModule.IsAudioInsightsTrace() && !bIsLayoutReseting)
 			{
 				TraceModule.StopTraceAnalysis();
 			}
@@ -518,12 +521,16 @@ namespace UE::Audio::Insights
 							{
 								// Wipe all the persisted areas and close tab
 								DashboardTabManager->CloseAllAreas();
+
+								bIsLayoutReseting = true;
+
 								OwnerTab->RequestCloseTab();
 
 								// Can't invoke the tab immediately (it won't show up), needs to be done a bit later
-								AsyncTask(ENamedThreads::GameThread, [AudioInsightsTabId = OwnerTab->GetLayoutIdentifier()]()
+								AsyncTask(ENamedThreads::GameThread, [AudioInsightsTabId = OwnerTab->GetLayoutIdentifier(), &bIsLayoutReseting = this->bIsLayoutReseting]()
 								{
 									FGlobalTabmanager::Get()->TryInvokeTab(AudioInsightsTabId);
+									bIsLayoutReseting = false;
 								});
 							}
 						}
diff --git a/Engine/Plugins/AudioInsights/Source/AudioInsightsEditor/Private/AudioInsightsEditorDashboardFactory.h b/Engine/Plugins/AudioInsights/Source/AudioInsightsEditor/Private/AudioInsightsEditorDashboardFactory.h
--- a/Engine/Plugins/AudioInsights/Source/AudioInsightsEditor/Private/AudioInsightsEditorDashboardFactory.h
+++ b/Engine/Plugins/AudioInsights/Source/AudioInsightsEditor/Private/AudioInsightsEditorDashboardFactory.h
@@ -70,6 +70,7 @@ namespace UE::Audio::Insights
 		FDelegateHandle OnTraceStartedHandle;
 
 		bool bOnlyTraceAudioChannels = false;
+		bool bIsLayoutReseting = false;
 
 		TSharedPtr<FTabManager> DashboardTabManager;
 		TSharedPtr<FWorkspaceItem> DashboardWorkspace;
