[Audio Insights] Submix dashboard: Separate alive ping message from loaded message to avoid requesting sending cvar to activate audio meter every second.

diff --git a/Engine/Plugins/AudioInsights/Source/AudioInsights/Private/Messages/SubmixTraceMessages.cpp b/Engine/Plugins/AudioInsights/Source/AudioInsights/Private/Messages/SubmixTraceMessages.cpp
--- a/Engine/Plugins/AudioInsights/Source/AudioInsights/Private/Messages/SubmixTraceMessages.cpp
+++ b/Engine/Plugins/AudioInsights/Source/AudioInsights/Private/Messages/SubmixTraceMessages.cpp
@@ -45,6 +45,28 @@ namespace UE::Audio::Insights
 		return MemorySize;
 	};
 
+	// FSubmixAlivePingMessage
+	FSubmixAlivePingMessage::FSubmixAlivePingMessage(const Trace::IAnalyzer::FOnEventContext& InContext)
+		: FSubmixMessageBase(InContext)
+	{
+		const Trace::IAnalyzer::FEventData& EventData = InContext.EventData;
+
+		EventData.GetString("Name", Name);
+		NumChannels = EventData.GetValue<int32>("NumChannels");
+		bIsMainSubmix = EventData.GetValue<bool>("IsMainSubmix");
+	}
+
+	uint32 FSubmixAlivePingMessage::GetSizeOf() const
+	{
+		// Start with the fixed size of the struct
+		uint32 MemorySize = sizeof(FSubmixAlivePingMessage);
+
+		// Add any dynamically sized members
+		MemorySize += Name.GetAllocatedSize();
+
+		return MemorySize;
+	};
+
 	// FSubmixHasActivityMessage
 	FSubmixHasActivityMessage::FSubmixHasActivityMessage(const Trace::IAnalyzer::FOnEventContext& InContext)
 		: FSubmixMessageBase(InContext)
diff --git a/Engine/Plugins/AudioInsights/Source/AudioInsights/Private/Providers/SubmixTraceProvider.cpp b/Engine/Plugins/AudioInsights/Source/AudioInsights/Private/Providers/SubmixTraceProvider.cpp
--- a/Engine/Plugins/AudioInsights/Source/AudioInsights/Private/Providers/SubmixTraceProvider.cpp
+++ b/Engine/Plugins/AudioInsights/Source/AudioInsights/Private/Providers/SubmixTraceProvider.cpp
@@ -102,8 +102,8 @@ namespace UE::Audio::Insights
 
 				UE::Trace::IAnalyzer::FInterfaceBuilder& Builder = Context.InterfaceBuilder;
 
-				Builder.RouteEvent(RouteId_Loaded,                  "Audio", "SubmixAlivePing");
 				Builder.RouteEvent(RouteId_Loaded,                  "Audio", "SubmixLoaded");
+				Builder.RouteEvent(RouteId_AlivePing,               "Audio", "SubmixAlivePing");
 				Builder.RouteEvent(RouteId_HasActivity,             "Audio", "SubmixHasActivity");
 				Builder.RouteEvent(RouteId_EnvelopeFollowerEnabled, "Audio", "SubmixEnvelopeFollowerEnabled");
 				Builder.RouteEvent(RouteId_EnvelopeValues,          "Audio", "SubmixEnvelopeValues");
@@ -128,6 +128,16 @@ namespace UE::Audio::Insights
 						break;
 					}
 
+					case RouteId_AlivePing:
+					{
+#if WITH_EDITOR
+						CacheMessage<FSubmixAlivePingMessage>(Context, Messages.AlivePingMessages);
+#else
+						Messages.AlivePingMessages.Enqueue(FSubmixAlivePingMessage{ Context });
+#endif // WITH_EDITOR
+						break;
+					}
+
 					case RouteId_HasActivity:
 					{
 #if WITH_EDITOR
@@ -188,6 +198,7 @@ namespace UE::Audio::Insights
 			enum : uint16
 			{
 				RouteId_Loaded,
+				RouteId_AlivePing,
 				RouteId_HasActivity,
 				RouteId_EnvelopeFollowerEnabled,
 				RouteId_EnvelopeValues,
@@ -261,6 +272,25 @@ namespace UE::Audio::Insights
 			}
 		});
 
+		ProcessMessageQueue<FSubmixAlivePingMessage>(TraceMessages.AlivePingMessages, CreateEntry,
+		[this](const FSubmixAlivePingMessage& Msg, TSharedPtr<FSubmixDashboardEntry>* OutEntry)
+		{
+#if !WITH_EDITOR
+			if (SessionCachedMessages.IsValid())
+			{
+				SessionCachedMessages->AlivePingCachedMessages.EmplaceBack(Msg);
+			}
+#endif // !WITH_EDITOR
+
+			if (OutEntry)
+			{
+				FSubmixDashboardEntry& EntryRef = *OutEntry->Get();
+				EntryRef.Name = Msg.Name;
+				EntryRef.AudioMeterInfo->NumChannels = Msg.NumChannels;
+				EntryRef.bIsMainSubmix = Msg.bIsMainSubmix;
+			}
+		});
+
 		ProcessMessageQueue<FSubmixHasActivityMessage>(TraceMessages.HasActivityMessages, GetEntry,
 		[this](const FSubmixHasActivityMessage& Msg, TSharedPtr<FSubmixDashboardEntry>* OutEntry)
 		{
@@ -519,6 +549,42 @@ namespace UE::Audio::Insights
 		}
 #endif // WITH_EDITOR
 
+		// Collect all the submix alive ping messages registered until this point in time 
+		auto ProcessSubmixAlivePingMessage = [this](const FSubmixAlivePingMessage& Msg)
+		{
+			UpdateDeviceEntry(Msg.DeviceId, Msg.SubmixId, [&Msg](TSharedPtr<FSubmixDashboardEntry>& Entry)
+			{
+				if (!Entry.IsValid())
+				{
+					Entry = MakeShared<FSubmixDashboardEntry>();
+					Entry->SubmixId = Msg.SubmixId;
+				}
+
+				Entry->Timestamp = Msg.Timestamp;
+
+				Entry->Name = *Msg.Name;
+				Entry->AudioMeterInfo->NumChannels = Msg.NumChannels;
+				Entry->bIsMainSubmix = Msg.bIsMainSubmix;
+			});
+		};
+
+#if WITH_EDITOR
+		CacheManager.IterateTo<FSubmixAlivePingMessage>(SubmixMessageNames::IsAlivePing, TimeMarker, [&ProcessSubmixAlivePingMessage](const FSubmixAlivePingMessage& SubmixAlivePingCachedMessage)
+		{
+			ProcessSubmixAlivePingMessage(SubmixAlivePingCachedMessage);
+		});
+#else
+		for (const FSubmixAlivePingMessage& SubmixAlivePingCachedMessage : SessionCachedMessages->AlivePingCachedMessages)
+		{
+			if (SubmixAlivePingCachedMessage.Timestamp > TimeMarker)
+			{
+				break;
+			}
+
+			ProcessSubmixAlivePingMessage(SubmixAlivePingCachedMessage);
+		}
+#endif // WITH_EDITOR
+
 		// Selectively remove loaded messages collected in the step above by knowing which submixes were unloaded.
 		// With this we will know what are the available submixes at this point in time.
 		auto ProcessSubmixUnloadedMessage = [this](const FSubmixUnloadedMessage& Msg)
diff --git a/Engine/Plugins/AudioInsights/Source/AudioInsights/Public/Messages/SubmixTraceMessages.h b/Engine/Plugins/AudioInsights/Source/AudioInsights/Public/Messages/SubmixTraceMessages.h
--- a/Engine/Plugins/AudioInsights/Source/AudioInsights/Public/Messages/SubmixTraceMessages.h
+++ b/Engine/Plugins/AudioInsights/Source/AudioInsights/Public/Messages/SubmixTraceMessages.h
@@ -49,6 +49,19 @@ namespace UE::Audio::Insights
 		bool bIsMainSubmix = false;
 	};
 
+	struct FSubmixAlivePingMessage : public FSubmixMessageBase
+	{
+		FSubmixAlivePingMessage() = default;
+		FSubmixAlivePingMessage(const Trace::IAnalyzer::FOnEventContext& InContext);
+
+		virtual uint32 GetSizeOf() const override;
+		virtual const FName GetMessageName() const override { return SubmixMessageNames::IsAlivePing; }
+
+		FString Name;
+		int32 NumChannels = 0;
+		bool bIsMainSubmix = false;
+	};
+
 	struct FSubmixHasActivityMessage : public FSubmixMessageBase
 	{
 		FSubmixHasActivityMessage() = default;
@@ -95,6 +108,7 @@ namespace UE::Audio::Insights
 	class FSubmixMessages
 	{
 		TAnalyzerMessageQueue<FSubmixLoadedMessage> LoadedMessages;
+		TAnalyzerMessageQueue<FSubmixAlivePingMessage> AlivePingMessages;
 		TAnalyzerMessageQueue<FSubmixHasActivityMessage> HasActivityMessages;
 		TAnalyzerMessageQueue<FSubmixEnvelopeFollowerEnabledMessage> EnvelopeFollowerEnabledMessages;
 		TAnalyzerMessageQueue<FSubmixEnvelopeValuesMessage> EnvelopeValuesMessages;
@@ -108,6 +122,7 @@ namespace UE::Audio::Insights
 	{
 		FSubmixSessionCachedMessages(TraceServices::IAnalysisSession& InSession)
 			: LoadedCachedMessages(InSession.GetLinearAllocator(), 4096)
+			, AlivePingCachedMessages(InSession.GetLinearAllocator(), 4096)
 			, HasActivityCachedMessages(InSession.GetLinearAllocator(), 4096)
 			, EnvelopeFollowerEnabledCachedMessages(InSession.GetLinearAllocator(), 4096)
 			, EnvelopeValuesCachedMessages(InSession.GetLinearAllocator(), 4096)
@@ -117,6 +132,7 @@ namespace UE::Audio::Insights
 		}
 
 		TraceServices::TPagedArray<FSubmixLoadedMessage> LoadedCachedMessages;
+		TraceServices::TPagedArray<FSubmixAlivePingMessage> AlivePingCachedMessages;
 		TraceServices::TPagedArray<FSubmixHasActivityMessage> HasActivityCachedMessages;
 		TraceServices::TPagedArray<FSubmixEnvelopeFollowerEnabledMessage> EnvelopeFollowerEnabledCachedMessages;
 		TraceServices::TPagedArray<FSubmixEnvelopeValuesMessage> EnvelopeValuesCachedMessages;
