Move Audio Insights to a Beta plugin status. Enable it by default.

diff --git a/Engine/Plugins/AudioInsights/AudioInsights.uplugin b/Engine/Plugins/AudioInsights/AudioInsights.uplugin
--- a/Engine/Plugins/AudioInsights/AudioInsights.uplugin
+++ b/Engine/Plugins/AudioInsights/AudioInsights.uplugin
@@ -10,19 +10,18 @@
 	"DocsURL": "",
 	"MarketplaceURL": "",
 	"SupportURL": "",
+	"EnabledByDefault": true,
 	"CanContainContent": true,
-	"IsBetaVersion": false,
-	"IsExperimentalVersion": true,
+	"IsBetaVersion": true,
+	"IsExperimentalVersion": false,
 	"Installed": false,
 	"SupportedPrograms": [ "UnrealInsights" ],
-	"Modules":
-	[
+	"Modules": [
 		{
 			"Name": "AudioInsights",
 			"Type": "EditorAndProgram",
 			"LoadingPhase": "Default",
-			"ProgramAllowList":
-			[
+			"ProgramAllowList": [
 				"UnrealInsights"
 			]
 		},
@@ -32,8 +31,7 @@
 			"LoadingPhase": "Default"
 		}
 	],
-	"Plugins":
-	[
+	"Plugins": [
 		{
 			"Name": "AudioWidgets",
 			"Enabled": true
