AudioInsights: Fix "AudioDeviceManager.h" missing during compilation

#rnx

diff --git a/Engine/Plugins/AudioInsights/Source/AudioInsights/AudioInsights.Build.cs b/Engine/Plugins/AudioInsights/Source/AudioInsights/AudioInsights.Build.cs
--- a/Engine/Plugins/AudioInsights/Source/AudioInsights/AudioInsights.Build.cs
+++ b/Engine/Plugins/AudioInsights/Source/AudioInsights/AudioInsights.Build.cs
@@ -41,10 +41,15 @@ public class AudioInsights : ModuleRules
 
 		if (Target.Type == TargetType.Editor)
 		{
+			PublicDependencyModuleNames.AddRange(
+				new string[]
+				{
+					"Engine",
+				});
+			
 			PrivateDependencyModuleNames.AddRange(
 			new string[]
 			{
-				"Engine",
 				"UnrealEd",
 			});
 		}
