Basic MusicEventSystem Audio Insights dashboard

#rnx
#tests In Editor

diff --git a/Engine/Plugins/AudioInsights/Source/AudioInsights/Public/AudioInsightsConstants.h b/Engine/Plugins/AudioInsights/Source/AudioInsights/Public/AudioInsightsConstants.h
new file mode 100644
--- /dev/null
+++ b/Engine/Plugins/AudioInsights/Source/AudioInsights/Public/AudioInsightsConstants.h
@@ -0,0 +1,10 @@
+﻿// Copyright Epic Games, Inc. All Rights Reserved.
+#pragma once
+
+#include "HAL/Platform.h"
+
+namespace UE::Audio::Insights
+{
+	inline const ANSICHAR* AudioLoggerName = "Audio";
+	inline const ANSICHAR* TimestampFieldName = "Timestamp";
+}
diff --git a/Engine/Plugins/AudioInsights/Source/AudioInsights/Public/AudioInsightsDataSource.h b/Engine/Plugins/AudioInsights/Source/AudioInsights/Public/AudioInsightsDataSource.h
--- a/Engine/Plugins/AudioInsights/Source/AudioInsights/Public/AudioInsightsDataSource.h
+++ b/Engine/Plugins/AudioInsights/Source/AudioInsights/Public/AudioInsightsDataSource.h
@@ -15,6 +15,12 @@ namespace UE::Audio::Insights
 
 		virtual bool IsValid() const = 0;
 	};
+	
+	template <typename T>
+	const T& CastEntry(const IDashboardDataViewEntry& InData)
+	{
+		return static_cast<const T&>(InData);
+	}
 
 	class IDashboardDataTreeViewEntry : public TSharedFromThis<IDashboardDataTreeViewEntry>
 	{
diff --git a/Engine/Plugins/AudioInsights/Source/AudioInsights/Public/AudioInsightsTraceProviderBase.h b/Engine/Plugins/AudioInsights/Source/AudioInsights/Public/AudioInsightsTraceProviderBase.h
--- a/Engine/Plugins/AudioInsights/Source/AudioInsights/Public/AudioInsightsTraceProviderBase.h
+++ b/Engine/Plugins/AudioInsights/Source/AudioInsights/Public/AudioInsightsTraceProviderBase.h
@@ -102,7 +102,7 @@ namespace UE::Audio::Insights
 			AUDIOINSIGHTS_API virtual bool OnEventFailure(uint16 RouteId, EStyle Style, const FOnEventContext& Context);
 
 			// Override in derived to handle incoming trace messages
-			AUDIOINSIGHTS_API virtual bool OnHandleEvent(uint16 RouteId, EStyle Style, const FOnEventContext& Context) = 0;
+			virtual bool OnHandleEvent(uint16 RouteId, EStyle Style, const FOnEventContext& Context) = 0;
 
 			AUDIOINSIGHTS_API bool ShouldIgnoreNewEvents() const;
 			AUDIOINSIGHTS_API virtual bool ShouldProcessNewEvents() const;
@@ -335,6 +335,7 @@ namespace UE::Audio::Insights
 
 		void UpdateDeviceEntry(::Audio::FDeviceId InDeviceId, const EntryKey& InKey, TFunctionRef<void(EntryType&)> InEntryMutator)
 		{
+			ensure(InDeviceId != INDEX_NONE);
 			FDeviceData& DeviceData = DeviceDataMap.FindOrAdd(InDeviceId);
 			EntryType& Entry = DeviceData.FindOrAdd(InKey);
 			InEntryMutator(Entry);
@@ -363,7 +364,7 @@ namespace UE::Audio::Insights
 #if WITH_EDITOR
 		void RegisterDelegates()
 		{
-			FTraceAuxiliary::OnTraceStarted.AddRaw(this, &TDeviceDataMapTraceProvider::OnTraceStarted);			
+			FTraceAuxiliary::OnTraceStarted.AddRaw(this, &TDeviceDataMapTraceProvider::OnTraceStarted);
 			FTraceAuxiliary::OnTraceStopped.AddRaw(this, &TDeviceDataMapTraceProvider::OnTraceStopped);
 			FAudioDeviceManagerDelegates::OnAudioDeviceCreated.AddRaw(this, &TDeviceDataMapTraceProvider::OnAudioDeviceCreated);
 		}
diff --git a/Engine/Plugins/AudioInsights/Source/AudioInsights/Public/Views/TableDashboardViewFactory.h b/Engine/Plugins/AudioInsights/Source/AudioInsights/Public/Views/TableDashboardViewFactory.h
--- a/Engine/Plugins/AudioInsights/Source/AudioInsights/Public/Views/TableDashboardViewFactory.h
+++ b/Engine/Plugins/AudioInsights/Source/AudioInsights/Public/Views/TableDashboardViewFactory.h
@@ -145,6 +145,30 @@ namespace UE::Audio::Insights
 			return false;
 		}
 
+		template <typename TDataViewEntry>
+		void SortByPredicate(TFunctionRef<bool(const TDataViewEntry&, const TDataViewEntry&)> Predicate)
+		{
+			auto SortDashboardEntries = [this](const TSharedPtr<IDashboardDataViewEntry>& First, const TSharedPtr<IDashboardDataViewEntry>& Second, TFunctionRef<bool(const TDataViewEntry&, const TDataViewEntry&)> Predicate_Internal)
+				{
+					return Predicate_Internal(CastEntry<TDataViewEntry>(*First), CastEntry<TDataViewEntry>(*Second));
+				};
+
+			if (SortMode == EColumnSortMode::Ascending)
+			{
+				DataViewEntries.Sort([&SortDashboardEntries, &Predicate](const TSharedPtr<IDashboardDataViewEntry>& A, const TSharedPtr<IDashboardDataViewEntry>& B)
+				{
+					return SortDashboardEntries(A, B, Predicate);
+				});
+			}
+			else if (SortMode == EColumnSortMode::Descending)
+			{
+				DataViewEntries.Sort([&SortDashboardEntries, &Predicate](const TSharedPtr<IDashboardDataViewEntry>& A, const TSharedPtr<IDashboardDataViewEntry>& B)
+				{
+					return SortDashboardEntries(B, A, Predicate);
+				});
+			}
+		}
+
 #if WITH_EDITOR
 		virtual bool IsDebugDrawEnabled() const { return false; }
 		virtual void DebugDraw(float InElapsed, const TArray<TSharedPtr<IDashboardDataViewEntry>>& InSelectedItems, ::Audio::FDeviceId InAudioDeviceId) const { };
