[Audio Insights] Expose IsModuleLoaded in IAudioInsightsModule

diff --git a/Engine/Plugins/AudioInsights/Source/AudioInsights/Private/AudioInsightsModule.cpp b/Engine/Plugins/AudioInsights/Source/AudioInsights/Private/AudioInsightsModule.cpp
--- a/Engine/Plugins/AudioInsights/Source/AudioInsights/Private/AudioInsightsModule.cpp
+++ b/Engine/Plugins/AudioInsights/Source/AudioInsights/Private/AudioInsightsModule.cpp
@@ -132,6 +132,11 @@ namespace UE::Audio::Insights
 		return DashboardFactory->GetDeviceId();
 	}
 
+	bool FAudioInsightsModule::IsModuleLoaded()
+	{
+		return FModuleManager::Get().IsModuleLoaded("AudioInsights");
+	}
+
 	FAudioInsightsModule& FAudioInsightsModule::GetChecked()
 	{
 		return static_cast<FAudioInsightsModule&>(FModuleManager::LoadModuleChecked<IAudioInsightsModule>("AudioInsights"));
diff --git a/Engine/Plugins/AudioInsights/Source/AudioInsights/Private/AudioInsightsModule.h b/Engine/Plugins/AudioInsights/Source/AudioInsights/Private/AudioInsightsModule.h
--- a/Engine/Plugins/AudioInsights/Source/AudioInsights/Private/AudioInsightsModule.h
+++ b/Engine/Plugins/AudioInsights/Source/AudioInsights/Private/AudioInsightsModule.h
@@ -36,6 +36,7 @@ namespace UE::Audio::Insights
 		virtual void UnregisterDashboardViewFactory(FName InName) override;
 		virtual ::Audio::FDeviceId GetDeviceId() const override;
 
+		static bool IsModuleLoaded();
 		static FAudioInsightsModule& GetChecked();
 		static FAudioInsightsModule* GetModulePtr();
 		virtual IAudioInsightsTraceModule& GetTraceModule() override;
diff --git a/Engine/Plugins/AudioInsights/Source/AudioInsights/Private/IAudioInsightsModule.cpp b/Engine/Plugins/AudioInsights/Source/AudioInsights/Private/IAudioInsightsModule.cpp
--- a/Engine/Plugins/AudioInsights/Source/AudioInsights/Private/IAudioInsightsModule.cpp
+++ b/Engine/Plugins/AudioInsights/Source/AudioInsights/Private/IAudioInsightsModule.cpp
@@ -17,6 +17,11 @@ UE::Audio::Insights::FAudioInsightsCacheManager& IAudioInsightsModule::GetCacheM
 }
 #endif // WITH_EDITOR
 
+bool IAudioInsightsModule::IsModuleLoaded()
+{
+	return UE::Audio::Insights::FAudioInsightsModule::IsModuleLoaded();
+}
+
 IAudioInsightsModule& IAudioInsightsModule::GetChecked()
 {
 	return static_cast<IAudioInsightsModule&>(UE::Audio::Insights::FAudioInsightsModule::GetChecked());
diff --git a/Engine/Plugins/AudioInsights/Source/AudioInsights/Public/IAudioInsightsModule.h b/Engine/Plugins/AudioInsights/Source/AudioInsights/Public/IAudioInsightsModule.h
--- a/Engine/Plugins/AudioInsights/Source/AudioInsights/Public/IAudioInsightsModule.h
+++ b/Engine/Plugins/AudioInsights/Source/AudioInsights/Public/IAudioInsightsModule.h
@@ -36,6 +36,8 @@ public:
 
 	virtual TSharedRef<SDockTab> CreateDashboardTabWidget(const FSpawnTabArgs& Args) = 0;
 
+	UE_API static bool IsModuleLoaded();
+
 	static UE_API IAudioInsightsModule& GetChecked();
 
 #ifdef WITH_EDITOR
