[Audio Insights] AudioInsightsDashboardFactory: Allow dashboards from external plugins to set they own icon in the View list

diff --git a/Engine/Plugins/AudioInsights/Source/AudioInsights/Private/AudioInsightsDashboardFactory.cpp b/Engine/Plugins/AudioInsights/Source/AudioInsights/Private/AudioInsightsDashboardFactory.cpp
--- a/Engine/Plugins/AudioInsights/Source/AudioInsights/Private/AudioInsightsDashboardFactory.cpp
+++ b/Engine/Plugins/AudioInsights/Source/AudioInsights/Private/AudioInsightsDashboardFactory.cpp
@@ -143,9 +143,10 @@ namespace UE::Audio::Insights
 					const FName& FactoryName = KVP.Key;
 					const TSharedPtr<IDashboardViewFactory>& Factory = KVP.Value;
 
-					MenuBuilder.AddMenuEntry(Factory->GetDisplayName(),
+					MenuBuilder.AddMenuEntry(
+						Factory->GetDisplayName(),
 						FText::GetEmpty(),
-						FSlateStyle::Get().CreateIcon(Factory->GetIcon().GetStyleName()),
+						Factory->GetIcon(),
 						FUIAction(FExecuteAction::CreateLambda([this, FactoryName]()
 						{
 							if (DashboardTabManager.IsValid())
