[Audio Insights] Fix cache status column having an empty field in the visible columns settings and in the header row

diff --git a/Engine/Plugins/AudioInsights/Source/AudioInsights/Private/Views/AudioEventLogDashboardViewFactory.cpp b/Engine/Plugins/AudioInsights/Source/AudioInsights/Private/Views/AudioEventLogDashboardViewFactory.cpp
--- a/Engine/Plugins/AudioInsights/Source/AudioInsights/Private/Views/AudioEventLogDashboardViewFactory.cpp
+++ b/Engine/Plugins/AudioInsights/Source/AudioInsights/Private/Views/AudioEventLogDashboardViewFactory.cpp
@@ -77,6 +77,7 @@ namespace UE::Audio::Insights
 
 		const FText ClearLogText = LOCTEXT("AudioEventLog_ClearLog_Text", "Clear Log");
 		const FText ClearLogTooltipText = LOCTEXT("AudioEventLog_ClearLog_Tooltip", "Clears the event log for the currently selected world.");
+		const FText CacheStatusColumnDisplayName = LOCTEXT("AudioEventLog_CacheStatus", "Cache Status");
 	} // namespace AudioEventLogPrivate
 
 	const TMap<FString, TSet<FString>> FAudioEventLogDashboardViewFactory::GetInitEventTypeFilters()
@@ -423,14 +424,14 @@ namespace UE::Audio::Insights
 				{
 					CacheStatusColumnID,
 					{
-						FText::GetEmpty(),
+						CacheStatusColumnDisplayName,
 						[](const IDashboardDataViewEntry& InData) { return FText::GetEmpty(); }, /* GetDisplayValue */
 						[](const IDashboardDataViewEntry& InData) -> FName
 						{
 							return FName("AudioInsights.Icon.EventLog.MarkedForDelete");
 						},
 						!VisibleColumns.GetIsVisible(CacheStatusColumnID) /* bDefaultHidden */,
-						0.075f /* FillWidth */,
+						0.1f /* FillWidth */,
 						HAlign_Center /* Alignment */,
 						[this](const IDashboardDataViewEntry& InData) -> FLinearColor /* GetIconColor */
 						{
@@ -459,7 +460,10 @@ namespace UE::Audio::Insights
 							}
 
 							return FText::GetEmpty();
-						}
+						},
+						FName("AudioInsights.Icon.EventLog.MarkedForDelete"), /* GetHeaderRowIconName */
+						CacheStatusColumnDisplayName, /* HeaderRowTooltip */
+						false /* bShowDisplayName */
 					}
 				},
 #endif // WITH_EDITOR
diff --git a/Engine/Plugins/AudioInsights/Source/AudioInsights/Private/Views/TableDashboardViewFactory.cpp b/Engine/Plugins/AudioInsights/Source/AudioInsights/Private/Views/TableDashboardViewFactory.cpp
--- a/Engine/Plugins/AudioInsights/Source/AudioInsights/Private/Views/TableDashboardViewFactory.cpp
+++ b/Engine/Plugins/AudioInsights/Source/AudioInsights/Private/Views/TableDashboardViewFactory.cpp
@@ -223,7 +223,30 @@ namespace UE::Audio::Insights
 				.HAlignCell(ColumnData.Alignment)
 				.FillWidth(ColumnData.FillWidth)
 				.SortMode(this, &FTraceTableDashboardViewFactory::GetColumnSortMode, ColumnName)
-				.OnSort(this, &FTraceTableDashboardViewFactory::OnColumnSortModeChanged);
+				.OnSort(this, &FTraceTableDashboardViewFactory::OnColumnSortModeChanged)
+				.HeaderContent()
+				[
+					SNew(SHorizontalBox)
+
+					// Icon (optional)
+					+ SHorizontalBox::Slot()
+					.AutoWidth()
+					.Padding(ColumnData.HeaderRowIconName != NAME_None ? 4.0f : 0.0f, 3.0f)
+					[
+						ColumnData.HeaderRowIconName != NAME_None
+							? SNew(SImage).Image(FSlateStyle::Get().GetBrush(ColumnData.HeaderRowIconName)).ToolTipText(ColumnData.HeaderRowTooltip)
+							: SNullWidget::NullWidget
+					]
+					// Text (optional)
+					+ SHorizontalBox::Slot()
+					.AutoWidth()
+					.Padding(0.0f, 3.0f, 0.0f, 3.0f)
+					[
+						ColumnData.bShowDisplayName && !ColumnData.DisplayName.IsEmpty()
+							? SNew(STextBlock).Text(ColumnData.DisplayName).ToolTipText(ColumnData.HeaderRowTooltip)
+							: SNullWidget::NullWidget
+					]
+				];
 
 			// .HiddenColumnsList workaround:
 			// simulate what SHeaderRow::AddColumn( const FColumn::FArguments& NewColumnArgs ) does but allowing us to modify the bIsVisible property
@@ -389,6 +412,11 @@ namespace UE::Audio::Insights
 			{
 				for (TSharedPtr<FTraceProviderBase> Provider : Providers)
 				{
+					if (!Provider.IsValid())
+					{
+						continue;
+					}
+
 					const FName ProviderName = Provider->GetName();
 					const uint64 LastUpdateId = Provider->GetLastUpdateId();
 					UpdateIds.FindOrAdd(ProviderName) = LastUpdateId;
diff --git a/Engine/Plugins/AudioInsights/Source/AudioInsights/Public/Views/TableDashboardViewFactory.h b/Engine/Plugins/AudioInsights/Source/AudioInsights/Public/Views/TableDashboardViewFactory.h
--- a/Engine/Plugins/AudioInsights/Source/AudioInsights/Public/Views/TableDashboardViewFactory.h
+++ b/Engine/Plugins/AudioInsights/Source/AudioInsights/Public/Views/TableDashboardViewFactory.h
@@ -184,6 +184,9 @@ namespace UE::Audio::Insights
 			const EHorizontalAlignment Alignment = HAlign_Left;
 			const TFunction<FLinearColor(const IDashboardDataViewEntry&)> GetIconColor = nullptr;
 			const TFunction<FText(const IDashboardDataViewEntry&)> GetIconTooltip = nullptr;
+			const FName HeaderRowIconName = NAME_None;
+			const FText HeaderRowTooltip = FText::GetEmpty();
+			const bool bShowDisplayName = true;
 		};
 
 		virtual const TMap<FName, FColumnData>& GetColumns() const = 0;
