[Audio Insights] Support cache manager inside standalone part 01: Create non-limited cache container inside AudioInsightsCacheManager for standalone

diff --git a/Engine/Plugins/AudioInsights/Source/AudioInsights/Private/Cache/AudioInsightsCacheManager.cpp b/Engine/Plugins/AudioInsights/Source/AudioInsights/Private/Cache/AudioInsightsCacheManager.cpp
--- a/Engine/Plugins/AudioInsights/Source/AudioInsights/Private/Cache/AudioInsightsCacheManager.cpp
+++ b/Engine/Plugins/AudioInsights/Source/AudioInsights/Private/Cache/AudioInsightsCacheManager.cpp
@@ -2,6 +2,11 @@
 #include "Cache/AudioInsightsCacheManager.h"
 
 #include "AudioInsightsModule.h"
+#if !WITH_EDITOR
+#include "AudioInsightsComponent.h"
+#include "Insights/IInsightsManager.h"
+#include "Insights/IUnrealInsightsModule.h"
+#endif // !WITH_EDITOR
 
 namespace UE::Audio::Insights
 {
@@ -13,7 +18,6 @@ namespace UE::Audio::Insights
 
 	FAudioInsightsCacheManager::FAudioInsightsCacheManager()
 		: Allocator(8 << 20) // 8mB for storage of pointers to messages
-		, Cache(CacheManagerPrivate::NumChunks)
 	{
 		check(CacheManagerPrivate::NumChunks > 0u);
 
@@ -33,7 +37,10 @@ namespace UE::Audio::Insights
 
 	void FAudioInsightsCacheManager::AddMessageToCache(TSharedPtr<IAudioCachedMessage> Message)
 	{
-		NewMessageQueue.Enqueue(MoveTemp(Message));
+		if (Message.IsValid())
+		{
+			NewMessageQueue.Enqueue(MoveTemp(Message));
+		}
 	}
 
 	void FAudioInsightsCacheManager::ClearCache()
@@ -51,17 +58,15 @@ namespace UE::Audio::Insights
 	double FAudioInsightsCacheManager::GetCacheStartTimeStamp() const
 	{
 		check(StartIndex < Cache.Capacity());
-		check(Cache[StartIndex].IsValid());
 
-		return Cache[StartIndex]->GetChunkTimeRangeStart();
+		return Cache[StartIndex].GetChunkTimeRangeStart();
 	}
 
 	double FAudioInsightsCacheManager::GetCacheEndTimeStamp() const
 	{
 		check(WriteIndex < Cache.Capacity());
-		check(Cache[WriteIndex].IsValid());
 
-		return Cache[WriteIndex]->GetChunkTimeRangeEnd();
+		return Cache[WriteIndex].GetChunkTimeRangeEnd();
 	}
 
 	uint32 FAudioInsightsCacheManager::GetUsedCacheSize() const
@@ -69,7 +74,7 @@ namespace UE::Audio::Insights
 		uint32 UsedSpace = 0u;
 		for (uint32 Index = 0u; Index < Cache.Capacity(); ++Index)
 		{
-			UsedSpace += Cache[Index]->GetCurrentChunkSize();
+			UsedSpace += Cache[Index].GetCurrentChunkSize();
 		}
 
 		return UsedSpace;
@@ -77,12 +82,12 @@ namespace UE::Audio::Insights
 
 	uint32 FAudioInsightsCacheManager::GetMaxCacheSize() const
 	{
-		return CacheManagerPrivate::MaxChunkSize * CacheManagerPrivate::NumChunks;
+		return CacheManagerPrivate::MaxChunkSize * GetNumChunks();
 	}
 
 	uint32 FAudioInsightsCacheManager::GetNumChunks() const
 	{
-		return CacheManagerPrivate::NumChunks;
+		return Cache.Capacity();
 	}
 
 	uint32 FAudioInsightsCacheManager::GetNumUsedChunks() const
@@ -90,7 +95,7 @@ namespace UE::Audio::Insights
 		uint32 NumUsedChunks = 0u;
 		for (uint32 Index = 0u; Index < Cache.Capacity(); ++Index)
 		{
-			if (Cache[Index]->HasAnyData())
+			if (Cache[Index].HasAnyData())
 			{
 				++NumUsedChunks;
 			}
@@ -104,7 +109,7 @@ namespace UE::Audio::Insights
 		uint32 ReadIndex = StartIndex;
 		for (uint32 Index = 0u; Index < Cache.Capacity(); ++Index)
 		{
-			if (Cache[ReadIndex]->TimestampIsInChunkRange(Timestamp))
+			if (Cache[ReadIndex].TimestampIsInChunkRange(Timestamp))
 			{
 				return Index;
 			}
@@ -131,20 +136,17 @@ namespace UE::Audio::Insights
 			ReadIndex = Cache.GetNextIndex(ReadIndex);
 		}
 
-		if (!Cache[ReadIndex]->HasAnyData())
+		if (!Cache[ReadIndex].HasAnyData())
 		{
 			return nullptr;
 		}
 
-		return Cache[ReadIndex].Get();
+		return &Cache[ReadIndex];
 	}
 
 	void FAudioInsightsCacheManager::CreateCache()
 	{
-		for (uint32 Index = 0u; Index < Cache.Capacity(); ++Index)
-		{
-			Cache[Index] = MakeUnique<FAudioCachedMessageChunk>(Allocator, CacheManagerPrivate::MaxChunkSize);
-		}
+		Cache.ResetCache(Allocator);
 
 		WriteIndex = 0u;
 		StartIndex = 0u;
@@ -154,17 +156,29 @@ namespace UE::Audio::Insights
 	{
 		check(WriteIndex < Cache.Capacity());
 
-		FAudioCachedMessageChunk* Chunk = Cache[WriteIndex].Get();
+		FAudioCachedMessageChunk* Chunk = &Cache[WriteIndex];
 
 		const TArray<TSharedPtr<IAudioCachedMessage>> NewMessages = NewMessageQueue.DequeueAll();
 		for (const TSharedPtr<IAudioCachedMessage>& Message : NewMessages)
 		{
 			if (Chunk->IsChunkFull())
 			{
+#if WITH_EDITOR
 				// Move to next chunk
 				WriteIndex = Cache.GetNextIndex(WriteIndex);
+#else
+				const uint32 NextWriteIndex = Cache.GetNextIndex(WriteIndex);
+
+				// If our cache is not limited, we may need to reserve more chunks
+				if (NextWriteIndex >= Cache.Capacity())
+				{
+					Cache.ReserveNewChunks(Allocator, CacheManagerPrivate::NumChunks);
+				}
+
+				WriteIndex = NextWriteIndex;
+#endif // WITH_EDITOR
 
-				Chunk = Cache[WriteIndex].Get();
+				Chunk = &Cache[WriteIndex];
 
 				// If this new chunk already has data, clear it before we begin writing
 				if (Chunk->HasAnyData())
@@ -183,4 +197,143 @@ namespace UE::Audio::Insights
 			}
 		}
 	}
+
+	FAudioInsightsCache::FAudioInsightsCache()
+#if WITH_EDITOR
+		: EditorCache(CacheManagerPrivate::NumChunks)
+#endif // WITH_EDITOR
+	{
+#if !WITH_EDITOR
+		IUnrealInsightsModule& UnrealInsightsModule = FModuleManager::LoadModuleChecked<IUnrealInsightsModule>("TraceInsights");
+		if (UnrealInsightsModule.GetAnalysisSession().IsValid())
+		{
+			InitCacheIsLimited();
+		}
+		else
+		{
+			// The session isn't ready yet - request a notification when it goes live
+			TSharedPtr<::Insights::IInsightsManager> InsightsManager = UnrealInsightsModule.GetInsightsManager();
+			if (InsightsManager.IsValid())
+			{
+				InsightsManager->GetSessionChangedEvent().AddRaw(this, &FAudioInsightsCache::OnSessionChangedEvent);
+			}
+		}
+#endif // !WITH_EDITOR
+	}
+
+	FAudioInsightsCache::~FAudioInsightsCache()
+	{
+#if !WITH_EDITOR
+		if (IUnrealInsightsModule* UnrealInsightsModule = FModuleManager::GetModulePtr<IUnrealInsightsModule>("TraceInsights"))
+		{
+			TSharedPtr<::Insights::IInsightsManager> InsightsManager = UnrealInsightsModule->GetInsightsManager();
+			if (InsightsManager.IsValid())
+			{
+				InsightsManager->GetSessionChangedEvent().RemoveAll(this);
+			}
+		}
+#endif // !WITH_EDITOR
+	}
+
+	void FAudioInsightsCache::ResetCache(TraceServices::FSlabAllocator& Allocator)
+	{
+#if !WITH_EDITOR
+		StandaloneCache.Empty();
+#endif // !WITH_EDITOR
+
+		ReserveNewChunks(Allocator, CacheManagerPrivate::NumChunks);
+	}
+
+	void FAudioInsightsCache::ReserveNewChunks(TraceServices::FSlabAllocator& Allocator, const uint32 NumToReserve)
+	{
+#if !WITH_EDITOR
+		StandaloneCache.Reserve(StandaloneCache.Num() + NumToReserve);
+#endif // !WITH_EDITOR
+
+		for (uint32 Index = 0u; Index < NumToReserve; ++Index)
+		{
+#if WITH_EDITOR
+			EditorCache[Index] = MakeUnique<FAudioCachedMessageChunk>(Allocator, CacheManagerPrivate::MaxChunkSize);
+#else
+			StandaloneCache.Add(MakeUnique<FAudioCachedMessageChunk>(Allocator, CacheManagerPrivate::MaxChunkSize));
+#endif // WITH_EDITOR
+		}
+	}
+
+	uint32 FAudioInsightsCache::Capacity() const
+	{
+#if WITH_EDITOR
+		return EditorCache.Capacity();
+#else
+		return static_cast<uint32>(StandaloneCache.Num());
+#endif // WITH_EDITOR
+	}
+
+	uint32 FAudioInsightsCache::GetNextIndex(uint32 ReadIndex) const
+	{
+#if WITH_EDITOR
+		return EditorCache.GetNextIndex(ReadIndex);
+#else
+		if (bLimitStandaloneCacheSize)
+		{
+			ReadIndex++;
+			if (ReadIndex >= Capacity())
+			{
+				ReadIndex = 0u;
+			}
+		}
+		else
+		{
+			ReadIndex++;
+		}
+
+		return ReadIndex;
+#endif // WITH_EDITOR
+	}
+
+	uint32 FAudioInsightsCache::GetPreviousIndex(uint32 ReadIndex) const
+	{
+#if WITH_EDITOR
+		return EditorCache.GetPreviousIndex(ReadIndex);
+#else
+		if (ReadIndex > 0u)
+		{
+			ReadIndex--;
+		}
+		else if (bLimitStandaloneCacheSize)
+		{
+			ReadIndex = Capacity() - 1u;
+		}
+		else
+		{
+			ReadIndex = INDEX_NONE;
+		}
+
+		return ReadIndex;
+#endif // WITH_EDITOR
+	}
+
+#if !WITH_EDITOR
+	void FAudioInsightsCache::InitCacheIsLimited()
+	{
+		FAudioInsightsModule& AudioInsightsModule = FAudioInsightsModule::GetChecked();
+		const TSharedPtr<FAudioInsightsComponent> AudioInsightsComponent = AudioInsightsModule.GetAudioInsightsComponent();
+		if (AudioInsightsComponent.IsValid() && AudioInsightsComponent->GetIsLiveSession())
+		{
+			bLimitStandaloneCacheSize = true;
+		}
+	}
+
+	void FAudioInsightsCache::OnSessionChangedEvent()
+	{
+		InitCacheIsLimited();
+
+		IUnrealInsightsModule& UnrealInsightsModule = FModuleManager::LoadModuleChecked<IUnrealInsightsModule>("TraceInsights");
+		TSharedPtr<::Insights::IInsightsManager> InsightsManager = UnrealInsightsModule.GetInsightsManager();
+		if (InsightsManager.IsValid())
+		{
+			InsightsManager->GetSessionChangedEvent().RemoveAll(this);
+		}
+	}
+#endif // !WITH_EDITOR
 } // namespace UE::Audio::Insights
diff --git a/Engine/Plugins/AudioInsights/Source/AudioInsights/Public/Cache/AudioInsightsCacheManager.h b/Engine/Plugins/AudioInsights/Source/AudioInsights/Public/Cache/AudioInsightsCacheManager.h
--- a/Engine/Plugins/AudioInsights/Source/AudioInsights/Public/Cache/AudioInsightsCacheManager.h
+++ b/Engine/Plugins/AudioInsights/Source/AudioInsights/Public/Cache/AudioInsightsCacheManager.h
@@ -13,6 +13,56 @@
 
 namespace UE::Audio::Insights
 {
+	// Helper class that provides a similar API in Editor and Standalone builds to the cache container
+	// In Editor, the cache is a limited circular buffer
+	// In standalone, the cache is an array that may or may not be limited in size
+	// - When connected to a live session (either via direct or live trace), the array will be behave like a circular buffer
+	// - When loading in a .utrace file, the array will grow to hold all of the trace file data in memory
+	class FAudioInsightsCache
+	{
+	public:
+		FAudioInsightsCache();
+		virtual ~FAudioInsightsCache();
+
+		FAudioCachedMessageChunk& operator[](uint32 Index)
+		{
+#if WITH_EDITOR
+			return *EditorCache[Index].Get();
+#else
+			return *StandaloneCache[Index].Get();
+#endif // WITH_EDITOR
+		}
+
+		const FAudioCachedMessageChunk& operator[](uint32 Index) const
+		{
+#if WITH_EDITOR
+			return *EditorCache[Index].Get();
+#else
+			return *StandaloneCache[Index].Get();
+#endif // WITH_EDITOR
+		}
+
+		void ResetCache(TraceServices::FSlabAllocator& Allocator);
+		void ReserveNewChunks(TraceServices::FSlabAllocator& Allocator, const uint32 NumToReserve);
+
+		uint32 Capacity() const;
+		uint32 GetNextIndex(uint32 ReadIndex) const;
+		uint32 GetPreviousIndex(uint32 ReadIndex) const;
+
+	private:
+#if !WITH_EDITOR
+		void InitCacheIsLimited();
+		void OnSessionChangedEvent();
+#endif // !WITH_EDITOR
+
+#if WITH_EDITOR
+		TCircularBuffer<TUniquePtr<FAudioCachedMessageChunk>> EditorCache;
+#else
+		TArray<TUniquePtr<FAudioCachedMessageChunk>> StandaloneCache;
+		bool bLimitStandaloneCacheSize = false;
+#endif // WITH_EDITOR
+	};
+
 	class FAudioInsightsCacheManager
 	{
 	public:
@@ -50,13 +100,15 @@ namespace UE::Audio::Insights
 		// Will look backwards from the chosen timestamp until a message is found
 		// @param ID: When overloaded, the search algorithm will only return messages that have a matching ID.
 		//			  If left at INDEX_NONE, the search algorithm will return the closest message irrespective of ID
+		// @param BackOutTimestamp: When iterating backwards, if a chunk is detected with an end timestamp less than
+		//							the backout timestamp, the search is canceled and this returns nullptr
 		template<TIsCacheableMessage T>
-		const T* FindClosestMessage(const FName& MessageID, const double Timestamp, const uint64 ID = INDEX_NONE) const
+		const T* FindClosestMessage(const FName& MessageID, const double Timestamp, const uint64 ID = INDEX_NONE, const double BackOutTimestamp = -1.0) const
 		{
 			uint32 ReadIndex = INDEX_NONE;
 			for (uint32 Index = 0u; Index < Cache.Capacity(); ++Index)
 			{
-				const FAudioCachedMessageChunk& Chunk = *Cache[Index].Get();
+				const FAudioCachedMessageChunk& Chunk = Cache[Index];
 
 				if (Chunk.TimestampIsInChunkRange(Timestamp))
 				{
@@ -74,7 +126,12 @@ namespace UE::Audio::Insights
 			// Iterate backwards until we find a matching message
 			for (uint32 Index = ReadIndex; ; )
 			{
-				const FAudioCachedMessageChunk& Chunk = *Cache[Index].Get();
+				const FAudioCachedMessageChunk& Chunk = Cache[Index];
+				if (Chunk.GetChunkTimeRangeEnd() <= BackOutTimestamp)
+				{
+					break;
+				}
+
 				if (const T* FoundMessage = Chunk.FindClosestMessage<T>(MessageID, Timestamp, ID))
 				{
 					return FoundMessage;
@@ -86,6 +143,12 @@ namespace UE::Audio::Insights
 				else
 				{
 					Index = Cache.GetPreviousIndex(Index);
+#if !WITH_EDITOR
+					if (Index == INDEX_NONE)
+					{
+						break;
+					}
+#endif // !WITH_EDITOR
 				}
 			}
 
@@ -103,7 +166,7 @@ namespace UE::Audio::Insights
 			double LatestChunkStartTimestamp = -1.0;
 			for (uint32 Index = 0u; Index < Cache.Capacity(); ++Index)
 			{
-				const FAudioCachedMessageChunk& Chunk = *Cache[ReadIndex].Get();
+				const FAudioCachedMessageChunk& Chunk = Cache[ReadIndex];
 
 				const double ChunkStartTime = Chunk.GetChunkTimeRangeStart();
 				if (ChunkStartTime > End)
@@ -197,8 +260,8 @@ namespace UE::Audio::Insights
 		void ProcessPendingMessages(float DeltaTime);
 
 		TraceServices::FSlabAllocator Allocator;
-		TCircularBuffer<TUniquePtr<FAudioCachedMessageChunk>> Cache;
 
+		FAudioInsightsCache Cache;
 		TAnalyzerMessageQueue<TSharedPtr<IAudioCachedMessage>> NewMessageQueue;
 
 		uint32 WriteIndex = 0u;
