[Audio Insights] Support cache manager inside standalone part 02: Expose Cache manager to non-editor builds from modules

diff --git a/Engine/Plugins/AudioInsights/Source/AudioInsights/Private/AudioInsightsModule.cpp b/Engine/Plugins/AudioInsights/Source/AudioInsights/Private/AudioInsightsModule.cpp
--- a/Engine/Plugins/AudioInsights/Source/AudioInsights/Private/AudioInsightsModule.cpp
+++ b/Engine/Plugins/AudioInsights/Source/AudioInsights/Private/AudioInsightsModule.cpp
@@ -42,22 +42,21 @@ namespace UE::Audio::Insights
 			IModularFeatures::Get().RegisterModularFeature(TraceServices::ModuleFeatureName, TraceModule.Get());
 			IModularFeatures::Get().RegisterModularFeature(IRewindDebuggerRuntimeExtension::ModularFeatureName, RewindDebuggerExtension.Get());
 
+#if !WITH_EDITOR
+			AudioInsightsComponent = FAudioInsightsComponent::CreateInstance();
+
+			IUnrealInsightsModule& UnrealInsightsModule = FModuleManager::LoadModuleChecked<IUnrealInsightsModule>("TraceInsights");
+			UnrealInsightsModule.RegisterComponent(AudioInsightsComponent);
+#endif // !WITH_EDITOR
+
 			DashboardFactory = MakeShared<FDashboardFactory>();
-			
 			FDashboardAssetCommands::Register();
 
-#if WITH_EDITOR
 			CacheManager = MakeUnique<FAudioInsightsCacheManager>();
-#endif // WITH_EDITOR
 
 #if !WITH_EDITOR
 			IModularFeatures::Get().RegisterModularFeature(UE::Insights::Timing::TimingViewExtenderFeatureName, &AudioInsightsTimingViewExtender);
 
-			AudioInsightsComponent = FAudioInsightsComponent::CreateInstance();
-
-			IUnrealInsightsModule& UnrealInsightsModule = FModuleManager::LoadModuleChecked<IUnrealInsightsModule>("TraceInsights");
-			UnrealInsightsModule.RegisterComponent(AudioInsightsComponent);
-
 			TSharedRef<FSoundDashboardViewFactory> SoundsDashboard = MakeShared<FSoundDashboardViewFactory>();
 			TSharedRef<FSoundPlotsDashboardViewFactory> PlotsDashboard = MakeShared<FSoundPlotsDashboardViewFactory>();
 			TSharedRef<FSubmixDashboardViewFactory> SubmixDashboard = MakeShared<FSubmixDashboardViewFactory>();
@@ -98,7 +97,7 @@ namespace UE::Audio::Insights
 
 			AudioInsightsComponent.Reset();
 
-			IModularFeatures::Get().UnregisterModularFeature(TraceServices::ModuleFeatureName, &AudioInsightsTimingViewExtender);
+			IModularFeatures::Get().UnregisterModularFeature(UE::Insights::Timing::TimingViewExtenderFeatureName, &AudioInsightsTimingViewExtender);
 #endif // !WITH_EDITOR
 
 			FDashboardAssetCommands::Unregister();
@@ -111,24 +110,34 @@ namespace UE::Audio::Insights
 			IModularFeatures::Get().UnregisterModularFeature(IRewindDebuggerRuntimeExtension::ModularFeatureName, RewindDebuggerExtension.Get());
 			RewindDebuggerExtension.Reset();
 
-#if WITH_EDITOR
 			CacheManager.Reset();
-#endif // WITH_EDITOR
 		}
 	}
 
 	void FAudioInsightsModule::RegisterDashboardViewFactory(TSharedRef<IDashboardViewFactory> InDashboardFactory)
 	{
-		DashboardFactory->RegisterViewFactory(InDashboardFactory);
+		check(!IsRunningCommandlet() && DashboardFactory.IsValid());
+
+		if (DashboardFactory.IsValid())
+		{
+			DashboardFactory->RegisterViewFactory(InDashboardFactory);
+		}
 	}
 
 	void FAudioInsightsModule::UnregisterDashboardViewFactory(FName InName)
 	{
-		DashboardFactory->UnregisterViewFactory(InName);
+		check(!IsRunningCommandlet() && DashboardFactory.IsValid());
+
+		if (DashboardFactory.IsValid())
+		{
+			DashboardFactory->UnregisterViewFactory(InName);
+		}
 	}
 
 	::Audio::FDeviceId FAudioInsightsModule::GetDeviceId() const
 	{
+		check(!IsRunningCommandlet() && DashboardFactory.IsValid());
+
 		return DashboardFactory->GetDeviceId();
 	}
 
@@ -154,28 +163,36 @@ namespace UE::Audio::Insights
 
 	TSharedRef<FDashboardFactory> FAudioInsightsModule::GetDashboardFactory()
 	{
+		check(!IsRunningCommandlet() && DashboardFactory.IsValid());
+
 		return DashboardFactory->AsShared();
 	}
 
 	const TSharedRef<FDashboardFactory> FAudioInsightsModule::GetDashboardFactory() const
 	{
+		check(!IsRunningCommandlet() && DashboardFactory.IsValid());
+
 		return DashboardFactory->AsShared();
 	}
 
 	IAudioInsightsTraceModule& FAudioInsightsModule::GetTraceModule()
 	{
+		check(!IsRunningCommandlet() && TraceModule.IsValid());
+
 		return *TraceModule;
 	}
 
-#if WITH_EDITOR
 	FAudioInsightsCacheManager& FAudioInsightsModule::GetCacheManager()
 	{
+		check(!IsRunningCommandlet() && CacheManager.IsValid());
+
 		return *CacheManager;
 	}
-#endif // WITH_EDITOR
 
 	TSharedRef<SDockTab> FAudioInsightsModule::CreateDashboardTabWidget(const FSpawnTabArgs& Args)
 	{
+		check(!IsRunningCommandlet() && DashboardFactory.IsValid());
+
 		return DashboardFactory->MakeDockTabWidget(Args);
 	}
 } // namespace UE::Audio::Insights
diff --git a/Engine/Plugins/AudioInsights/Source/AudioInsights/Private/AudioInsightsModule.h b/Engine/Plugins/AudioInsights/Source/AudioInsights/Private/AudioInsightsModule.h
--- a/Engine/Plugins/AudioInsights/Source/AudioInsights/Private/AudioInsightsModule.h
+++ b/Engine/Plugins/AudioInsights/Source/AudioInsights/Private/AudioInsightsModule.h
@@ -40,10 +40,7 @@ namespace UE::Audio::Insights
 		static FAudioInsightsModule& GetChecked();
 		static FAudioInsightsModule* GetModulePtr();
 		virtual IAudioInsightsTraceModule& GetTraceModule() override;
-
-#if WITH_EDITOR
 		virtual class FAudioInsightsCacheManager& GetCacheManager() override;
-#endif // WITH_EDITOR
 
 #if !WITH_EDITOR
 		TSharedPtr<FAudioInsightsComponent> GetAudioInsightsComponent() { return AudioInsightsComponent; };
@@ -61,10 +58,7 @@ namespace UE::Audio::Insights
 		TSharedPtr<FDashboardFactory> DashboardFactory;
 		TUniquePtr<FTraceModule> TraceModule;
 		TUniquePtr<FRewindDebuggerAudioInsightsRuntime> RewindDebuggerExtension;
-
-#if WITH_EDITOR
 		TUniquePtr<class FAudioInsightsCacheManager> CacheManager;
-#endif // WITH_EDITOR
 
 #if !WITH_EDITOR
 		TSharedPtr<FAudioInsightsComponent> AudioInsightsComponent;
diff --git a/Engine/Plugins/AudioInsights/Source/AudioInsights/Private/AudioInsightsTraceModule.cpp b/Engine/Plugins/AudioInsights/Source/AudioInsights/Private/AudioInsightsTraceModule.cpp
--- a/Engine/Plugins/AudioInsights/Source/AudioInsights/Private/AudioInsightsTraceModule.cpp
+++ b/Engine/Plugins/AudioInsights/Source/AudioInsights/Private/AudioInsightsTraceModule.cpp
@@ -305,10 +305,8 @@ namespace UE::Audio::Insights
 			InSession.AddAnalyzer(Provider->ConstructAnalyzer(InSession));
 		}
 
-#if WITH_EDITOR
 		FAudioInsightsCacheManager& CacheManager = FAudioInsightsModule::GetChecked().GetCacheManager();
 		CacheManager.ClearCache();
-#endif // WITH_EDITOR
 
 		OnAnalysisStarting.Broadcast(FPlatformTime::Seconds());
 	}
diff --git a/Engine/Plugins/AudioInsights/Source/AudioInsights/Private/IAudioInsightsModule.cpp b/Engine/Plugins/AudioInsights/Source/AudioInsights/Private/IAudioInsightsModule.cpp
--- a/Engine/Plugins/AudioInsights/Source/AudioInsights/Private/IAudioInsightsModule.cpp
+++ b/Engine/Plugins/AudioInsights/Source/AudioInsights/Private/IAudioInsightsModule.cpp
@@ -9,13 +9,11 @@ IAudioInsightsTraceModule& IAudioInsightsModule::GetTraceModule()
 	return AudioInsightsModule.GetTraceModule();
 }
 
-#if WITH_EDITOR
 UE::Audio::Insights::FAudioInsightsCacheManager& IAudioInsightsModule::GetCacheManager()
 {
 	IAudioInsightsModule& AudioInsightsModule = static_cast<IAudioInsightsModule&>(UE::Audio::Insights::FAudioInsightsModule::GetChecked());
 	return AudioInsightsModule.GetCacheManager();
 }
-#endif // WITH_EDITOR
 
 bool IAudioInsightsModule::IsModuleLoaded()
 {
diff --git a/Engine/Plugins/AudioInsights/Source/AudioInsights/Public/IAudioInsightsModule.h b/Engine/Plugins/AudioInsights/Source/AudioInsights/Public/IAudioInsightsModule.h
--- a/Engine/Plugins/AudioInsights/Source/AudioInsights/Public/IAudioInsightsModule.h
+++ b/Engine/Plugins/AudioInsights/Source/AudioInsights/Public/IAudioInsightsModule.h
@@ -30,9 +30,7 @@ public:
 
 	UE_API virtual IAudioInsightsTraceModule& GetTraceModule() override;
 
-#if WITH_EDITOR
 	UE_API virtual UE::Audio::Insights::FAudioInsightsCacheManager& GetCacheManager() override;
-#endif
 
 	virtual TSharedRef<SDockTab> CreateDashboardTabWidget(const FSpawnTabArgs& Args) = 0;
 
diff --git a/Engine/Plugins/AudioInsights/Source/AudioInsights/Public/IAudioInsightsModuleInterface.h b/Engine/Plugins/AudioInsights/Source/AudioInsights/Public/IAudioInsightsModuleInterface.h
--- a/Engine/Plugins/AudioInsights/Source/AudioInsights/Public/IAudioInsightsModuleInterface.h
+++ b/Engine/Plugins/AudioInsights/Source/AudioInsights/Public/IAudioInsightsModuleInterface.h
@@ -24,8 +24,5 @@ public:
 	virtual ::Audio::FDeviceId GetDeviceId() const = 0;
 
 	virtual IAudioInsightsTraceModule& GetTraceModule() = 0;
-
-#if WITH_EDITOR
 	virtual UE::Audio::Insights::FAudioInsightsCacheManager& GetCacheManager() = 0;
-#endif // WITH_EDITOR
 };
diff --git a/Engine/Plugins/AudioInsights/Source/AudioInsightsEditor/Private/AudioInsightsEditorModule.h b/Engine/Plugins/AudioInsights/Source/AudioInsightsEditor/Private/AudioInsightsEditorModule.h
--- a/Engine/Plugins/AudioInsights/Source/AudioInsightsEditor/Private/AudioInsightsEditorModule.h
+++ b/Engine/Plugins/AudioInsights/Source/AudioInsightsEditor/Private/AudioInsightsEditorModule.h
@@ -35,10 +35,7 @@ namespace UE::Audio::Insights
 		static bool IsModuleLoaded();
 		static FAudioInsightsEditorModule& GetChecked();
 		virtual IAudioInsightsTraceModule& GetTraceModule() override;
-
-#if WITH_EDITOR
 		virtual class FAudioInsightsCacheManager& GetCacheManager() override;
-#endif // WITH_EDITOR
 
 	private:
 		void RegisterMenus();
