[Audio Insights] Bughawk suggestions for cache manager

diff --git a/Engine/Plugins/AudioInsights/Source/AudioInsights/Private/Cache/AudioInsightsCacheManager.cpp b/Engine/Plugins/AudioInsights/Source/AudioInsights/Private/Cache/AudioInsightsCacheManager.cpp
--- a/Engine/Plugins/AudioInsights/Source/AudioInsights/Private/Cache/AudioInsightsCacheManager.cpp
+++ b/Engine/Plugins/AudioInsights/Source/AudioInsights/Private/Cache/AudioInsightsCacheManager.cpp
@@ -116,6 +116,10 @@ namespace UE::Audio::Insights
 			else
 			{
 				ReadIndex = Cache.GetNextIndex(ReadIndex);
+				if (ReadIndex >= Cache.Capacity())
+				{
+					break;
+				}
 			}
 		}
 
@@ -136,6 +140,11 @@ namespace UE::Audio::Insights
 			ReadIndex = Cache.GetNextIndex(ReadIndex);
 		}
 
+		if (ReadIndex >= Cache.Capacity())
+		{
+			return nullptr;
+		}
+
 		if (!Cache[ReadIndex].HasAnyData())
 		{
 			return nullptr;
