[Audio Insights] Support cache manager inside standalone part 03: Update provider base to support standalone cache

diff --git a/Engine/Plugins/AudioInsights/Source/AudioInsights/Private/AudioInsightsTraceProviderBase.cpp b/Engine/Plugins/AudioInsights/Source/AudioInsights/Private/AudioInsightsTraceProviderBase.cpp
--- a/Engine/Plugins/AudioInsights/Source/AudioInsights/Private/AudioInsightsTraceProviderBase.cpp
+++ b/Engine/Plugins/AudioInsights/Source/AudioInsights/Private/AudioInsightsTraceProviderBase.cpp
@@ -35,6 +35,7 @@ namespace UE::Audio::Insights
 
 	FTraceProviderBase::FTraceAnalyzerBase::FTraceAnalyzerBase(TSharedRef<FTraceProviderBase> InProvider)
 		: Provider(InProvider)
+		, CacheManager(IAudioInsightsModule::GetChecked().GetCacheManager())
 	{
 	}
 
@@ -49,12 +50,10 @@ namespace UE::Audio::Insights
 
 	bool FTraceProviderBase::FTraceAnalyzerBase::OnEvent(uint16 RouteId, EStyle Style, const FOnEventContext& Context)
 	{
-#if WITH_EDITOR
 		if (ShouldIgnoreNewEvents())
 		{
 			return true;
 		}
-#endif // WITH_EDITOR
 
 		switch (RouteId)
 		{
@@ -164,11 +163,7 @@ namespace UE::Audio::Insights
 
 	bool FTraceProviderBase::IsMessageProcessingPaused() const
 	{
-#if WITH_EDITOR
 		return GetMessageCacheAndProcessingStatus() != ECacheAndProcess::Latest;
-#else
-		return false;
-#endif // WITH_EDITOR
 	}
 
 	ECacheAndProcess FTraceProviderBase::GetMessageCacheAndProcessingStatus() const
diff --git a/Engine/Plugins/AudioInsights/Source/AudioInsights/Public/AudioInsightsTraceProviderBase.h b/Engine/Plugins/AudioInsights/Source/AudioInsights/Public/AudioInsightsTraceProviderBase.h
--- a/Engine/Plugins/AudioInsights/Source/AudioInsights/Public/AudioInsightsTraceProviderBase.h
+++ b/Engine/Plugins/AudioInsights/Source/AudioInsights/Public/AudioInsightsTraceProviderBase.h
@@ -3,6 +3,8 @@
 
 #include <atomic>
 #include "AudioDefines.h"
+#include "Cache/AudioInsightsCacheManager.h"
+#include "Cache/IAudioCachedMessage.h"
 #include "Containers/Ticker.h"
 #include "IAudioInsightsModule.h"
 #include "IAudioInsightsTraceModule.h"
@@ -15,8 +17,6 @@
 
 #if WITH_EDITOR
 #include "AudioDeviceManager.h"
-#include "Cache/AudioInsightsCacheManager.h"
-#include "Cache/IAudioCachedMessage.h"
 #include "ProfilingDebugging/TraceAuxiliary.h"
 #endif // WITH_EDITOR
 
@@ -119,22 +119,23 @@ namespace UE::Audio::Insights
 				return *StaticCastSharedRef<TProviderType>(Provider);
 			}
 
-#if WITH_EDITOR
 			template <TIsCacheableMessage T>
 			void CacheMessage(const FOnEventContext& Context, TAnalyzerMessageQueue<T>& OutMessageQueue)
 			{
-				FAudioInsightsCacheManager& CacheManager = IAudioInsightsModule::GetChecked().GetCacheManager();
-
 				TSharedPtr<T> Message = MakeShared<T>(Context);
-				CacheManager.AddMessageToCache(Message);
+				AddMessageToCache(MoveTemp(Message), OutMessageQueue);
+			}
 
-				if (ShouldProcessNewEvents())
+			template <TIsCacheableMessage T>
+			void CacheMessageIfLatest(const FOnEventContext& Context, TAnalyzerMessageQueue<T>& OutMessageQueue)
+			{
+				TSharedPtr<T> Message = MakeShared<T>(Context);
+				const T* LastMessage = CacheManager.FindClosestMessage<T>(Message->GetMessageName(), CacheManager.GetCacheEndTimeStamp(), Message->GetID(), Message->Timestamp);
+				if (LastMessage == nullptr || LastMessage->Timestamp < Message->Timestamp)
 				{
-					T MessageCopy = *Message;
-					OutMessageQueue.Enqueue(MoveTemp(MessageCopy));
+					AddMessageToCache(MoveTemp(Message), OutMessageQueue);
 				}
 			}
-#endif // WITH_EDITOR
 
 		private:
 			AUDIOINSIGHTS_API virtual bool OnEvent(uint16 RouteId, EStyle Style, const FOnEventContext& Context) override;
@@ -142,6 +143,18 @@ namespace UE::Audio::Insights
 			bool HandleChannelAnnounceEvent(uint16 RouteId, EStyle Style, const FOnEventContext& Context);
 			bool HandleChannelToggleEvent(uint16 RouteId, EStyle Style, const FOnEventContext& Context);
 
+			template <TIsCacheableMessage T>
+			void AddMessageToCache(TSharedPtr<T>&& Message, TAnalyzerMessageQueue<T>& OutMessageQueue)
+			{
+				CacheManager.AddMessageToCache(Message);
+
+				if (ShouldProcessNewEvents())
+				{
+					T MessageCopy = *Message;
+					OutMessageQueue.Enqueue(MoveTemp(MessageCopy));
+				}
+			}
+
 			TSharedRef<FTraceProviderBase> Provider;
 
 			enum : uint16
@@ -150,6 +163,8 @@ namespace UE::Audio::Insights
 				RouteID_ChannelToggle = TNumericLimits<uint16>::Max()
 			};
 
+			FAudioInsightsCacheManager& CacheManager;
+
 			TOptional<uint32> AudioChannelID;
 			TOptional<uint32> AudioMixerChannelID;
 
