[Audio Insights] Support cache manager inside standalone part 05: Submixes and Bus dashboards use cache manager in standalone

diff --git a/Engine/Plugins/AudioInsights/Source/AudioInsights/Private/Providers/AudioBusTraceProvider.cpp b/Engine/Plugins/AudioInsights/Source/AudioInsights/Private/Providers/AudioBusTraceProvider.cpp
--- a/Engine/Plugins/AudioInsights/Source/AudioInsights/Private/Providers/AudioBusTraceProvider.cpp
+++ b/Engine/Plugins/AudioInsights/Source/AudioInsights/Private/Providers/AudioBusTraceProvider.cpp
@@ -17,24 +17,7 @@ namespace UE::Audio::Insights
 {
 	namespace FAudioBusTraceProviderPrivate
 	{
-#if !WITH_EDITOR
-		template<typename T>
-		const T* FindClosestMessageToTimestamp(const TraceServices::TPagedArray<T>& InCachedMessages, const double InTimeMarker, const uint32 InAudioBusId)
-		{
-			const int32 ClosestMessageToTimeStampIndex = TraceServices::PagedArrayAlgo::BinarySearchClosestBy(InCachedMessages, InTimeMarker, [](const T& Msg) { return Msg.Timestamp; });
-
-			// Iterate backwards from TimeMarker until we find the matching AudioBusId
-			for (auto It = InCachedMessages.GetIteratorFromItem(ClosestMessageToTimeStampIndex); It; --It)
-			{
-				if (It->AudioBusId == InAudioBusId)
-				{
-					return &(*It);
-				}
-			}
-
-			return nullptr;
-		}
-#else
+#if WITH_EDITOR
 		TSharedPtr<FAudioBusDashboardEntry> CreateEntryFromAsset(const FString& InAssetPath)
 		{
 			TSharedPtr<FAudioBusDashboardEntry> Entry = MakeShared<FAudioBusDashboardEntry>();
@@ -51,7 +34,7 @@ namespace UE::Audio::Insights
 
 			return Entry;
 		}
-#endif // !WITH_EDITOR
+#endif // WITH_EDITOR
 	}
 
 	FAudioBusTraceProvider::FAudioBusTraceProvider()
@@ -113,51 +96,31 @@ namespace UE::Audio::Insights
 				{
 					case RouteId_Start:
 					{
-#if WITH_EDITOR
 						CacheMessage<FAudioBusStartMessage>(Context, Messages.StartMessages);
-#else
-						Messages.StartMessages.Enqueue(FAudioBusStartMessage{ Context });
-#endif // WITH_EDITOR
 						break;
 					}
 
 					case RouteId_HasActivity:
 					{
-#if WITH_EDITOR
 						CacheMessage<FAudioBusHasActivityMessage>(Context, Messages.HasActivityMessages);
-#else
-						Messages.HasActivityMessages.Enqueue(FAudioBusHasActivityMessage{ Context });
-#endif // WITH_EDITOR
 						break;
 					}
 
 					case RouteId_EnvelopeFollowerEnabled:
 					{
-#if WITH_EDITOR
 						CacheMessage<FAudioBusEnvelopeFollowerEnabledMessage>(Context, Messages.EnvelopeFollowerEnabledMessages);
-#else
-						Messages.EnvelopeFollowerEnabledMessages.Enqueue(FAudioBusEnvelopeFollowerEnabledMessage{ Context });
-#endif // WITH_EDITOR
 						break;
 					}
 
 					case RouteId_EnvelopeValues:
 					{
-#if WITH_EDITOR
 						CacheMessage<FAudioBusEnvelopeValuesMessage>(Context, Messages.EnvelopeValuesMessages);
-#else
-						Messages.EnvelopeValuesMessages.Enqueue(FAudioBusEnvelopeValuesMessage{ Context });
-#endif // WITH_EDITOR
 						break;
 					}
 
 					case RouteId_Stop:
 					{
-#if WITH_EDITOR
 						CacheMessage<FAudioBusStopMessage>(Context, Messages.StopMessages);
-#else
-						Messages.StopMessages.Enqueue(FAudioBusStopMessage{ Context });
-#endif // WITH_EDITOR
 						break;
 					}
 
@@ -235,13 +198,6 @@ namespace UE::Audio::Insights
 		ProcessMessageQueue<FAudioBusStartMessage>(TraceMessages.StartMessages, CreateEntry,
 		[this](const FAudioBusStartMessage& Msg, TSharedPtr<FAudioBusDashboardEntry>* OutEntry)
 		{
-#if !WITH_EDITOR
-			if (SessionCachedMessages.IsValid())
-			{
-				SessionCachedMessages->StartCachedMessages.EmplaceBack(Msg);
-			}
-#endif // !WITH_EDITOR
-
 			if (OutEntry)
 			{
 				FAudioBusDashboardEntry& EntryRef = *OutEntry->Get();
@@ -268,13 +224,6 @@ namespace UE::Audio::Insights
 		ProcessMessageQueue<FAudioBusHasActivityMessage>(TraceMessages.HasActivityMessages, GetEntry,
 		[this](const FAudioBusHasActivityMessage& Msg, TSharedPtr<FAudioBusDashboardEntry>* OutEntry)
 		{
-#if !WITH_EDITOR
-			if (SessionCachedMessages.IsValid())
-			{
-				SessionCachedMessages->HasActivityCachedMessages.EmplaceBack(Msg);
-			}
-#endif // !WITH_EDITOR
-
 			if (OutEntry)
 			{
 				FAudioBusDashboardEntry& EntryRef = *OutEntry->Get();
@@ -286,13 +235,6 @@ namespace UE::Audio::Insights
 		ProcessMessageQueue<FAudioBusEnvelopeFollowerEnabledMessage>(TraceMessages.EnvelopeFollowerEnabledMessages, GetEntry,
 		[this](const FAudioBusEnvelopeFollowerEnabledMessage& Msg, TSharedPtr<FAudioBusDashboardEntry>* OutEntry)
 		{
-#if !WITH_EDITOR
-			if (SessionCachedMessages.IsValid())
-			{
-				SessionCachedMessages->EnvelopeFollowerEnabledCachedMessages.EmplaceBack(Msg);
-			}
-#endif // !WITH_EDITOR
-
 			if (OutEntry)
 			{
 				FAudioBusDashboardEntry& EntryRef = *OutEntry->Get();
@@ -304,13 +246,6 @@ namespace UE::Audio::Insights
 		ProcessMessageQueue<FAudioBusEnvelopeValuesMessage>(TraceMessages.EnvelopeValuesMessages, GetEntry,
 		[this](const FAudioBusEnvelopeValuesMessage& Msg, TSharedPtr<FAudioBusDashboardEntry>* OutEntry)
 		{
-#if !WITH_EDITOR
-			if (SessionCachedMessages.IsValid())
-			{
-				SessionCachedMessages->EnvelopeValuesCachedMessages.EmplaceBack(Msg);
-			}
-#endif // !WITH_EDITOR
-
 			if (OutEntry)
 			{
 				FAudioBusDashboardEntry& EntryRef = *OutEntry->Get();
@@ -323,13 +258,6 @@ namespace UE::Audio::Insights
 		ProcessMessageQueue<FAudioBusStopMessage>(TraceMessages.StopMessages, GetEntry,
 		[this](const FAudioBusStopMessage& Msg, TSharedPtr<FAudioBusDashboardEntry>* OutEntry)
 		{
-#if !WITH_EDITOR
-			if (SessionCachedMessages.IsValid())
-			{
-				SessionCachedMessages->StopCachedMessages.EmplaceBack(Msg);
-			}
-#endif // !WITH_EDITOR
-
 			if (OutEntry && (*OutEntry)->Timestamp < Msg.Timestamp)
 			{
 #if WITH_EDITOR
@@ -471,24 +399,12 @@ namespace UE::Audio::Insights
 		OnAudioBusListUpdated.ExecuteIfBound();
 		bAssetsUpdated = true;
 	}
-#else
-	void FAudioBusTraceProvider::InitSessionCachedMessages(TraceServices::IAnalysisSession& InSession)
-	{
-		SessionCachedMessages = MakeUnique<FAudioBusSessionCachedMessages>(InSession);
-	}
 #endif // WITH_EDITOR
 
 	void FAudioBusTraceProvider::OnTimingViewTimeMarkerChanged(double TimeMarker)
 	{
 		using namespace FAudioBusTraceProviderPrivate;
 
-#if !WITH_EDITOR
-		if (!SessionCachedMessages.IsValid())
-		{
-			return;
-		}
-#endif // !WITH_EDITOR
-
 		DeviceDataMap.Empty();
 
 		// Collect all the audio bus start messages registered until this point in time 
@@ -510,23 +426,11 @@ namespace UE::Audio::Insights
 			});
 		};
 
-#if WITH_EDITOR
 		const FAudioInsightsCacheManager& CacheManager = FAudioInsightsModule::GetChecked().GetCacheManager();
 		CacheManager.IterateTo<FAudioBusStartMessage>(AudioBusMessageNames::Start, TimeMarker, [&ProcessBusStartMessage](const FAudioBusStartMessage& AudioBusStartCachedMessage)
 		{
 			ProcessBusStartMessage(AudioBusStartCachedMessage);
 		});
-#else
-		for (const FAudioBusStartMessage& AudioBusStartCachedMessage : SessionCachedMessages->StartCachedMessages)
-		{
-			if (AudioBusStartCachedMessage.Timestamp > TimeMarker)
-			{
-				break;
-			}
-
-			ProcessBusStartMessage(AudioBusStartCachedMessage);
-		}
-#endif // WITH_EDITOR
 
 		// Selectively remove start messages collected in the step above by knowing which audio buses were stopped.
 		// With this we will know what are the available audio buses at this point in time.
@@ -540,22 +444,10 @@ namespace UE::Audio::Insights
 			}
 		};
 
-#if WITH_EDITOR
 		CacheManager.IterateTo<FAudioBusStopMessage>(AudioBusMessageNames::Stop, TimeMarker, [&ProcessBusStopMessage](const FAudioBusStopMessage& AudioBusStopCachedMessage)
 		{
 			ProcessBusStopMessage(AudioBusStopCachedMessage);
 		});
-#else
-		for (const FAudioBusStopMessage& AudioBusStopCachedMessage : SessionCachedMessages->StopCachedMessages)
-		{
-			if (AudioBusStopCachedMessage.Timestamp > TimeMarker)
-			{
-				break;
-			}
-
-			ProcessBusStopMessage(AudioBusStopCachedMessage);
-		}
-#endif // WITH_EDITOR
 
 		const FDeviceData* DeviceData = FindFilteredDeviceData();
 		if (DeviceData)
@@ -574,43 +466,26 @@ namespace UE::Audio::Insights
 			TArray<CachedEntryInfo> CachedEntryInfos;
 			CachedEntryInfos.SetNum(AudioBusIdArray.Num());
 
-			ParallelFor(AudioBusIdArray.Num(),
-#if WITH_EDITOR
-			[&AudioBusIdArray, &CachedEntryInfos, &CacheManager, TimeMarker, this](const int32 Index)
-#else
-			[&AudioBusIdArray, &CachedEntryInfos, TimeMarker, this](const int32 Index)
-#endif // WITH_EDITOR
+			ParallelFor(AudioBusIdArray.Num(), [&AudioBusIdArray, &CachedEntryInfos, &CacheManager, TimeMarker, this](const int32 Index)
 			{
 				const uint32 AudioBusId = AudioBusIdArray[Index];
 
 				// HasActivity
-#if WITH_EDITOR
 				const FAudioBusHasActivityMessage* FoundHasActivityCachedMessage = CacheManager.FindClosestMessage<FAudioBusHasActivityMessage>(AudioBusMessageNames::HasActivity, TimeMarker, AudioBusId);
-#else
-				const FAudioBusHasActivityMessage* FoundHasActivityCachedMessage = FindClosestMessageToTimestamp(SessionCachedMessages->HasActivityCachedMessages, TimeMarker, AudioBusId);
-#endif // WITH_EDITOR
 				if (FoundHasActivityCachedMessage)
 				{
 					CachedEntryInfos[Index].HasActivityMessage = *FoundHasActivityCachedMessage;
 				}
 
 				// EnvelopeFollowerEnabled
-#if WITH_EDITOR
 				const FAudioBusEnvelopeFollowerEnabledMessage* FoundEnvelopeFollowerEnabledCachedMessage = CacheManager.FindClosestMessage<FAudioBusEnvelopeFollowerEnabledMessage>(AudioBusMessageNames::EnvelopeFollowerEnabled, TimeMarker, AudioBusId);
-#else
-				const FAudioBusEnvelopeFollowerEnabledMessage* FoundEnvelopeFollowerEnabledCachedMessage = FindClosestMessageToTimestamp(SessionCachedMessages->EnvelopeFollowerEnabledCachedMessages, TimeMarker, AudioBusId);
-#endif // WITH_EDITOR
 				if (FoundEnvelopeFollowerEnabledCachedMessage)
 				{
 					CachedEntryInfos[Index].EnvelopeFollowerEnabledMessage = *FoundEnvelopeFollowerEnabledCachedMessage;
 				}
 
 				// Envelope Values
-#if WITH_EDITOR
 				const FAudioBusEnvelopeValuesMessage* FoundEnvelopeValuesCachedMessage = CacheManager.FindClosestMessage<FAudioBusEnvelopeValuesMessage>(AudioBusMessageNames::EnvelopeValues, TimeMarker, AudioBusId);
-#else
-				const FAudioBusEnvelopeValuesMessage* FoundEnvelopeValuesCachedMessage = FindClosestMessageToTimestamp(SessionCachedMessages->EnvelopeValuesCachedMessages, TimeMarker, AudioBusId);
-#endif // WITH_EDITOR
 				if (FoundEnvelopeValuesCachedMessage)
 				{
 					CachedEntryInfos[Index].EnvelopeValuesMessage = *FoundEnvelopeValuesCachedMessage;
diff --git a/Engine/Plugins/AudioInsights/Source/AudioInsights/Private/Providers/AudioBusTraceProvider.h b/Engine/Plugins/AudioInsights/Source/AudioInsights/Private/Providers/AudioBusTraceProvider.h
--- a/Engine/Plugins/AudioInsights/Source/AudioInsights/Private/Providers/AudioBusTraceProvider.h
+++ b/Engine/Plugins/AudioInsights/Source/AudioInsights/Private/Providers/AudioBusTraceProvider.h
@@ -23,9 +23,6 @@ namespace UE::Audio::Insights
 		
 		virtual UE::Trace::IAnalyzer* ConstructAnalyzer(TraceServices::IAnalysisSession& InSession) override;
 
-#if !WITH_EDITOR
-		virtual void InitSessionCachedMessages(TraceServices::IAnalysisSession& InSession) override;
-#endif // !WITH_EDITOR
 		void RequestEntriesUpdate();
 
 		DECLARE_DELEGATE_OneParam(FOnAudioBusChanged, const uint32 /*AudioBusId*/);
@@ -51,8 +48,6 @@ namespace UE::Audio::Insights
 		void HandleOnAudioBusAssetAdded(const FString& InAssetPath);
 		void HandleOnAudioBusAssetRemoved(const FString& InAssetPath);
 		void HandleOnAudioBusAssetListUpdated(const TArray<FString>& InAssetPaths);
-#else
-		TUniquePtr<FAudioBusSessionCachedMessages> SessionCachedMessages;
 #endif // WITH_EDITOR
 
 		virtual void OnTimingViewTimeMarkerChanged(double TimeMarker) override;
diff --git a/Engine/Plugins/AudioInsights/Source/AudioInsights/Private/Providers/SubmixTraceProvider.cpp b/Engine/Plugins/AudioInsights/Source/AudioInsights/Private/Providers/SubmixTraceProvider.cpp
--- a/Engine/Plugins/AudioInsights/Source/AudioInsights/Private/Providers/SubmixTraceProvider.cpp
+++ b/Engine/Plugins/AudioInsights/Source/AudioInsights/Private/Providers/SubmixTraceProvider.cpp
@@ -17,24 +17,7 @@ namespace UE::Audio::Insights
 {
 	namespace FSubmixTraceProviderPrivate
 	{
-#if !WITH_EDITOR
-		template<typename T>
-		const T* FindClosestMessageToTimestamp(const TraceServices::TPagedArray<T>& InCachedMessages, const double InTimeMarker, const uint32 InSubmixId)
-		{
-			const int32 ClosestMessageToTimeStampIndex = TraceServices::PagedArrayAlgo::BinarySearchClosestBy(InCachedMessages, InTimeMarker, [](const T& Msg) { return Msg.Timestamp; });
-
-			// Iterate backwards from TimeMarker until we find the matching SubmixId
-			for (auto It = InCachedMessages.GetIteratorFromItem(ClosestMessageToTimeStampIndex); It; --It)
-			{
-				if (It->SubmixId == InSubmixId)
-				{
-					return &(*It);
-				}
-			}
-
-			return nullptr;
-		}
-#else
+#if WITH_EDITOR
 		TSharedPtr<FSubmixDashboardEntry> CreateEntryFromAsset(const FString& InAssetPath)
 		{
 			TSharedPtr<FSubmixDashboardEntry> Entry = MakeShared<FSubmixDashboardEntry>();
@@ -57,7 +40,7 @@ namespace UE::Audio::Insights
 
 			return Entry;
 		}
-#endif // !WITH_EDITOR
+#endif // WITH_EDITOR
 	}
 
 	FSubmixTraceProvider::FSubmixTraceProvider()
@@ -120,61 +103,37 @@ namespace UE::Audio::Insights
 				{
 					case RouteId_Loaded:
 					{
-#if WITH_EDITOR
 						CacheMessage<FSubmixLoadedMessage>(Context, Messages.LoadedMessages);
-#else
-						Messages.LoadedMessages.Enqueue(FSubmixLoadedMessage{ Context });
-#endif // WITH_EDITOR
 						break;
 					}
 
 					case RouteId_AlivePing:
 					{
-#if WITH_EDITOR
 						CacheMessage<FSubmixAlivePingMessage>(Context, Messages.AlivePingMessages);
-#else
-						Messages.AlivePingMessages.Enqueue(FSubmixAlivePingMessage{ Context });
-#endif // WITH_EDITOR
 						break;
 					}
 
 					case RouteId_HasActivity:
 					{
-#if WITH_EDITOR
 						CacheMessage<FSubmixHasActivityMessage>(Context, Messages.HasActivityMessages);
-#else
-						Messages.HasActivityMessages.Enqueue(FSubmixHasActivityMessage{ Context });
-#endif // WITH_EDITOR
 						break;
 					}
 
 					case RouteId_EnvelopeFollowerEnabled:
 					{
-#if WITH_EDITOR
 						CacheMessage<FSubmixEnvelopeFollowerEnabledMessage>(Context, Messages.EnvelopeFollowerEnabledMessages);
-#else
-						Messages.EnvelopeFollowerEnabledMessages.Enqueue(FSubmixEnvelopeFollowerEnabledMessage{ Context });
-#endif // WITH_EDITOR
 						break;
 					}
 
 					case RouteId_EnvelopeValues:
 					{
-#if WITH_EDITOR
 						CacheMessage<FSubmixEnvelopeValuesMessage>(Context, Messages.EnvelopeValuesMessages);
-#else
-						Messages.EnvelopeValuesMessages.Enqueue(FSubmixEnvelopeValuesMessage{ Context });
-#endif // WITH_EDITOR
 						break;
 					}
 
 					case RouteId_Unloaded:
 					{
-#if WITH_EDITOR
 						CacheMessage<FSubmixUnloadedMessage>(Context, Messages.UnloadedMessages);
-#else
-						Messages.UnloadedMessages.Enqueue(FSubmixUnloadedMessage{ Context });
-#endif // WITH_EDITOR
 						break;
 					}
 
@@ -254,13 +213,6 @@ namespace UE::Audio::Insights
 		ProcessMessageQueue<FSubmixLoadedMessage>(TraceMessages.LoadedMessages, CreateEntry,
 		[this](const FSubmixLoadedMessage& Msg, TSharedPtr<FSubmixDashboardEntry>* OutEntry)
 		{
-#if !WITH_EDITOR
-			if (SessionCachedMessages.IsValid())
-			{
-				SessionCachedMessages->LoadedCachedMessages.EmplaceBack(Msg);
-			}
-#endif // !WITH_EDITOR
-
 			if (OutEntry)
 			{
 				FSubmixDashboardEntry& EntryRef = *OutEntry->Get();
@@ -275,13 +227,6 @@ namespace UE::Audio::Insights
 		ProcessMessageQueue<FSubmixAlivePingMessage>(TraceMessages.AlivePingMessages, CreateEntry,
 		[this](const FSubmixAlivePingMessage& Msg, TSharedPtr<FSubmixDashboardEntry>* OutEntry)
 		{
-#if !WITH_EDITOR
-			if (SessionCachedMessages.IsValid())
-			{
-				SessionCachedMessages->AlivePingCachedMessages.EmplaceBack(Msg);
-			}
-#endif // !WITH_EDITOR
-
 			if (OutEntry)
 			{
 				FSubmixDashboardEntry& EntryRef = *OutEntry->Get();
@@ -294,13 +239,6 @@ namespace UE::Audio::Insights
 		ProcessMessageQueue<FSubmixHasActivityMessage>(TraceMessages.HasActivityMessages, GetEntry,
 		[this](const FSubmixHasActivityMessage& Msg, TSharedPtr<FSubmixDashboardEntry>* OutEntry)
 		{
-#if !WITH_EDITOR
-			if (SessionCachedMessages.IsValid())
-			{
-				SessionCachedMessages->HasActivityCachedMessages.EmplaceBack(Msg);
-			}
-#endif // !WITH_EDITOR
-
 			if (OutEntry)
 			{
 				FSubmixDashboardEntry& EntryRef = *OutEntry->Get();
@@ -312,13 +250,6 @@ namespace UE::Audio::Insights
 		ProcessMessageQueue<FSubmixEnvelopeFollowerEnabledMessage>(TraceMessages.EnvelopeFollowerEnabledMessages, GetEntry,
 		[this](const FSubmixEnvelopeFollowerEnabledMessage& Msg, TSharedPtr<FSubmixDashboardEntry>* OutEntry)
 		{
-#if !WITH_EDITOR
-			if (SessionCachedMessages.IsValid())
-			{
-				SessionCachedMessages->EnvelopeFollowerEnabledCachedMessages.EmplaceBack(Msg);
-			}
-#endif // !WITH_EDITOR
-
 			if (OutEntry)
 			{
 				FSubmixDashboardEntry& EntryRef = *OutEntry->Get();
@@ -330,13 +261,6 @@ namespace UE::Audio::Insights
 		ProcessMessageQueue<FSubmixEnvelopeValuesMessage>(TraceMessages.EnvelopeValuesMessages, GetEntry,
 		[this](const FSubmixEnvelopeValuesMessage& Msg, TSharedPtr<FSubmixDashboardEntry>* OutEntry)
 		{
-#if !WITH_EDITOR
-			if (SessionCachedMessages.IsValid())
-			{
-				SessionCachedMessages->EnvelopeValuesCachedMessages.EmplaceBack(Msg);
-			}
-#endif // !WITH_EDITOR
-				
 			if (OutEntry)
 			{
 				FSubmixDashboardEntry& EntryRef = *OutEntry->Get();
@@ -349,13 +273,6 @@ namespace UE::Audio::Insights
 		ProcessMessageQueue<FSubmixUnloadedMessage>(TraceMessages.UnloadedMessages, GetEntry,
 		[this](const FSubmixUnloadedMessage& Msg, TSharedPtr<FSubmixDashboardEntry>* OutEntry)
 		{
-#if !WITH_EDITOR
-			if (SessionCachedMessages.IsValid())
-			{
-				SessionCachedMessages->UnloadedCachedMessages.EmplaceBack(Msg);
-			}
-#endif // !WITH_EDITOR
-
 			if (OutEntry && (*OutEntry)->Timestamp < Msg.Timestamp)
 			{
 #if WITH_EDITOR
@@ -491,24 +408,12 @@ namespace UE::Audio::Insights
 		OnSubmixListUpdated.ExecuteIfBound();
 		bAssetsUpdated = true;
 	}
-#else
-	void FSubmixTraceProvider::InitSessionCachedMessages(TraceServices::IAnalysisSession& InSession)
-	{
-		SessionCachedMessages = MakeUnique<FSubmixSessionCachedMessages>(InSession);
-	}
 #endif // WITH_EDITOR
 
 	void FSubmixTraceProvider::OnTimingViewTimeMarkerChanged(double TimeMarker)
 	{
 		using namespace FSubmixTraceProviderPrivate;
 
-#if !WITH_EDITOR
-		if (!SessionCachedMessages.IsValid())
-		{
-			return;
-		}
-#endif // !WITH_EDITOR
-
 		DeviceDataMap.Empty();
 
 		// Collect all the submix loaded messages registered until this point in time 
@@ -530,24 +435,12 @@ namespace UE::Audio::Insights
 			});
 		};
 
-#if WITH_EDITOR
 		const FAudioInsightsCacheManager& CacheManager = FAudioInsightsModule::GetChecked().GetCacheManager();
 		CacheManager.IterateTo<FSubmixLoadedMessage>(SubmixMessageNames::Loaded, TimeMarker, [&ProcessSubmixLoadedMessage](const FSubmixLoadedMessage& SubmixLoadedCachedMessage)
 		{
 			ProcessSubmixLoadedMessage(SubmixLoadedCachedMessage);
 		});
 
-#else
-		for (const FSubmixLoadedMessage& SubmixLoadedCachedMessage : SessionCachedMessages->LoadedCachedMessages)
-		{
-			if (SubmixLoadedCachedMessage.Timestamp > TimeMarker)
-			{
-				break;
-			}
-
-			ProcessSubmixLoadedMessage(SubmixLoadedCachedMessage);
-		}
-#endif // WITH_EDITOR
 
 		// Collect all the submix alive ping messages registered until this point in time 
 		auto ProcessSubmixAlivePingMessage = [this](const FSubmixAlivePingMessage& Msg)
@@ -568,22 +461,10 @@ namespace UE::Audio::Insights
 			});
 		};
 
-#if WITH_EDITOR
 		CacheManager.IterateTo<FSubmixAlivePingMessage>(SubmixMessageNames::IsAlivePing, TimeMarker, [&ProcessSubmixAlivePingMessage](const FSubmixAlivePingMessage& SubmixAlivePingCachedMessage)
 		{
 			ProcessSubmixAlivePingMessage(SubmixAlivePingCachedMessage);
 		});
-#else
-		for (const FSubmixAlivePingMessage& SubmixAlivePingCachedMessage : SessionCachedMessages->AlivePingCachedMessages)
-		{
-			if (SubmixAlivePingCachedMessage.Timestamp > TimeMarker)
-			{
-				break;
-			}
-
-			ProcessSubmixAlivePingMessage(SubmixAlivePingCachedMessage);
-		}
-#endif // WITH_EDITOR
 
 		// Selectively remove loaded messages collected in the step above by knowing which submixes were unloaded.
 		// With this we will know what are the available submixes at this point in time.
@@ -597,22 +478,10 @@ namespace UE::Audio::Insights
 			}
 		};
 
-#if WITH_EDITOR
 		CacheManager.IterateTo<FSubmixUnloadedMessage>(SubmixMessageNames::Unloaded, TimeMarker, [&ProcessSubmixUnloadedMessage](const FSubmixUnloadedMessage& SubmixUnloadedCachedMessage)
 		{
 			ProcessSubmixUnloadedMessage(SubmixUnloadedCachedMessage);
 		});
-#else
-		for (const FSubmixUnloadedMessage& SubmixUnloadedCachedMessage : SessionCachedMessages->UnloadedCachedMessages)
-		{
-			if (SubmixUnloadedCachedMessage.Timestamp > TimeMarker)
-			{
-				break;
-			}
-
-			ProcessSubmixUnloadedMessage(SubmixUnloadedCachedMessage);
-		}
-#endif // WITH_EDITOR
 
 		const FDeviceData* DeviceData = FindFilteredDeviceData();
 		if (DeviceData)
@@ -631,43 +500,26 @@ namespace UE::Audio::Insights
 			TArray<CachedEntryInfo> CachedEntryInfos;
 			CachedEntryInfos.SetNum(SubmixIdArray.Num());
 
-			ParallelFor(SubmixIdArray.Num(),
-#if WITH_EDITOR
-			[&SubmixIdArray, &CachedEntryInfos, &CacheManager, TimeMarker, this](const int32 Index)
-#else
-			[&SubmixIdArray, &CachedEntryInfos, TimeMarker, this](const int32 Index)
-#endif // WITH_EDITOR
+			ParallelFor(SubmixIdArray.Num(), [&SubmixIdArray, &CachedEntryInfos, &CacheManager, TimeMarker, this](const int32 Index)
 			{
 				const uint32 SubmixId = SubmixIdArray[Index];
 
 				// HasActivity
-#if WITH_EDITOR
 				const FSubmixHasActivityMessage* FoundHasActivityCachedMessage = CacheManager.FindClosestMessage<FSubmixHasActivityMessage>(SubmixMessageNames::HasActivity, TimeMarker, SubmixId);
-#else
-				const FSubmixHasActivityMessage* FoundHasActivityCachedMessage = FindClosestMessageToTimestamp(SessionCachedMessages->HasActivityCachedMessages, TimeMarker, SubmixId);
-#endif // WITH_EDITOR
 				if (FoundHasActivityCachedMessage)
 				{
 					CachedEntryInfos[Index].HasActivityMessage = *FoundHasActivityCachedMessage;
 				}
 
 				// EnvelopeFollowerEnabled
-#if WITH_EDITOR
 				const FSubmixEnvelopeFollowerEnabledMessage* FoundEnvelopeFollowerEnabledCachedMessage = CacheManager.FindClosestMessage<FSubmixEnvelopeFollowerEnabledMessage>(SubmixMessageNames::EnvelopeFollowerEnabled, TimeMarker, SubmixId);
-#else
-				const FSubmixEnvelopeFollowerEnabledMessage* FoundEnvelopeFollowerEnabledCachedMessage = FindClosestMessageToTimestamp(SessionCachedMessages->EnvelopeFollowerEnabledCachedMessages, TimeMarker, SubmixId);
-#endif // WITH_EDITOR
 				if (FoundEnvelopeFollowerEnabledCachedMessage)
 				{
 					CachedEntryInfos[Index].EnvelopeFollowerEnabledMessage = *FoundEnvelopeFollowerEnabledCachedMessage;
 				}
 
 				// EnvelopeValues
-#if WITH_EDITOR
 				const FSubmixEnvelopeValuesMessage* FoundEnvelopeValuesCachedMessage = CacheManager.FindClosestMessage<FSubmixEnvelopeValuesMessage>(SubmixMessageNames::EnvelopeValues, TimeMarker, SubmixId);
-#else
-				const FSubmixEnvelopeValuesMessage* FoundEnvelopeValuesCachedMessage = FindClosestMessageToTimestamp(SessionCachedMessages->EnvelopeValuesCachedMessages, TimeMarker, SubmixId);
-#endif // WITH_EDITOR
 				if (FoundEnvelopeValuesCachedMessage)
 				{
 					CachedEntryInfos[Index].EnvelopeValuesMessage = *FoundEnvelopeValuesCachedMessage;
diff --git a/Engine/Plugins/AudioInsights/Source/AudioInsights/Private/Providers/SubmixTraceProvider.h b/Engine/Plugins/AudioInsights/Source/AudioInsights/Private/Providers/SubmixTraceProvider.h
--- a/Engine/Plugins/AudioInsights/Source/AudioInsights/Private/Providers/SubmixTraceProvider.h
+++ b/Engine/Plugins/AudioInsights/Source/AudioInsights/Private/Providers/SubmixTraceProvider.h
@@ -22,10 +22,6 @@ namespace UE::Audio::Insights
 		static FName GetName_Static();
 		
 		virtual UE::Trace::IAnalyzer* ConstructAnalyzer(TraceServices::IAnalysisSession& InSession) override;
-
-#if !WITH_EDITOR
-		virtual void InitSessionCachedMessages(TraceServices::IAnalysisSession& InSession) override;
-#endif // !WITH_EDITOR
 		void RequestEntriesUpdate();
 
 		DECLARE_DELEGATE_OneParam(FOnSubmixChanged, const uint32 /*SubmixId*/);
@@ -51,8 +47,6 @@ namespace UE::Audio::Insights
 		void HandleOnSubmixAssetAdded(const FString& InAssetPath);
 		void HandleOnSubmixAssetRemoved(const FString& InAssetPath);
 		void HandleOnSubmixAssetListUpdated(const TArray<FString>& InAssetPaths);
-#else
-		TUniquePtr<FSubmixSessionCachedMessages> SessionCachedMessages;
 #endif // WITH_EDITOR
 
 		virtual void OnTimingViewTimeMarkerChanged(double TimeMarker) override;
diff --git a/Engine/Plugins/AudioInsights/Source/AudioInsights/Public/Messages/AudioBusTraceMessages.h b/Engine/Plugins/AudioInsights/Source/AudioInsights/Public/Messages/AudioBusTraceMessages.h
--- a/Engine/Plugins/AudioInsights/Source/AudioInsights/Public/Messages/AudioBusTraceMessages.h
+++ b/Engine/Plugins/AudioInsights/Source/AudioInsights/Public/Messages/AudioBusTraceMessages.h
@@ -107,27 +107,6 @@ namespace UE::Audio::Insights
 		friend class FAudioBusTraceProvider;
 	};
 
-#if !WITH_EDITOR
-	struct FAudioBusSessionCachedMessages
-	{
-		FAudioBusSessionCachedMessages(TraceServices::IAnalysisSession& InSession)
-			: StartCachedMessages(InSession.GetLinearAllocator(), 4096)
-			, HasActivityCachedMessages(InSession.GetLinearAllocator(), 4096)
-			, EnvelopeFollowerEnabledCachedMessages(InSession.GetLinearAllocator(), 4096)
-			, EnvelopeValuesCachedMessages(InSession.GetLinearAllocator(), 4096)
-			, StopCachedMessages(InSession.GetLinearAllocator(), 4096)
-		{
-
-		}
-
-		TraceServices::TPagedArray<FAudioBusStartMessage> StartCachedMessages;
-		TraceServices::TPagedArray<FAudioBusHasActivityMessage> HasActivityCachedMessages;
-		TraceServices::TPagedArray<FAudioBusEnvelopeFollowerEnabledMessage> EnvelopeFollowerEnabledCachedMessages;
-		TraceServices::TPagedArray<FAudioBusEnvelopeValuesMessage> EnvelopeValuesCachedMessages;
-		TraceServices::TPagedArray<FAudioBusStopMessage> StopCachedMessages;
-	};
-#endif // !WITH_EDITOR
-
 	struct FAudioBusDashboardEntry : public IObjectDashboardEntry
 	{
 		FAudioBusDashboardEntry();
diff --git a/Engine/Plugins/AudioInsights/Source/AudioInsights/Public/Messages/SubmixTraceMessages.h b/Engine/Plugins/AudioInsights/Source/AudioInsights/Public/Messages/SubmixTraceMessages.h
--- a/Engine/Plugins/AudioInsights/Source/AudioInsights/Public/Messages/SubmixTraceMessages.h
+++ b/Engine/Plugins/AudioInsights/Source/AudioInsights/Public/Messages/SubmixTraceMessages.h
@@ -117,29 +117,6 @@ namespace UE::Audio::Insights
 		friend class FSubmixTraceProvider;
 	};
 
-#if !WITH_EDITOR
-	struct FSubmixSessionCachedMessages
-	{
-		FSubmixSessionCachedMessages(TraceServices::IAnalysisSession& InSession)
-			: LoadedCachedMessages(InSession.GetLinearAllocator(), 4096)
-			, AlivePingCachedMessages(InSession.GetLinearAllocator(), 4096)
-			, HasActivityCachedMessages(InSession.GetLinearAllocator(), 4096)
-			, EnvelopeFollowerEnabledCachedMessages(InSession.GetLinearAllocator(), 4096)
-			, EnvelopeValuesCachedMessages(InSession.GetLinearAllocator(), 4096)
-			, UnloadedCachedMessages(InSession.GetLinearAllocator(), 4096)
-		{
-
-		}
-
-		TraceServices::TPagedArray<FSubmixLoadedMessage> LoadedCachedMessages;
-		TraceServices::TPagedArray<FSubmixAlivePingMessage> AlivePingCachedMessages;
-		TraceServices::TPagedArray<FSubmixHasActivityMessage> HasActivityCachedMessages;
-		TraceServices::TPagedArray<FSubmixEnvelopeFollowerEnabledMessage> EnvelopeFollowerEnabledCachedMessages;
-		TraceServices::TPagedArray<FSubmixEnvelopeValuesMessage> EnvelopeValuesCachedMessages;
-		TraceServices::TPagedArray<FSubmixUnloadedMessage> UnloadedCachedMessages;
-	};
-#endif // !WITH_EDITOR
-
 	struct FSubmixDashboardEntry : public IObjectDashboardEntry
 	{
 		FSubmixDashboardEntry();
