[Audio Insights] Support cache manager inside standalone part 04: Virtual Loops uses cache manager in standalone

diff --git a/Engine/Plugins/AudioInsights/Source/AudioInsights/Private/Providers/VirtualLoopTraceProvider.cpp b/Engine/Plugins/AudioInsights/Source/AudioInsights/Private/Providers/VirtualLoopTraceProvider.cpp
--- a/Engine/Plugins/AudioInsights/Source/AudioInsights/Private/Providers/VirtualLoopTraceProvider.cpp
+++ b/Engine/Plugins/AudioInsights/Source/AudioInsights/Private/Providers/VirtualLoopTraceProvider.cpp
@@ -3,39 +3,16 @@
 
 #include "Async/ParallelFor.h"
 #include "AudioInsightsModule.h"
+#include "Cache/AudioInsightsCacheManager.h"
 #include "Trace/Analyzer.h"
 #include "TraceServices/Model/AnalysisSession.h"
 
-#if WITH_EDITOR
-#include "Cache/AudioInsightsCacheManager.h"
-#else
+#if !WITH_EDITOR
 #include "Common/PagedArray.h"
-#endif // WITH_EDITOR
+#endif // !WITH_EDITOR
 
 namespace UE::Audio::Insights
 {
-	namespace FVirtualLoopTraceProviderPrivate
-	{
-#if !WITH_EDITOR
-		template<typename T>
-		const T* FindClosestMessageToTimestamp(const TraceServices::TPagedArray<T>& InCachedMessages, const double InTimeMarker, const uint32 InPlayOrder)
-		{
-			const int32 ClosestMessageToTimeStampIndex = TraceServices::PagedArrayAlgo::BinarySearchClosestBy(InCachedMessages, InTimeMarker, [](const T& Msg) { return Msg.Timestamp; });
-
-			// Iterate backwards from TimeMarker until we find the matching PlayOrder
-			for (auto It = InCachedMessages.GetIteratorFromItem(ClosestMessageToTimeStampIndex); It; --It)
-			{
-				if (It->PlayOrder == InPlayOrder)
-				{
-					return &(*It);
-				}
-			}
-
-			return nullptr;
-		}
-#endif // !WITH_EDITOR
-	}
-
 	FName FVirtualLoopTraceProvider::GetName_Static()
 	{
 		return "AudioVirtualLoopProvider";
@@ -43,13 +20,6 @@ namespace UE::Audio::Insights
 
 	void FVirtualLoopTraceProvider::OnTimingViewTimeMarkerChanged(double TimeMarker)
 	{
-#if !WITH_EDITOR
-		if (!SessionCachedMessages.IsValid())
-		{
-			return;
-		}
-#endif // !WITH_EDITOR
-
 		// Helper lambdas
 		auto ProcessVirtualizeMessage = [this](const FVirtualLoopVirtualizeMessage& VirtualizeCachedMessage)
 		{
@@ -82,7 +52,6 @@ namespace UE::Audio::Insights
 		// Process cached messages
 		DeviceDataMap.Empty();
 
-#if WITH_EDITOR
 		const FAudioInsightsCacheManager& CacheManager = FAudioInsightsModule::GetChecked().GetCacheManager();
 
 		// Virtualize cached messages
@@ -96,29 +65,6 @@ namespace UE::Audio::Insights
 		{
 			ProcessStopOrRealizeMessage(StopOrRealizeCachedMessage);
 		});
-#else
-		// Virtualize cached messages
-		for (const FVirtualLoopVirtualizeMessage& VirtualizeCachedMessage : SessionCachedMessages->VirtualizeCachedMessages)
-		{
-			if (VirtualizeCachedMessage.Timestamp > TimeMarker)
-			{
-				break;
-			}
-
-			ProcessVirtualizeMessage(VirtualizeCachedMessage);
-		}
-
-		// Stop or realize cached messages
-		for (const FVirtualLoopStopOrRealizeMessage& StopOrRealizeCachedMessage : SessionCachedMessages->StopOrRealizeCachedMessages)
-		{
-			if (StopOrRealizeCachedMessage.Timestamp > TimeMarker)
-			{
-				break;
-			}
-
-			ProcessStopOrRealizeMessage(StopOrRealizeCachedMessage);
-		}
-#endif // WITH_EDITOR
 
 		CollectParamsForTimestamp(TimeMarker);
 
@@ -128,8 +74,6 @@ namespace UE::Audio::Insights
 
 	void FVirtualLoopTraceProvider::CollectParamsForTimestamp(const double InTimeMarker)
 	{
-		using namespace FVirtualLoopTraceProviderPrivate;
-
 		const FDeviceData* DeviceData = FindFilteredDeviceData();
 		if (DeviceData == nullptr)
 		{
@@ -148,22 +92,13 @@ namespace UE::Audio::Insights
 		TArray<CachedEntryInfo> CachedEntryInfos;
 		CachedEntryInfos.SetNum(PlayOrderArray.Num());
 
-#if WITH_EDITOR
 		const FAudioInsightsCacheManager& CacheManager = FAudioInsightsModule::GetChecked().GetCacheManager();
 
 		ParallelFor(PlayOrderArray.Num(), [&PlayOrderArray, &CachedEntryInfos, &CacheManager, InTimeMarker, this](const int32 Index)
-#else
-		ParallelFor(PlayOrderArray.Num(), [&PlayOrderArray, &CachedEntryInfos, InTimeMarker, this](const int32 Index)
-#endif // WITH_EDITOR
 		{
 			const uint32 PlayOrder = PlayOrderArray[Index];
 
-#if WITH_EDITOR
 			const FVirtualLoopUpdateMessage* FoundUpdateCachedMessage = CacheManager.FindClosestMessage<FVirtualLoopUpdateMessage>(VirtualLoopMessageNames::Update, InTimeMarker, PlayOrder);
-#else
-			const FVirtualLoopUpdateMessage* FoundUpdateCachedMessage = FindClosestMessageToTimestamp(SessionCachedMessages->UpdateCachedMessages, InTimeMarker, PlayOrder);
-#endif // WITH_EDITOR
-
 			if (FoundUpdateCachedMessage)
 			{
 				CachedEntryInfos[Index].UpdateMessage = *FoundUpdateCachedMessage;
@@ -198,13 +133,6 @@ namespace UE::Audio::Insights
 	{
 		auto RemoveEntryFunc = [this](const FVirtualLoopStopOrRealizeMessage& Msg, TSharedPtr<FVirtualLoopDashboardEntry>* OutEntry)
 		{
-#if !WITH_EDITOR
-			if (SessionCachedMessages.IsValid())
-			{
-				SessionCachedMessages->StopOrRealizeCachedMessages.EmplaceBack(Msg);
-			}
-#endif // !WITH_EDITOR
-
 			if (OutEntry && (*OutEntry)->Timestamp < Msg.Timestamp)
 			{
 				RemoveDeviceEntry(Msg.DeviceId, Msg.PlayOrder);
@@ -239,13 +167,6 @@ namespace UE::Audio::Insights
 		BumpEntryFunc,
 		[this](const FVirtualLoopVirtualizeMessage& Msg, TSharedPtr<FVirtualLoopDashboardEntry>* OutEntry)
 		{
-#if !WITH_EDITOR
-			if (SessionCachedMessages.IsValid())
-			{
-				SessionCachedMessages->VirtualizeCachedMessages.EmplaceBack(Msg);
-			}
-#endif // !WITH_EDITOR
-
 			FVirtualLoopDashboardEntry& EntryRef = *OutEntry->Get();
 			EntryRef.Name = Msg.Name;
 			EntryRef.ComponentId = Msg.ComponentId;
@@ -255,13 +176,6 @@ namespace UE::Audio::Insights
 		GetEntryFunc,
 		[this](const FVirtualLoopUpdateMessage& Msg, TSharedPtr<FVirtualLoopDashboardEntry>* OutEntry)
 		{
-#if !WITH_EDITOR
-			if (SessionCachedMessages.IsValid())
-			{
-				SessionCachedMessages->UpdateCachedMessages.EmplaceBack(Msg);
-			}
-#endif // !WITH_EDITOR
-
 			if (OutEntry)
 			{
 				FVirtualLoopDashboardEntry& EntryRef = *OutEntry->Get();
@@ -278,13 +192,6 @@ namespace UE::Audio::Insights
 		return true;
 	}
 
-#if !WITH_EDITOR
-	void FVirtualLoopTraceProvider::InitSessionCachedMessages(TraceServices::IAnalysisSession& InSession)
-	{
-		SessionCachedMessages = MakeUnique<FVirtualLoopSessionCachedMessages>(InSession);
-	}
-#endif // !WITH_EDITOR
-
 	UE::Trace::IAnalyzer* FVirtualLoopTraceProvider::ConstructAnalyzer(TraceServices::IAnalysisSession& InSession)
 	{
 		class FVirtualLoopTraceAnalyzer : public FTraceAnalyzerBase
@@ -317,31 +224,19 @@ namespace UE::Audio::Insights
 					case RouteId_Virtualize:
 					case RouteId_IsStillVirtualizedPing:
 					{
-#if WITH_EDITOR
 						CacheMessage<FVirtualLoopVirtualizeMessage>(Context, Messages.VirtualizeMessages);
-#else
-						Messages.VirtualizeMessages.Enqueue(FVirtualLoopVirtualizeMessage{ Context });
-#endif // WITH_EDITOR
 						break;
 					}
 
 					case RouteId_Stop:
 					{
-#if WITH_EDITOR
 						CacheMessage<FVirtualLoopStopOrRealizeMessage>(Context, Messages.StopOrRealizeMessages);
-#else
-						Messages.StopOrRealizeMessages.Enqueue(FVirtualLoopStopOrRealizeMessage{ Context });
-#endif // WITH_EDITOR
 						break;
 					}
 
 					case RouteId_Update:
 					{
-#if WITH_EDITOR
 						CacheMessage<FVirtualLoopUpdateMessage>(Context, Messages.UpdateMessages);
-#else
-						Messages.UpdateMessages.Enqueue(FVirtualLoopUpdateMessage{ Context });
-#endif // WITH_EDITOR
 						break;
 					}
 
diff --git a/Engine/Plugins/AudioInsights/Source/AudioInsights/Private/Providers/VirtualLoopTraceProvider.h b/Engine/Plugins/AudioInsights/Source/AudioInsights/Private/Providers/VirtualLoopTraceProvider.h
--- a/Engine/Plugins/AudioInsights/Source/AudioInsights/Private/Providers/VirtualLoopTraceProvider.h
+++ b/Engine/Plugins/AudioInsights/Source/AudioInsights/Private/Providers/VirtualLoopTraceProvider.h
@@ -23,18 +23,10 @@ namespace UE::Audio::Insights
 
 		static FName GetName_Static();
 
-#if !WITH_EDITOR
-		virtual void InitSessionCachedMessages(TraceServices::IAnalysisSession& InSession) override;
-#endif // !WITH_EDITOR
-
 	private:
 		virtual void OnTimingViewTimeMarkerChanged(double TimeMarker) override;
 		void CollectParamsForTimestamp(const double InTimeMarker);
 
-#if !WITH_EDITOR
-		TUniquePtr<FVirtualLoopSessionCachedMessages> SessionCachedMessages;
-#endif // !WITH_EDITOR
-
 		virtual bool ProcessMessages() override;
 
 		FVirtualLoopMessages TraceMessages;
diff --git a/Engine/Plugins/AudioInsights/Source/AudioInsights/Public/Messages/VirtualLoopTraceMessages.h b/Engine/Plugins/AudioInsights/Source/AudioInsights/Public/Messages/VirtualLoopTraceMessages.h
--- a/Engine/Plugins/AudioInsights/Source/AudioInsights/Public/Messages/VirtualLoopTraceMessages.h
+++ b/Engine/Plugins/AudioInsights/Source/AudioInsights/Public/Messages/VirtualLoopTraceMessages.h
@@ -105,22 +105,4 @@ namespace UE::Audio::Insights
 
 		friend class FVirtualLoopTraceProvider;
 	};
-
-#if !WITH_EDITOR
-	struct FVirtualLoopSessionCachedMessages
-	{
-		FVirtualLoopSessionCachedMessages(TraceServices::IAnalysisSession& InSession)
-			: VirtualizeCachedMessages(InSession.GetLinearAllocator(), 4096)
-			, StopOrRealizeCachedMessages(InSession.GetLinearAllocator(), 4096)
-			, UpdateCachedMessages(InSession.GetLinearAllocator(), 16384)
-		{
-
-		}
-
-		TraceServices::TPagedArray<FVirtualLoopVirtualizeMessage> VirtualizeCachedMessages;
-		TraceServices::TPagedArray<FVirtualLoopStopOrRealizeMessage> StopOrRealizeCachedMessages;
-		TraceServices::TPagedArray<FVirtualLoopUpdateMessage> UpdateCachedMessages;
-	};
-#endif // !WITH_EDITOR
-
 } // namespace UE::Audio::Insights
