[Audio Insights] Support cache manager inside standalone part 07: Tweaks to timing authority and plots to enable standalone cache to be scrubbable on running packaged builds

diff --git a/Engine/Plugins/AudioInsights/Source/AudioInsights/Private/AudioInsightsTimingViewExtender.cpp b/Engine/Plugins/AudioInsights/Source/AudioInsights/Private/AudioInsightsTimingViewExtender.cpp
--- a/Engine/Plugins/AudioInsights/Source/AudioInsights/Private/AudioInsightsTimingViewExtender.cpp
+++ b/Engine/Plugins/AudioInsights/Source/AudioInsights/Private/AudioInsightsTimingViewExtender.cpp
@@ -84,9 +84,7 @@ namespace UE::Audio::Insights
 		TraceServices::FAnalysisSessionReadScope SessionReadScope(*AnalysisSession);
 
 		const bool bIsValidTime = InTimeMarker >= 0.0 && InTimeMarker <= AnalysisSession->GetDurationSeconds();
-
-		// Scrub Audio Insights when all data is loaded in non-live sessions
-		if (AnalysisSession->IsAnalysisComplete() && bIsValidTime)
+		if (bIsValidTime)
 		{
 			if (!bUserInputDetected)
 			{
diff --git a/Engine/Plugins/AudioInsights/Source/AudioInsights/Private/Views/SoundPlotsWidgetView.cpp b/Engine/Plugins/AudioInsights/Source/AudioInsights/Private/Views/SoundPlotsWidgetView.cpp
--- a/Engine/Plugins/AudioInsights/Source/AudioInsights/Private/Views/SoundPlotsWidgetView.cpp
+++ b/Engine/Plugins/AudioInsights/Source/AudioInsights/Private/Views/SoundPlotsWidgetView.cpp
@@ -878,7 +878,7 @@ namespace UE::Audio::Insights
 		RangeUpperBound = FinalCurrentTime - BeginTimestamp + TimestampsDiff - RangeAlignmentOffset;
 #else
 		const TSharedPtr<const FAudioInsightsComponent> AudioInsightsComponent = FAudioInsightsModule::GetChecked().GetAudioInsightsComponent();
-		if (AudioInsightsComponent.IsValid() && !AudioInsightsComponent->GetIsLiveSession())
+		if (GameState == EGameState::ControlledByTimemarker || (AudioInsightsComponent.IsValid() && !AudioInsightsComponent->GetIsLiveSession()))
 		{
 			CurrentRange = TRange<double>(CurrentTimestamp - FAudioInsightsTimingViewExtender::MaxPlottingHistorySeconds, CurrentTimestamp);
 			return CurrentRange;
