AudioInsights

* Fixes: Plots grid not displayed when opening Audio Insights
* Fixes: https://jira.it.epicgames.com/browse/UE-346955

* The reason the plot wasn't rendering:
FSoundPlotsWidgetView::ConstructPlotsWidgetContents() wasn't calling
PlotWidget->SetCurvesMetadata()
PlotWidget->SetCurvesPointData();

* Which meant the following condition returned early, and the plot wasn't rendered
SAudioCurveView::PaintCurves()
{
// Skip drawing if curve data is not initialized yet
if (!PointDataPerCurve.IsValid() || !MetadataPerCurve.IsValid())
{
return LayerId;
}
}

* The fix is for FSoundPlotsWidgetView::ConstructPlotsWidgetContents() to call both SetCurvesMetadata() and SetCurvesPointData(), to pass the IsValid() checks in PaintCurves()

* Whilst here consolidating repeated code into a single function UpdatePlot()
* Lifetime of FPointDataPerCurveMap, FNumberFormattingOptions, PlotWidget should outlive UpdatePlot(), but defensively coding here just in case

Testing:
AudioQA
Editor build
Standalone Unreal Insights build

diff --git a/Engine/Plugins/AudioInsights/Source/AudioInsights/Private/Views/SoundPlotsWidgetView.cpp b/Engine/Plugins/AudioInsights/Source/AudioInsights/Private/Views/SoundPlotsWidgetView.cpp
--- a/Engine/Plugins/AudioInsights/Source/AudioInsights/Private/Views/SoundPlotsWidgetView.cpp
+++ b/Engine/Plugins/AudioInsights/Source/AudioInsights/Private/Views/SoundPlotsWidgetView.cpp
@@ -4,6 +4,7 @@
 #include "AudioInsightsModule.h"
 #include "AudioInsightsStyle.h"
 #include "Features/IModularFeatures.h"
+#include "Misc/AssertionMacros.h"
 #include "SoundPlotsWidgetCommands.h"
 #include "Widgets/Input/SComboBox.h"
 #include "Widgets/Input/SNumericEntryBox.h"
@@ -165,15 +166,7 @@ namespace UE::Audio::Insights
 #endif // !WITH_EDITOR
 			}
 
-			if (!CurveDataMapPtr->IsEmpty())
-			{
-				PlotWidget->SetCurvesPointData(CurveDataMapPtr);
-				PlotWidget->SetYValueFormattingOptions(*GetPlotColumnNumberFormat(SelectedPlotColumnName));
-				if (TOptional<FFloatInterval> CustomizedYRange = GetPlotColumnCustomizedYRange(SelectedPlotColumnName))
-				{
-					PlotWidget->SetYAxisRange(*CustomizedYRange);
-				}
-			}
+			UpdatePlot();
 		}
 
 #if !WITH_EDITOR
@@ -460,6 +453,14 @@ namespace UE::Audio::Insights
 
 		SelectedPlotColumnName = ColumnNames[0];
 
+		SAssignNew(PlotWidget, SAudioCurveView)
+			.AutoRangeYAxis_Lambda([this]() { return !GetPlotColumnUseCustomYRange(SelectedPlotColumnName); })
+			.ViewRange_Lambda([this]() { return GetViewRange(); })
+			.PixelSnappingMethod(EWidgetPixelSnapping::Disabled)
+			.AdditionalToolTipText_Lambda([this]() { return GetAdditionalTooltipText(); })
+			.OnScrubPositionChanged_Lambda([this](const double TimestampPosition, const float YValue) { SetTimestampFromCurveView(TimestampPosition); })
+			.OnKeyDown_Lambda([this](const FGeometry& Geometry, const FKeyEvent& KeyEvent) { return OnHandleKeyDown(Geometry, KeyEvent); });
+
 		// Initialize curve data and metadata
 		if (!PlotWidgetMetadataPerCurve.IsValid())
 		{
@@ -470,18 +471,11 @@ namespace UE::Audio::Insights
 				PlotWidgetCurveIdToPointDataMapPerColumn.Emplace(ColumnName, MoveTemp(PointDataPerCurveMap));
 				ColumnNameToLatestTimestamp.Emplace(ColumnName, MakeShared<FLatestTimestampPerCurveMap>());
 			}
-		}
 
-		SAssignNew(PlotWidget, SAudioCurveView)
-			.AutoRangeYAxis_Lambda([this]() { return !GetPlotColumnUseCustomYRange(SelectedPlotColumnName); })
-			.ViewRange_Lambda([this]() { return GetViewRange(); })
-			.PixelSnappingMethod(EWidgetPixelSnapping::Disabled)
-			.AdditionalToolTipText_Lambda([this]() { return GetAdditionalTooltipText(); })
-			.OnScrubPositionChanged_Lambda([this](const double TimestampPosition, const float YValue) { SetTimestampFromCurveView(TimestampPosition); })
-			.OnKeyDown_Lambda([this](const FGeometry& Geometry, const FKeyEvent& KeyEvent) { return OnHandleKeyDown(Geometry, KeyEvent); });
+			PlotWidget->SetCurvesMetadata(PlotWidgetMetadataPerCurve);
+		}
 
-		PlotWidget->SetYValueFormattingOptions(*GetPlotColumnNumberFormat(SelectedPlotColumnName));
-		PlotWidget->SetYAxisRange(GetPlotColumnCustomizedOrDefaultYRange(SelectedPlotColumnName));
+		UpdatePlot();
 
 		SAssignNew(ComboBox, SComboBox<FName>)
 			.ToolTipText(LOCTEXT("SoundPlotsWidgetView_SelectPlotColumnDescription", "Select a column from the table to plot."))
@@ -495,16 +489,10 @@ namespace UE::Audio::Insights
 			.OnSelectionChanged_Lambda([this](FName NewColumnName, ESelectInfo::Type)
 			{
 				SelectedPlotColumnName = NewColumnName;
-				if (TSharedPtr<FPointDataPerCurveMap>* DataMap = PlotWidgetCurveIdToPointDataMapPerColumn.Find(NewColumnName))
-				{
-					PlotWidget->SetCurvesPointData(*DataMap);
-					PlotWidget->SetYValueFormattingOptions(*GetPlotColumnNumberFormat(NewColumnName));
-					PlotWidget->SetYAxisRange(GetPlotColumnCustomizedOrDefaultYRange(NewColumnName));
-
+				UpdatePlot();
 #if !WITH_EDITOR
-					bTimeMarkerMoved = true;
+				bTimeMarkerMoved = true;
 #endif
-				}
 			})
 			[
 				SNew(STextBlock)
@@ -657,6 +645,33 @@ namespace UE::Audio::Insights
 		}
 	}
 
+	void FSoundPlotsWidgetView::UpdatePlot()
+	{
+		const TSharedPtr<FPointDataPerCurveMap>* const DataMap = PlotWidgetCurveIdToPointDataMapPerColumn.Find(SelectedPlotColumnName);
+		if (!DataMap || !DataMap->IsValid())
+		{
+			ensureMsgf(false, TEXT("Invalid FPointDataPerCurveMap in FSoundPlotsWidgetView::UpdatePlot"));
+			return;
+		}
+
+		const FNumberFormattingOptions* const PlotColumnNumberFormat = GetPlotColumnNumberFormat(SelectedPlotColumnName);
+		if (!PlotColumnNumberFormat)
+		{
+			ensureMsgf(false, TEXT("Invalid FNumberFormattingOptions in FSoundPlotsWidgetView::UpdatePlot"));
+			return;
+		}
+
+		if (!PlotWidget.IsValid())
+		{
+			ensureMsgf(false, TEXT("Invalid SAudioCurveView in FSoundPlotsWidgetView::UpdatePlot"));
+			return;
+		}
+
+		PlotWidget->SetCurvesPointData(*DataMap);
+		PlotWidget->SetYValueFormattingOptions(*PlotColumnNumberFormat);
+		PlotWidget->SetYAxisRange(GetPlotColumnCustomizedOrDefaultYRange(SelectedPlotColumnName));
+	}
+
 	void FSoundPlotsWidgetView::ResetPlots()
 	{
 		for (const auto& [ColumnName, PointDataPerCurveMap] : PlotWidgetCurveIdToPointDataMapPerColumn)
diff --git a/Engine/Plugins/AudioInsights/Source/AudioInsights/Public/Views/SoundPlotsWidgetView.h b/Engine/Plugins/AudioInsights/Source/AudioInsights/Public/Views/SoundPlotsWidgetView.h
--- a/Engine/Plugins/AudioInsights/Source/AudioInsights/Public/Views/SoundPlotsWidgetView.h
+++ b/Engine/Plugins/AudioInsights/Source/AudioInsights/Public/Views/SoundPlotsWidgetView.h
@@ -74,6 +74,8 @@ namespace UE::Audio::Insights
 
 		void UpdatePlotVisibilityRecursive(const TSharedPtr<IDashboardDataTreeViewEntry>& Entry);
 
+		void UpdatePlot();
+
 		// Column information used by plot widgets, keyed by column name. These keys should be a subset of the keys in GetColumns(). 
 		const TFunction<const ::Audio::TCircularAudioBuffer<FDataPoint>& (const IDashboardDataTreeViewEntry& InEntry)> GetPlotColumnDataFunc(const FName& ColumnName);
 		const TFunction<float (const float)> GetPlotColumnUnitConversionFunc(const FName& ColumnName);
