[Audio Insights] Support cache manager inside standalone part 06: Event Log and Sounds tab use cache manager in standalone

diff --git a/Engine/Plugins/AudioInsights/Source/AudioInsights/Private/Providers/AudioEventLogTraceProvider.cpp b/Engine/Plugins/AudioInsights/Source/AudioInsights/Private/Providers/AudioEventLogTraceProvider.cpp
--- a/Engine/Plugins/AudioInsights/Source/AudioInsights/Private/Providers/AudioEventLogTraceProvider.cpp
+++ b/Engine/Plugins/AudioInsights/Source/AudioInsights/Private/Providers/AudioEventLogTraceProvider.cpp
@@ -8,9 +8,11 @@
 
 #if WITH_EDITOR
 #include "AudioInsightsLog.h"
-#else
-#include "Common/PagedArray.h"
-#include "Async/ParallelFor.h"
+#endif 
+
+#if !WITH_EDITOR
+#include "AudioInsightsComponent.h"
+#include "AudioInsightsModule.h"
 #endif // WITH_EDITOR
 
 #define LOCTEXT_NAMESPACE "AudioInsights"
@@ -20,21 +22,17 @@ namespace UE::Audio::Insights
 	FAudioEventLogTraceProvider::FAudioEventLogTraceProvider()
 		: TDeviceDataMapTraceProvider<uint32, TSharedPtr<FAudioEventLogDashboardEntry>>(GetName_Static())
 	{
-#if WITH_EDITOR
 		FAudioInsightsCacheManager& CacheManager = IAudioInsightsModule::GetChecked().GetCacheManager();
 		CacheManager.OnChunkOverwritten.AddRaw(this, &FAudioEventLogTraceProvider::OnCacheChunkOverwritten);
-#endif // WITH_EDITOR
 	}
 
 	FAudioEventLogTraceProvider::~FAudioEventLogTraceProvider()
 	{
-#if WITH_EDITOR
 		if (FModuleManager::Get().IsModuleLoaded("AudioInsights"))
 		{
 			FAudioInsightsCacheManager& CacheManager = IAudioInsightsModule::GetChecked().GetCacheManager();
 			CacheManager.OnChunkOverwritten.RemoveAll(this);
 		}
-#endif // WITH_EDITOR
 	}
 
 	UE::Trace::IAnalyzer* FAudioEventLogTraceProvider::ConstructAnalyzer(TraceServices::IAnalysisSession& InSession)
@@ -69,11 +67,7 @@ namespace UE::Audio::Insights
 					case RouteId_SoundStart:
 					case RouteId_SoundStop:
 					{
-#if WITH_EDITOR
 						CacheMessage<FAudioEventLogMessage>(Context, Messages.AudioEventLogMessages);
-#else
-						Messages.AudioEventLogMessages.Enqueue(FAudioEventLogMessage(Context));
-#endif // WITH_EDITOR
 						break;
 					}
 
@@ -164,12 +158,22 @@ namespace UE::Audio::Insights
 
 	bool FAudioEventLogTraceProvider::ProcessManuallyUpdatedEntries()
 	{
-#if WITH_EDITOR
 		if (!bCacheChunksUpdated)
 		{
 			return false;
 		}
 
+#if !WITH_EDITOR
+		const TSharedPtr<const FAudioInsightsComponent> AudioInsightsComponent = FAudioInsightsModule::GetChecked().GetAudioInsightsComponent();
+		if (AudioInsightsComponent.IsValid() && !AudioInsightsComponent->GetIsLiveSession())
+		{
+			// Don't update the caching and processing status if the cache isn't limited (i.e. loading in a .urace file)
+			bCacheChunksUpdated = false;
+			return false;
+		}
+#endif // !WITH_EDITOR
+
+
 		const double AboutToBeDeletedTimestamp = GetAboutToBeDeletedFromCacheTimeThreshold();
 
 		FAudioInsightsCacheManager& CacheManager = IAudioInsightsModule::GetChecked().GetCacheManager();
@@ -214,9 +218,6 @@ namespace UE::Audio::Insights
 
 		bCacheChunksUpdated = false;
 		return bAnyEntriesRemoved;
-#else
-		return false;
-#endif // WITH_EDITOR
 	}
 
 	TSharedPtr<FAudioEventLogDashboardEntry> FAudioEventLogTraceProvider::CreateNewEntry(const FAudioEventLogMessage& Msg)
@@ -268,7 +269,6 @@ namespace UE::Audio::Insights
 		return Entry;
 	}
 
-#if WITH_EDITOR
 	void FAudioEventLogTraceProvider::OnCacheChunkOverwritten(const double NewCacheStartTimestamp)
 	{
 		CacheStartTimestamp = NewCacheStartTimestamp;
@@ -302,7 +302,6 @@ namespace UE::Audio::Insights
 
 		return NextToBeDeletedChunk ? NextToBeDeletedChunk->GetChunkTimeRangeEnd() : TNumericLimits<double>::Lowest();
 	}
-#endif // WITH_EDITOR
 
 	void FAudioEventLogTraceProvider::ClearActiveAudioDeviceEntries()
 	{
diff --git a/Engine/Plugins/AudioInsights/Source/AudioInsights/Private/Providers/AudioEventLogTraceProvider.h b/Engine/Plugins/AudioInsights/Source/AudioInsights/Private/Providers/AudioEventLogTraceProvider.h
--- a/Engine/Plugins/AudioInsights/Source/AudioInsights/Private/Providers/AudioEventLogTraceProvider.h
+++ b/Engine/Plugins/AudioInsights/Source/AudioInsights/Private/Providers/AudioEventLogTraceProvider.h
@@ -30,10 +30,8 @@ namespace UE::Audio::Insights
 
 		TSharedPtr<FAudioEventLogDashboardEntry> CreateNewEntry(const FAudioEventLogMessage& Msg);
 
-#if WITH_EDITOR
 		void OnCacheChunkOverwritten(const double NewCacheStartTimestamp);
 		double GetAboutToBeDeletedFromCacheTimeThreshold() const;
-#endif // WITH_EDITOR
 
 		FAudioEventLogMessages TraceMessages;
 
diff --git a/Engine/Plugins/AudioInsights/Source/AudioInsights/Private/Providers/SoundTraceProvider.cpp b/Engine/Plugins/AudioInsights/Source/AudioInsights/Private/Providers/SoundTraceProvider.cpp
--- a/Engine/Plugins/AudioInsights/Source/AudioInsights/Private/Providers/SoundTraceProvider.cpp
+++ b/Engine/Plugins/AudioInsights/Source/AudioInsights/Private/Providers/SoundTraceProvider.cpp
@@ -8,10 +8,6 @@
 #include "Trace/Analyzer.h"
 #include "TraceServices/Model/AnalysisSession.h"
 
-#if !WITH_EDITOR
-#include "Common/PagedArray.h"
-#endif // !WITH_EDITOR
-
 #define LOCTEXT_NAMESPACE "AudioInsights"
 
 namespace UE::Audio::Insights
@@ -89,48 +85,6 @@ namespace UE::Audio::Insights
 			OutDataRange.Push({ Msg.Timestamp, DataPoint });
 		}
 
-#if !WITH_EDITOR
-		template<typename T>
-		void CacheStandaloneParameterEntry(const T& Msg, TraceServices::TPagedArray<T>& OutCachedMessages)
-		{
-			// When loading in trace files, we very occasionally can get parameter messages out of order
-			// In these situations we drop the one message that was sent out of sync
-			if (OutCachedMessages.Num() == 0 || Msg.Timestamp >= OutCachedMessages.Last().Timestamp)
-			{
-				OutCachedMessages.EmplaceBack(Msg);
-			}
-		}
-
-		struct PlayOrderPair
-		{
-			uint32 ActiveSoundPlayOrderID = static_cast<uint32>(INDEX_NONE);
-			uint32 WaveInstancePlayOrderID = static_cast<uint32>(INDEX_NONE);
-		};
-
-		void CollectPlayOrderPairsRecursive(const TSharedPtr<IDashboardDataTreeViewEntry>& Entry, TArray<PlayOrderPair>& OutPlayOrderArray)
-		{
-			using namespace FSoundTraceProviderPrivate;
-
-			if (!Entry.IsValid())
-			{
-				return;
-			}
-
-			const FSoundDashboardEntry& SoundEntry = CastEntry(*Entry);
-			const PlayOrderPair PlayOrderPair { SoundEntry.ActiveSoundPlayOrder, SoundEntry.WaveInstancePlayOrder };
-
-			if (PlayOrderPair.WaveInstancePlayOrderID != INDEX_NONE || PlayOrderPair.ActiveSoundPlayOrderID != INDEX_NONE)
-			{
-				OutPlayOrderArray.Add(PlayOrderPair);
-			}
-
-			for (const TSharedPtr<IDashboardDataTreeViewEntry>& Child : SoundEntry.Children)
-			{
-				CollectPlayOrderPairsRecursive(Child, OutPlayOrderArray);
-			}
-		}
-#endif // !WITH_EDITOR
-
 		//////////////////////////////////////////////////////////////////////////////////////////////////////////////
 		// Finds an Active Sound entry inside a Category entry
 		void FindActiveSoundEntryInCategory(const TSharedPtr<FSoundDashboardEntry>* CategoryEntry, const uint32 ActiveSoundPlayOrder, TSharedPtr<FSoundDashboardEntry>& OutActiveSoundEntry)
@@ -185,34 +139,6 @@ namespace UE::Audio::Insights
 
 			return nullptr;
 		}
-
-#if !WITH_EDITOR
-		template<typename T>
-		void IterateOverRange(const TraceServices::TPagedArray<T>& InCachedMessages, const double Start, const double End, TFunctionRef<void(const T& Message)> OnMessageFunc)
-		{
-			if (InCachedMessages.Num() == 0)
-			{
-				return;
-			}
-
-			const int32 ClosestMessageToTimeStampIndex = TraceServices::PagedArrayAlgo::BinarySearchClosestBy(InCachedMessages, Start, [](const T& Msg) { return Msg.Timestamp; });
-			
-			for (auto It = InCachedMessages.GetIteratorFromItem(ClosestMessageToTimeStampIndex); It != InCachedMessages.end(); ++It)
-			{
-				if (It->Timestamp > End)
-				{
-					break;
-				}
-
-				if (It->Timestamp < Start)
-				{
-					continue;
-				}
-
-				OnMessageFunc(*It);
-			}
-		}
-#endif // !WITH_EDITOR
 	}
 
 	FSoundTraceProvider::FSoundTraceProvider()
@@ -281,134 +207,80 @@ namespace UE::Audio::Insights
 				{
 					case RouteId_ActiveSoundStart:
 					{
-#if WITH_EDITOR
 						CacheMessage<FSoundStartMessage>(Context, Messages.ActiveSoundStartMessages);
-#else
-						Messages.ActiveSoundStartMessages.Enqueue(FSoundStartMessage{ Context });
-#endif // WITH_EDITOR
 						break;
 					}
 
 					case RouteId_SoundWaveStart:
 					{
-#if WITH_EDITOR
 						CacheMessage<FSoundWaveStartMessage>(Context, Messages.SoundWaveStartMessages);
-#else
-						Messages.SoundWaveStartMessages.Enqueue(FSoundWaveStartMessage(Context));
-#endif // WITH_EDITOR
 						break;
 					}
 
 					case RouteId_SoundIsAlivePing:
 					{
-#if WITH_EDITOR
 						CacheMessage<FSoundIsAlivePingMessage>(Context, Messages.ActiveSoundIsAlivePingMessages);
-#else
-						Messages.ActiveSoundIsAlivePingMessages.Enqueue(FSoundIsAlivePingMessage(Context));
-#endif // WITH_EDITOR
 						break;
 					}
 
 					case RouteId_SoundWaveIsAlivePing:
 					{
-#if WITH_EDITOR
 						CacheMessage<FSoundWaveIsAlivePingMessage>(Context, Messages.SoundWaveIsAlivePingMessages);
-#else
-						Messages.SoundWaveIsAlivePingMessages.Enqueue(FSoundWaveIsAlivePingMessage(Context));
-#endif // WITH_EDITOR
 						break;
 					}
 
 					case RouteId_Stop:
 					{
-#if WITH_EDITOR
 						CacheMessage<FSoundStopMessage>(Context, Messages.StopMessages);
-#else
-						Messages.StopMessages.Enqueue(FSoundStopMessage{ Context });
-#endif // WITH_EDITOR
 						break;
 					}
 
 					case RouteId_Priority:
 					{
-#if WITH_EDITOR
-						CacheMessage<FSoundPriorityMessage>(Context, Messages.PriorityMessages);
-#else
-						Messages.PriorityMessages.Enqueue(FSoundPriorityMessage{ Context });
-#endif // WITH_EDITOR
+						CacheMessageIfLatest<FSoundPriorityMessage>(Context, Messages.PriorityMessages);
 						break;
 					}
 
 					case RouteId_Distance:
 					{
-#if WITH_EDITOR
-						CacheMessage<FSoundDistanceMessage>(Context, Messages.DistanceMessages);
-#else
-						Messages.DistanceMessages.Enqueue(FSoundDistanceMessage{ Context });
-#endif // WITH_EDITOR
+						CacheMessageIfLatest<FSoundDistanceMessage>(Context, Messages.DistanceMessages);
 						break;
 					}
 
 					case RouteId_DistanceAttenuation:
 					{
-#if WITH_EDITOR
-						CacheMessage<FSoundDistanceAttenuationMessage>(Context, Messages.DistanceAttenuationMessages);
-#else
-						Messages.DistanceAttenuationMessages.Enqueue(FSoundDistanceAttenuationMessage{ Context });
-#endif // WITH_EDITOR
+						CacheMessageIfLatest<FSoundDistanceAttenuationMessage>(Context, Messages.DistanceAttenuationMessages);
 						break;
 					}
 
 					case RouteId_Filters:
 					{
-#if WITH_EDITOR
-						CacheMessage<FSoundLPFFreqMessage>(Context, Messages.LPFFreqMessages);
-						CacheMessage<FSoundHPFFreqMessage>(Context, Messages.HPFFreqMessages);
-#else
-						Messages.LPFFreqMessages.Enqueue(FSoundLPFFreqMessage{ Context });
-						Messages.HPFFreqMessages.Enqueue(FSoundHPFFreqMessage{ Context });
-#endif // WITH_EDITOR
-
+						CacheMessageIfLatest<FSoundLPFFreqMessage>(Context, Messages.LPFFreqMessages);
+						CacheMessageIfLatest<FSoundHPFFreqMessage>(Context, Messages.HPFFreqMessages);
 						break;
 					}
 
 					case RouteId_Amplitude:
 					{
-#if WITH_EDITOR
-						CacheMessage<FSoundEnvelopeMessage>(Context, Messages.AmplitudeMessages);
-#else
-						Messages.AmplitudeMessages.Enqueue(FSoundEnvelopeMessage{ Context });
-#endif // WITH_EDITOR
+						CacheMessageIfLatest<FSoundEnvelopeMessage>(Context, Messages.AmplitudeMessages);
 						break;
 					}
 
 					case RouteId_Volume:
 					{
-#if WITH_EDITOR
-						CacheMessage<FSoundVolumeMessage>(Context, Messages.VolumeMessages);
-#else
-						Messages.VolumeMessages.Enqueue(FSoundVolumeMessage{ Context });
-#endif // WITH_EDITOR
+						CacheMessageIfLatest<FSoundVolumeMessage>(Context, Messages.VolumeMessages);
 						break;
 					}
 
 					case RouteId_Pitch:
 					{
-#if WITH_EDITOR
-						CacheMessage<FSoundPitchMessage>(Context, Messages.PitchMessages);
-#else
-						Messages.PitchMessages.Enqueue(FSoundPitchMessage{ Context });
-#endif // WITH_EDITOR
+						CacheMessageIfLatest<FSoundPitchMessage>(Context, Messages.PitchMessages);
 						break;
 					}
 
 					case RouteId_RelativeRenderCost:
 					{
-#if WITH_EDITOR
-						CacheMessage<FSoundRelativeRenderCostMessage>(Context, Messages.RelativeRenderCostMessages);
-#else
-						Messages.RelativeRenderCostMessages.Enqueue(FSoundRelativeRenderCostMessage{ Context });
-#endif // WITH_EDITOR
+						CacheMessageIfLatest<FSoundRelativeRenderCostMessage>(Context, Messages.RelativeRenderCostMessages);
 						break;
 					}
 
@@ -457,13 +329,6 @@ namespace UE::Audio::Insights
 		return new FSoundTraceAnalyzer(AsShared(), InSession);
 	}
 
-#if !WITH_EDITOR
-	void FSoundTraceProvider::InitSessionCachedMessages(TraceServices::IAnalysisSession& InSession)
-	{
-		SessionCachedMessages = MakeUnique<FSoundSessionCachedMessages>(InSession);
-	}
-#endif // !WITH_EDITOR
-
 	bool FSoundTraceProvider::ProcessMessages()
 	{
 		using namespace FSoundTraceProviderPrivate;
@@ -579,52 +444,24 @@ namespace UE::Audio::Insights
 		[this, &ProcessStartMessage](const FSoundStartMessage& Msg, const TSharedPtr<FSoundDashboardEntry>* OutActiveSoundEntry)
 		{
 			ProcessStartMessage(Msg, OutActiveSoundEntry);
-
-#if !WITH_EDITOR
-			if (SessionCachedMessages.IsValid())
-			{
-				SessionCachedMessages->StartCachedMessages.EmplaceBack(Msg);
-			}
-#endif // !WITH_EDITOR
 		});
 
 		ProcessMessageQueue<FSoundWaveStartMessage>(TraceMessages.SoundWaveStartMessages, GetActiveSoundEntryFromSoundWaveStartMessage,
 		[this, &ProcessSoundWaveStartMessage](const FSoundWaveStartMessage& Msg, const TSharedPtr<FSoundDashboardEntry>* OutActiveSoundEntry)
 		{
 			ProcessSoundWaveStartMessage(Msg, OutActiveSoundEntry);
-
-#if !WITH_EDITOR
-			if (SessionCachedMessages.IsValid())
-			{
-				SessionCachedMessages->SoundWaveStartCachedMessages.EmplaceBack(Msg);
-			}
-#endif // !WITH_EDITOR
 		});
 
 		ProcessMessageQueue<FSoundIsAlivePingMessage>(TraceMessages.ActiveSoundIsAlivePingMessages, CreateAndGetActiveSoundEntryFromPingIfNotFound,
 		[this, &ProcessStartMessage](const FSoundIsAlivePingMessage& Msg, const TSharedPtr<FSoundDashboardEntry>* OutActiveSoundEntry)
 		{
 			ProcessStartMessage(Msg, OutActiveSoundEntry);
-
-#if !WITH_EDITOR
-			if (SessionCachedMessages.IsValid())
-			{
-				SessionCachedMessages->SoundIsAlivePingCachedMessages.EmplaceBack(Msg);
-			}
-#endif // !WITH_EDITOR
 		});
 
 		ProcessMessageQueue<FSoundWaveIsAlivePingMessage>(TraceMessages.SoundWaveIsAlivePingMessages, GetActiveSoundEntryFromSoundWaveStartMessage,
 		[this, &ProcessSoundWaveStartMessage](const FSoundWaveIsAlivePingMessage& Msg, const TSharedPtr<FSoundDashboardEntry>* OutActiveSoundEntry)
 		{
 			ProcessSoundWaveStartMessage(Msg, OutActiveSoundEntry);
-			
-#if !WITH_EDITOR
-			if (SessionCachedMessages.IsValid())
-			{
-				SessionCachedMessages->SoundWaveIsAlivePingCachedMessages.EmplaceBack(Msg);
-			}
-#endif //! WITH_EDITOR
 		});
 
 		TArray<int32, TInlineAllocator<64>> EntriesWithPoppedDataPoints;
@@ -638,13 +475,6 @@ namespace UE::Audio::Insights
 				return;
 			}
 
-#if !WITH_EDITOR
-			if (SessionCachedMessages.IsValid())
-			{
-				CacheStandaloneParameterEntry(Msg, SessionCachedMessages->PriorityCachedMessages);
-			}
-#endif // !WITH_EDITOR
-
 			FSoundTraceProviderPrivate::UpdateParameterEntry(Msg, Msg.Priority, SoundEntry->PriorityDataPoint, SoundEntry->PriorityDataRange, EntriesWithPoppedDataPoints);
 		});
 		EntriesWithPoppedDataPoints.Reset();
@@ -658,13 +488,6 @@ namespace UE::Audio::Insights
 				return;
 			}
 
-#if !WITH_EDITOR
-			if (SessionCachedMessages.IsValid())
-			{
-				CacheStandaloneParameterEntry(Msg, SessionCachedMessages->DistanceCachedMessages);
-			}
-#endif // !WITH_EDITOR
-
 			FSoundTraceProviderPrivate::UpdateParameterEntry(Msg, Msg.Distance, SoundEntry->DistanceDataPoint, SoundEntry->DistanceDataRange, EntriesWithPoppedDataPoints);
 		});
 		EntriesWithPoppedDataPoints.Reset();
@@ -678,13 +501,6 @@ namespace UE::Audio::Insights
 				return;
 			}
 
-#if !WITH_EDITOR
-			if (SessionCachedMessages.IsValid())
-			{
-				CacheStandaloneParameterEntry(Msg, SessionCachedMessages->DistanceAttenuationCachedMessages);
-			}
-#endif // !WITH_EDITOR
-
 			FSoundTraceProviderPrivate::UpdateParameterEntry(Msg, Msg.DistanceAttenuation, SoundEntry->DistanceAttenuationDataPoint, SoundEntry->DistanceAttenuationDataRange, EntriesWithPoppedDataPoints);
 		});
 		EntriesWithPoppedDataPoints.Reset();
@@ -698,13 +514,6 @@ namespace UE::Audio::Insights
 				return;
 			}
 
-#if !WITH_EDITOR
-			if (SessionCachedMessages.IsValid())
-			{
-				CacheStandaloneParameterEntry(Msg, SessionCachedMessages->LPFFreqCachedMessages);
-			}
-#endif // !WITH_EDITOR
-
 			FSoundTraceProviderPrivate::UpdateParameterEntry(Msg, Msg.LPFFrequency, SoundEntry->LPFFreqDataPoint, SoundEntry->LPFFreqDataRange, EntriesWithPoppedDataPoints);
 		});
 		EntriesWithPoppedDataPoints.Reset();
@@ -718,13 +527,6 @@ namespace UE::Audio::Insights
 				return;
 			}
 
-#if !WITH_EDITOR
-			if (SessionCachedMessages.IsValid())
-			{
-				CacheStandaloneParameterEntry(Msg, SessionCachedMessages->HPFFreqCachedMessages);
-			}
-#endif // !WITH_EDITOR
-
 			FSoundTraceProviderPrivate::UpdateParameterEntry(Msg, Msg.HPFFrequency, SoundEntry->HPFFreqDataPoint, SoundEntry->HPFFreqDataRange, EntriesWithPoppedDataPoints);
 		});
 		EntriesWithPoppedDataPoints.Reset();
@@ -738,13 +540,6 @@ namespace UE::Audio::Insights
 				return;
 			}
 
-#if !WITH_EDITOR
-			if (SessionCachedMessages.IsValid())
-			{
-				CacheStandaloneParameterEntry(Msg, SessionCachedMessages->AmplitudeCachedMessages);
-			}
-#endif // !WITH_EDITOR
-
 			FSoundTraceProviderPrivate::UpdateParameterEntry(Msg, Msg.Envelope, SoundEntry->AmplitudeDataPoint, SoundEntry->AmplitudeDataRange, EntriesWithPoppedDataPoints);
 		});
 		EntriesWithPoppedDataPoints.Reset();
@@ -758,13 +553,6 @@ namespace UE::Audio::Insights
 				return;
 			}
 
-#if !WITH_EDITOR
-			if (SessionCachedMessages.IsValid())
-			{
-				CacheStandaloneParameterEntry(Msg, SessionCachedMessages->VolumeCachedMessages);
-			}
-#endif // !WITH_EDITOR
-
 			FSoundTraceProviderPrivate::UpdateParameterEntry(Msg, Msg.Volume, SoundEntry->VolumeDataPoint, SoundEntry->VolumeDataRange, EntriesWithPoppedDataPoints);
 		});
 		EntriesWithPoppedDataPoints.Reset();
@@ -778,13 +566,6 @@ namespace UE::Audio::Insights
 				return;
 			}
 
-#if !WITH_EDITOR
-			if (SessionCachedMessages.IsValid())
-			{
-				CacheStandaloneParameterEntry(Msg, SessionCachedMessages->PitchCachedMessages);
-			}
-#endif // !WITH_EDITOR
-
 			FSoundTraceProviderPrivate::UpdateParameterEntry(Msg, Msg.Pitch, SoundEntry->PitchDataPoint, SoundEntry->PitchDataRange, EntriesWithPoppedDataPoints);
 		});
 		EntriesWithPoppedDataPoints.Reset();
@@ -798,13 +579,6 @@ namespace UE::Audio::Insights
 				return;
 			}
 
-#if !WITH_EDITOR
-			if (SessionCachedMessages.IsValid())
-			{
-				CacheStandaloneParameterEntry(Msg, SessionCachedMessages->RelativeRenderCostCachedMessages);
-			}
-#endif // !WITH_EDITOR
-
 			FSoundTraceProviderPrivate::UpdateParameterEntry(Msg, Msg.RelativeRenderCost, SoundEntry->RelativeRenderCostDataPoint, SoundEntry->RelativeRenderCostDataRange, EntriesWithPoppedDataPoints);
 		});
 		EntriesWithPoppedDataPoints.Reset();
@@ -812,13 +586,6 @@ namespace UE::Audio::Insights
 		ProcessMessageQueue<FSoundStopMessage>(TraceMessages.StopMessages, GetSoundEntryFromStopMessage,
 		[this](const FSoundStopMessage& Msg, const TSharedPtr<FSoundDashboardEntry>* OutActiveSoundEntry)
 		{
-#if !WITH_EDITOR
-			if (SessionCachedMessages.IsValid())
-			{
-				SessionCachedMessages->StopCachedMessages.EmplaceBack(Msg);
-			}
-#endif // !WITH_EDITOR
-
 			if (OutActiveSoundEntry && OutActiveSoundEntry->IsValid())
 			{
 				if ((*OutActiveSoundEntry)->Timestamp < Msg.Timestamp)
@@ -1179,13 +946,6 @@ namespace UE::Audio::Insights
 	{
 		using namespace FSoundTraceProviderPrivate;
 
-#if !WITH_EDITOR
-		if (!SessionCachedMessages.IsValid())
-		{
-			return;
-		}
-#endif // !WITH_EDITOR
-
 		auto InitActiveSoundEntry = [this](const FSoundStartMessage& Msg, FSoundDashboardEntry& OutActiveSoundEntry)
 		{
 			OutActiveSoundEntry.SetName(*Msg.Name);
@@ -1245,25 +1005,12 @@ namespace UE::Audio::Insights
 			InitActiveSoundEntry(StartCachedMessage, *ActiveSoundEntry);
 		};
 
-#if WITH_EDITOR
 		const FAudioInsightsCacheManager& CacheManager = FAudioInsightsModule::GetChecked().GetCacheManager();
 		CacheManager.IterateTo<FSoundStartMessage>(SoundMessageNames::SoundStart, TimeMarker, [&ProcessStartMessage](const FSoundStartMessage& StartCachedMessage)
 		{
 			ProcessStartMessage(StartCachedMessage);
 		});
-#else
-		for (const FSoundStartMessage& StartCachedMessage : SessionCachedMessages->StartCachedMessages)
-		{
-			if (StartCachedMessage.Timestamp > TimeMarker)
-			{
-				break;
-			}
-
-			ProcessStartMessage(StartCachedMessage);
-		}
-#endif // WITH_EDITOR
 			
-
 		// Collect all the is alive ping messages registered until this point in time
 		auto ProcessSoundIsAlivePingMessage = [this, &AliveActiveSoundIDs, &InitActiveSoundEntry](const FSoundIsAlivePingMessage& IsAlivePingCachedMessage)
 		{
@@ -1280,22 +1027,10 @@ namespace UE::Audio::Insights
 			InitActiveSoundEntry(IsAlivePingCachedMessage, *ActiveSoundEntry);
 		};
 
-#if WITH_EDITOR
 		CacheManager.IterateTo<FSoundIsAlivePingMessage>(SoundMessageNames::SoundIsAlivePing, TimeMarker, [&ProcessSoundIsAlivePingMessage](const FSoundIsAlivePingMessage& IsAlivePingCachedMessage)
 		{
 			ProcessSoundIsAlivePingMessage(IsAlivePingCachedMessage);
 		});
-#else
-		for (const FSoundIsAlivePingMessage& IsAlivePingCachedMessage : SessionCachedMessages->SoundIsAlivePingCachedMessages)
-		{
-			if (IsAlivePingCachedMessage.Timestamp > TimeMarker)
-			{
-				break;
-			}
-
-			ProcessSoundIsAlivePingMessage(IsAlivePingCachedMessage);
-		}
-#endif // WITH_EDITOR
 
 		////////////////////////////////////////
 		// Add all soundwave start messages registered until this point in time 
@@ -1325,22 +1060,10 @@ namespace UE::Audio::Insights
 			CreateChildEntry(SoundWaveStartCachedMessage, ParentSoundEntryReturn);
 		};
 
-#if WITH_EDITOR
 		CacheManager.IterateTo<FSoundWaveStartMessage>(SoundMessageNames::SoundWaveStart, TimeMarker, [&ProcessSoundWaveStartMessage](const FSoundWaveStartMessage& SoundWaveStartCachedMessage)
 		{
 			ProcessSoundWaveStartMessage(SoundWaveStartCachedMessage);
 		});
-#else
-		for (const FSoundWaveStartMessage& SoundWaveStartCachedMessage : SessionCachedMessages->SoundWaveStartCachedMessages)
-		{
-			if (SoundWaveStartCachedMessage.Timestamp > TimeMarker)
-			{
-				break;
-			}
-
-			ProcessSoundWaveStartMessage(SoundWaveStartCachedMessage);
-		}
-#endif // WITH_EDITOR
 
 		////////////////////////////////////////
 		// Add all soundwave ping messages registered until this point in time 
@@ -1371,22 +1094,10 @@ namespace UE::Audio::Insights
 			CreateChildEntry(SoundWaveIsAlivePingCachedMessage, ParentSoundIsAlivePingEntryReturn);
 		};
 
-#if WITH_EDITOR
 		CacheManager.IterateTo<FSoundWaveIsAlivePingMessage>(SoundMessageNames::SoundWaveIsAlivePing, TimeMarker, [&ProcessSoundWaveIsAlivePingMessage](const FSoundWaveIsAlivePingMessage& SoundWaveIsAlivePingCachedMessage)
 		{
 			ProcessSoundWaveIsAlivePingMessage(SoundWaveIsAlivePingCachedMessage);
 		});
-#else
-		for (const FSoundWaveIsAlivePingMessage& SoundWaveIsAlivePingCachedMessage : SessionCachedMessages->SoundWaveIsAlivePingCachedMessages)
-		{
-			if (SoundWaveIsAlivePingCachedMessage.Timestamp > TimeMarker)
-			{
-				break;
-			}
-
-			ProcessSoundWaveIsAlivePingMessage(SoundWaveIsAlivePingCachedMessage);
-		}
-#endif // WITH_EDITOR
 
 		// Selectively remove start messages collected in the step above by knowing which sounds were stopped.
 		// With this we will know what are the active sounds at this point in time.
@@ -1438,22 +1149,11 @@ namespace UE::Audio::Insights
 				}
 			}
 		};
-#if WITH_EDITOR
+
 		CacheManager.IterateTo<FSoundStopMessage>(SoundMessageNames::SoundStop, TimeMarker, [&ProcessSoundStopMessageMessage](const FSoundStopMessage& StopCachedMessage)
 		{
 			ProcessSoundStopMessageMessage(StopCachedMessage);
 		});
-#else
-		for (const FSoundStopMessage& StopCachedMessage : SessionCachedMessages->StopCachedMessages)
-		{
-			if (StopCachedMessage.Timestamp > TimeMarker)
-			{
-				break;
-			}
-
-			ProcessSoundStopMessageMessage(StopCachedMessage);
-		}
-#endif // WITH_EDITOR
 
 		TArray<uint32> ActiveSoundPlayOrderArray;
 		ActiveSoundToEntryKeysMap.GetKeys(ActiveSoundPlayOrderArray);
@@ -1500,11 +1200,7 @@ namespace UE::Audio::Insights
 			}
 		}
 
-#if WITH_EDITOR
 		CollectParamsForTimestamp(CacheManager, TimeMarker);
-#else
-		CollectParamsForTimestamp(TimeMarker);
-#endif // WITH_EDITOR
 
 		// Call parent method to update LastMessageId
 		FTraceProviderBase::OnTimingViewTimeMarkerChanged(TimeMarker);
@@ -1523,11 +1219,7 @@ namespace UE::Audio::Insights
 		Reset();
 	}
 
-#if WITH_EDITOR
 	void FSoundTraceProvider::CollectParamsForTimestamp(const FAudioInsightsCacheManager& CacheManager, const double TimeMarker)
-#else
-	void FSoundTraceProvider::CollectParamsForTimestamp(const double TimeMarker)
-#endif // WITH_EDITOR
 	{
 		using namespace FSoundTraceProviderPrivate;
 
@@ -1565,21 +1257,12 @@ namespace UE::Audio::Insights
 		const double EndTime = FMath::Max(0.0f, PlottingRange.GetUpperBoundValue() + FAudioInsightsTimingViewExtender::PlottingMarginSeconds);
 
 		// Using ParallelFor to speed-up the cached messages retrieval, using a traditional for loop is unacceptably slower, specially in large traces.
-		ParallelFor(NumMessageTypesToProcess,
-#if WITH_EDITOR
-		[&CacheManager, TimeMarker, StartTime, EndTime, this](const int32 Index)
-#else
-		[TimeMarker, StartTime, EndTime, this](const int32 Index)
-#endif // WITH_EDITOR
+		ParallelFor(NumMessageTypesToProcess, [&CacheManager, TimeMarker, StartTime, EndTime, this](const int32 Index)
 		{
 			switch (static_cast<ESoundParamMessageType>(Index))
 			{
 			case ESoundParamMessageType::Priority:
-#if WITH_EDITOR
 				CacheManager.IterateOverRange<FSoundPriorityMessage>(SoundMessageNames::PriorityParam, StartTime, EndTime, [this, TimeMarker](const FSoundPriorityMessage& Message)
-#else
-				IterateOverRange<FSoundPriorityMessage>(SessionCachedMessages->PriorityCachedMessages, StartTime, EndTime, [this, TimeMarker](const FSoundPriorityMessage& Message)
-#endif // WITH_EDITOR
 				{
 					TSharedPtr<FSoundDashboardEntry> ActiveSoundEntry;
 					GetActiveSoundEntryFromIDs(Message.ActiveSoundPlayOrder, ActiveSoundEntry);
@@ -1601,11 +1284,7 @@ namespace UE::Audio::Insights
 				break;
 
 			case ESoundParamMessageType::Distance:
-#if WITH_EDITOR
 				CacheManager.IterateOverRange<FSoundDistanceMessage>(SoundMessageNames::DistanceParam, StartTime, EndTime, [this, TimeMarker](const FSoundDistanceMessage& Message)
-#else
-				IterateOverRange<FSoundDistanceMessage>(SessionCachedMessages->DistanceCachedMessages, StartTime, EndTime, [this, TimeMarker](const FSoundDistanceMessage& Message)
-#endif // WITH_EDITOR
 				{
 					TSharedPtr<FSoundDashboardEntry> ActiveSoundEntry;
 					GetActiveSoundEntryFromIDs(Message.ActiveSoundPlayOrder, ActiveSoundEntry);
@@ -1627,11 +1306,7 @@ namespace UE::Audio::Insights
 				break;
 
 			case ESoundParamMessageType::DistanceAttenuation:
-#if WITH_EDITOR
 				CacheManager.IterateOverRange<FSoundDistanceAttenuationMessage>(SoundMessageNames::DistanceAttenuationParam, StartTime, EndTime, [this, TimeMarker](const FSoundDistanceAttenuationMessage& Message)
-#else
-				IterateOverRange<FSoundDistanceAttenuationMessage>(SessionCachedMessages->DistanceAttenuationCachedMessages, StartTime, EndTime, [this, TimeMarker](const FSoundDistanceAttenuationMessage& Message)
-#endif // WITH_EDITOR
 				{
 					TSharedPtr<FSoundDashboardEntry> ActiveSoundEntry;
 					GetActiveSoundEntryFromIDs(Message.ActiveSoundPlayOrder, ActiveSoundEntry);
@@ -1653,11 +1328,7 @@ namespace UE::Audio::Insights
 				break;
 
 			case ESoundParamMessageType::LPFFreq:
-#if WITH_EDITOR
 				CacheManager.IterateOverRange<FSoundLPFFreqMessage>(SoundMessageNames::LPFFreqParam, StartTime, EndTime, [this, TimeMarker](const FSoundLPFFreqMessage& Message)
-#else
-				IterateOverRange<FSoundLPFFreqMessage>(SessionCachedMessages->LPFFreqCachedMessages, StartTime, EndTime, [this, TimeMarker](const FSoundLPFFreqMessage& Message)
-#endif // WITH_EDITOR
 				{
 					TSharedPtr<FSoundDashboardEntry> ActiveSoundEntry;
 					GetActiveSoundEntryFromIDs(Message.ActiveSoundPlayOrder, ActiveSoundEntry);
@@ -1679,11 +1350,7 @@ namespace UE::Audio::Insights
 				break;
 
 			case ESoundParamMessageType::HPFFreq:
-#if WITH_EDITOR
 				CacheManager.IterateOverRange<FSoundHPFFreqMessage>(SoundMessageNames::HPFFreqParam, StartTime, EndTime, [this, TimeMarker](const FSoundHPFFreqMessage& Message)
-#else
-				IterateOverRange<FSoundHPFFreqMessage>(SessionCachedMessages->HPFFreqCachedMessages, StartTime, EndTime, [this, TimeMarker](const FSoundHPFFreqMessage& Message)
-#endif // WITH_EDITOR
 				{
 					TSharedPtr<FSoundDashboardEntry> ActiveSoundEntry;
 					GetActiveSoundEntryFromIDs(Message.ActiveSoundPlayOrder, ActiveSoundEntry);
@@ -1705,11 +1372,7 @@ namespace UE::Audio::Insights
 				break;
 
 			case ESoundParamMessageType::Envelope:
-#if WITH_EDITOR
 				CacheManager.IterateOverRange<FSoundEnvelopeMessage>(SoundMessageNames::EnvelopeParam, StartTime, EndTime, [this, TimeMarker](const FSoundEnvelopeMessage& Message)
-#else
-				IterateOverRange<FSoundEnvelopeMessage>(SessionCachedMessages->AmplitudeCachedMessages, StartTime, EndTime, [this, TimeMarker](const FSoundEnvelopeMessage& Message)
-#endif // WITH_EDITOR
 				{
 					TSharedPtr<FSoundDashboardEntry> ActiveSoundEntry;
 					GetActiveSoundEntryFromIDs(Message.ActiveSoundPlayOrder, ActiveSoundEntry);
@@ -1731,11 +1394,7 @@ namespace UE::Audio::Insights
 				break;
 
 			case ESoundParamMessageType::Volume:
-#if WITH_EDITOR
 				CacheManager.IterateOverRange<FSoundVolumeMessage>(SoundMessageNames::VolumeParam, StartTime, EndTime, [this, TimeMarker](const FSoundVolumeMessage& Message)
-#else
-				IterateOverRange<FSoundVolumeMessage>(SessionCachedMessages->VolumeCachedMessages, StartTime, EndTime, [this, TimeMarker](const FSoundVolumeMessage& Message)
-#endif // WITH_EDITOR
 				{
 					TSharedPtr<FSoundDashboardEntry> ActiveSoundEntry;
 					GetActiveSoundEntryFromIDs(Message.ActiveSoundPlayOrder, ActiveSoundEntry);
@@ -1757,11 +1416,7 @@ namespace UE::Audio::Insights
 				break;
 
 			case ESoundParamMessageType::Pitch:
-#if WITH_EDITOR
 				CacheManager.IterateOverRange<FSoundPitchMessage>(SoundMessageNames::PitchParam, StartTime, EndTime, [this, TimeMarker](const FSoundPitchMessage& Message)
-#else
-				IterateOverRange<FSoundPitchMessage>(SessionCachedMessages->PitchCachedMessages, StartTime, EndTime, [this, TimeMarker](const FSoundPitchMessage& Message)
-#endif // WITH_EDITOR
 				{
 					TSharedPtr<FSoundDashboardEntry> ActiveSoundEntry;
 					GetActiveSoundEntryFromIDs(Message.ActiveSoundPlayOrder, ActiveSoundEntry);
@@ -1783,11 +1438,7 @@ namespace UE::Audio::Insights
 				break;
 
 			case ESoundParamMessageType::RelativeRenderCost:
-#if WITH_EDITOR
 				CacheManager.IterateOverRange<FSoundRelativeRenderCostMessage>(SoundMessageNames::RelativeRenderCostParam, StartTime, EndTime, [this, TimeMarker](const FSoundRelativeRenderCostMessage& Message)
-#else
-				IterateOverRange<FSoundRelativeRenderCostMessage>(SessionCachedMessages->RelativeRenderCostCachedMessages, StartTime, EndTime, [this, TimeMarker](const FSoundRelativeRenderCostMessage& Message)
-#endif // WITH_EDITOR
 				{
 					TSharedPtr<FSoundDashboardEntry> ActiveSoundEntry;
 					GetActiveSoundEntryFromIDs(Message.ActiveSoundPlayOrder, ActiveSoundEntry);
diff --git a/Engine/Plugins/AudioInsights/Source/AudioInsights/Private/Providers/SoundTraceProvider.h b/Engine/Plugins/AudioInsights/Source/AudioInsights/Private/Providers/SoundTraceProvider.h
--- a/Engine/Plugins/AudioInsights/Source/AudioInsights/Private/Providers/SoundTraceProvider.h
+++ b/Engine/Plugins/AudioInsights/Source/AudioInsights/Private/Providers/SoundTraceProvider.h
@@ -28,10 +28,6 @@ namespace UE::Audio::Insights
 
 		virtual UE::Trace::IAnalyzer* ConstructAnalyzer(TraceServices::IAnalysisSession& InSession) override;
 
-#if !WITH_EDITOR
-		virtual void InitSessionCachedMessages(TraceServices::IAnalysisSession& InSession) override;
-#endif // !WITH_EDITOR
-
 		virtual bool ProcessMessages() override;
 		virtual bool ProcessManuallyUpdatedEntries() override;
 
@@ -52,19 +48,11 @@ namespace UE::Audio::Insights
 		virtual void OnTimingViewTimeMarkerChanged(double TimeMarker) override;
 		virtual void OnTimeControlMethodReset() override;
 
-#if WITH_EDITOR
 		void CollectParamsForTimestamp(const FAudioInsightsCacheManager& CacheManager, const double TimeMarker);
-#else
-		void CollectParamsForTimestamp(const double TimeMarker);
-#endif // WITH_EDITOR
 
 		void CacheIsPlottingFlagRecursive(const TSharedPtr<IDashboardDataTreeViewEntry>& Entry);
 		void ResetDataBuffersRecursive(const TSharedPtr<IDashboardDataTreeViewEntry>& Entry);
 
-#if !WITH_EDITOR
-		TUniquePtr<FSoundSessionCachedMessages> SessionCachedMessages;
-#endif // !WITH_EDITOR
-
 		FSoundMessages TraceMessages;
 
 		struct SoundEntryKeys
diff --git a/Engine/Plugins/AudioInsights/Source/AudioInsights/Private/Views/AudioEventLogDashboardViewFactory.cpp b/Engine/Plugins/AudioInsights/Source/AudioInsights/Private/Views/AudioEventLogDashboardViewFactory.cpp
--- a/Engine/Plugins/AudioInsights/Source/AudioInsights/Private/Views/AudioEventLogDashboardViewFactory.cpp
+++ b/Engine/Plugins/AudioInsights/Source/AudioInsights/Private/Views/AudioEventLogDashboardViewFactory.cpp
@@ -42,7 +42,6 @@ namespace UE::Audio::Insights
 #endif // WITH_EDITOR
 		}
 
-#if WITH_EDITOR
 		bool TimestampIsInChunkMarkedToDelete(const double Timestamp)
 		{
 			FAudioInsightsCacheManager& CacheManager = IAudioInsightsModule::GetChecked().GetCacheManager();
@@ -56,7 +55,6 @@ namespace UE::Audio::Insights
 			const TOptional<uint32> NumChunksFromStart = CacheManager.TryGetNumChunksFromStartForTimestamp(Timestamp);
 			return NumChunksFromStart.IsSet() && NumChunksFromStart.GetValue() == 0u;
 		}
-#endif // WITH_EDITOR
 
 		const FString CategorySoundActivity = LOCTEXT("EventLogCategory_SoundActivity", "Sound Activity").ToString();
 		const FString CategoryVirtualization = LOCTEXT("EventLogCategory_Virtualization", "Virtualization").ToString();
@@ -180,32 +178,32 @@ namespace UE::Audio::Insights
 		FAudioEventLogEditorCommands::Register();
 		BindCommands();
 
+		FAudioInsightsCacheManager& CacheManager = IAudioInsightsModule::GetChecked().GetCacheManager();
+		CacheManager.OnChunkOverwritten.AddRaw(this, &FAudioEventLogDashboardViewFactory::OnCacheChunkOverwritten);
+
 #if WITH_EDITOR
 		FAudioEventLogSettings::OnReadSettings.AddRaw(this, &FAudioEventLogDashboardViewFactory::OnReadEditorSettings);
 		FAudioEventLogSettings::OnWriteSettings.AddRaw(this, &FAudioEventLogDashboardViewFactory::OnWriteEditorSettings);
-
-		FAudioInsightsCacheManager& CacheManager = IAudioInsightsModule::GetChecked().GetCacheManager();
-		CacheManager.OnChunkOverwritten.AddRaw(this, &FAudioEventLogDashboardViewFactory::OnCacheChunkOverwritten);
 #endif // WITH_EDITOR
 	}
 
 	FAudioEventLogDashboardViewFactory::~FAudioEventLogDashboardViewFactory()
 	{
-		if (FModuleManager::Get().IsModuleLoaded("AudioInsights") && IModularFeatures::Get().IsModularFeatureAvailable(TraceServices::ModuleFeatureName))
+		if (FModuleManager::Get().IsModuleLoaded("AudioInsights"))
 		{
-			FTraceModule& TraceModule = static_cast<FTraceModule&>(FAudioInsightsModule::GetChecked().GetTraceModule());
-			TraceModule.OnAnalysisStarting.RemoveAll(this);
+			if (IModularFeatures::Get().IsModularFeatureAvailable(TraceServices::ModuleFeatureName))
+			{
+				FTraceModule& TraceModule = static_cast<FTraceModule&>(FAudioInsightsModule::GetChecked().GetTraceModule());
+				TraceModule.OnAnalysisStarting.RemoveAll(this);
+			}
+			
+			FAudioInsightsCacheManager& CacheManager = IAudioInsightsModule::GetChecked().GetCacheManager();
+			CacheManager.OnChunkOverwritten.RemoveAll(this);
 		}
 
 #if WITH_EDITOR
 		FAudioEventLogSettings::OnReadSettings.RemoveAll(this);
 		FAudioEventLogSettings::OnWriteSettings.RemoveAll(this);
-
-		if (FModuleManager::Get().IsModuleLoaded("AudioInsights"))
-		{
-			FAudioInsightsCacheManager& CacheManager = IAudioInsightsModule::GetChecked().GetCacheManager();
-			CacheManager.OnChunkOverwritten.RemoveAll(this);
-		}
 #endif // WITH_EDITOR
 
 		FAudioEventLogEditorCommands::Unregister();
@@ -222,6 +220,7 @@ namespace UE::Audio::Insights
 #if WITH_EDITOR
 		CommandList->MapAction(Commands.GetBrowseCommand(), FExecuteAction::CreateLambda([this]() { BrowseToAsset(); }));
 		CommandList->MapAction(Commands.GetEditCommand(), FExecuteAction::CreateLambda([this]() { OpenAsset(); }));
+#endif // WITH_EDITOR
 
 		CommandList->MapAction(
 			Commands.GetAutoStopCachingWhenLastInCacheCommand(),
@@ -240,7 +239,6 @@ namespace UE::Audio::Insights
 			FExecuteAction::CreateLambda([this]() { AutoStopCachingMode = EAutoStopCachingMode::Never; }),
 			FCanExecuteAction(),
 			FIsActionChecked::CreateLambda([this]() { return AutoStopCachingMode == EAutoStopCachingMode::Never; }));
-#endif // WITH_EDITOR
 	}
 
 	void FAudioEventLogDashboardViewFactory::ResetInspectTimestamp()
@@ -380,7 +378,6 @@ namespace UE::Audio::Insights
 			return false;
 		});
 
-#if WITH_EDITOR
 		if (Reason == EProcessReason::EntriesUpdated && FilteredEntriesListView.IsValid())
 		{
 			// Reset bCacheStatusIsDirty flag to push icon update to the UI
@@ -403,7 +400,6 @@ namespace UE::Audio::Insights
 				}
 			}
 		}
-#endif // WITH_EDITOR
 	}
 
 	const TMap<FName, FTraceTableDashboardViewFactory::FColumnData>& FAudioEventLogDashboardViewFactory::GetColumns() const
@@ -420,7 +416,6 @@ namespace UE::Audio::Insights
 
 			return TMap<FName, FTraceTableDashboardViewFactory::FColumnData>
 			{
-#if WITH_EDITOR
 				{
 					CacheStatusColumnID,
 					{
@@ -466,7 +461,6 @@ namespace UE::Audio::Insights
 						false /* bShowDisplayName */
 					}
 				},
-#endif // WITH_EDITOR
 				{
 					MessageIDColumnID,
 					{
@@ -662,13 +656,9 @@ namespace UE::Audio::Insights
 	{
 		if (!SettingsAreaWidget.IsValid())
 		{
-#if WITH_EDITOR
 			constexpr FLinearColor StatTitleColor(1.0f, 1.0f, 1.0f, 0.5f);
-#endif // WITH_EDITOR
-
 			SAssignNew(SettingsAreaWidget, SHorizontalBox)
 			
-#if WITH_EDITOR
 			// Cache paused icon
 			+ SHorizontalBox::Slot()
 			.VAlign(VAlign_Center)
@@ -751,7 +741,6 @@ namespace UE::Audio::Insights
 										, FText::AsNumber(CacheDuration, FSlateStyle::Get().GetShortTimeFormat()));
 				})
 			]
-#endif // WITH_EDITOR
 
 			// Clear log Button
 			+ SHorizontalBox::Slot()
@@ -826,7 +815,6 @@ namespace UE::Audio::Insights
 				FNewMenuDelegate::CreateRaw(this, &FAudioEventLogDashboardViewFactory::BuildVisibleColumnsMenuContent)
 			);
 
-#if WITH_EDITOR
 			// Auto-stop monitoring
 			MenuBuilder.AddSubMenu
 			(
@@ -834,7 +822,6 @@ namespace UE::Audio::Insights
 				LOCTEXT("AudioEventLog_Settings_AutoStopCachingTooltip", "Choose if/how the event log automatically stops caching when inspecting an event."),
 				FNewMenuDelegate::CreateRaw(this, &FAudioEventLogDashboardViewFactory::BuildAutoStopCachingMenuContent)
 			);
-#endif // WITH_EDITOR
 		}
 
 		MenuBuilder.EndSection();
@@ -1163,18 +1150,6 @@ namespace UE::Audio::Insights
 		}
 	}
 
-	void FAudioEventLogDashboardViewFactory::OnCacheChunkOverwritten(const double NewCacheStartTimestamp)
-	{
-		if (IsInGameThread())
-		{
-			UpdateCacheMessageProcessMethod();
-		}
-		else
-		{
-			ExecuteOnGameThread(TEXT("FAudioEventLogDashboardViewFactory::UpdateCacheMessageProcessMethod"), [this] { UpdateCacheMessageProcessMethod(); });
-		}
-	}
-
 	void FAudioEventLogDashboardViewFactory::RemoveDeletedCustomEvents(const TMap<FString, FAudioEventLogCustomEvents>& CustomEventsFromSettings)
 	{
 		if (!FilterBar.IsValid())
@@ -1271,11 +1246,22 @@ namespace UE::Audio::Insights
 	}
 #endif // WITH_EDITOR
 
+	void FAudioEventLogDashboardViewFactory::OnCacheChunkOverwritten(const double NewCacheStartTimestamp)
+	{
+		if (IsInGameThread())
+		{
+			UpdateCacheMessageProcessMethod();
+		}
+		else
+		{
+			ExecuteOnGameThread(TEXT("FAudioEventLogDashboardViewFactory::UpdateCacheMessageProcessMethod"), [this] { UpdateCacheMessageProcessMethod(); });
+		}
+	}
+
 	void FAudioEventLogDashboardViewFactory::OnTimingViewTimeMarkerChanged(double InTimeMarker)
 	{
 		using namespace AudioEventLogPrivate;
 
-#if WITH_EDITOR
 		FAudioInsightsModule& AudioInsightsModule = FAudioInsightsModule::GetChecked();
 
 		if (AutoStopCachingMode == EAutoStopCachingMode::OnInspect)
@@ -1299,7 +1285,6 @@ namespace UE::Audio::Insights
 		{
 			AudioInsightsModule.GetTimingViewExtender().StopProcessingNewMessages();
 		}
-#endif // WITH_EDITOR
 
 		bAutoScroll = false;
 
@@ -1350,13 +1335,22 @@ namespace UE::Audio::Insights
 	{
 		using namespace AudioEventLogPrivate;
 
-#if WITH_EDITOR
 		if (AutoStopCachingMode != EAutoStopCachingMode::WhenInLastChunk || !FilteredEntriesListView.IsValid())
 		{
 			return;
 		}
 
 		FAudioInsightsModule& AudioInsightsModule = FAudioInsightsModule::GetChecked();
+
+#if !WITH_EDITOR
+		const TSharedPtr<const FAudioInsightsComponent> AudioInsightsComponent = FAudioInsightsModule::GetChecked().GetAudioInsightsComponent();
+		if (AudioInsightsComponent.IsValid() && !AudioInsightsComponent->GetIsLiveSession())
+		{
+			// Don't update the caching and processing status if the cache isn't limited (i.e. loading in a .urace file)
+			return;
+		}
+#endif // !WITH_EDITOR
+
 		if (AudioInsightsModule.GetTimingViewExtender().GetMessageCacheAndProcessingStatus() != ECacheAndProcess::CacheLatestNoProcess)
 		{
 			return;
@@ -1380,7 +1374,6 @@ namespace UE::Audio::Insights
 				FocusedItem = SelectedItems[0];
 			}
 		}
-#endif // WITH_EDITOR
 	}
 
 } // namespace UE::Audio::Insights
diff --git a/Engine/Plugins/AudioInsights/Source/AudioInsights/Public/Messages/AudioEventLogTraceMessages.h b/Engine/Plugins/AudioInsights/Source/AudioInsights/Public/Messages/AudioEventLogTraceMessages.h
--- a/Engine/Plugins/AudioInsights/Source/AudioInsights/Public/Messages/AudioEventLogTraceMessages.h
+++ b/Engine/Plugins/AudioInsights/Source/AudioInsights/Public/Messages/AudioEventLogTraceMessages.h
@@ -8,11 +8,6 @@
 #include "Trace/Analyzer.h"
 #include "Views/TableDashboardViewFactory.h"
 
-#if !WITH_EDITOR
-#include "Common/PagedArray.h"
-#include "TraceServices/Model/AnalysisSession.h"
-#endif // !WITH_EDITOR
-
 namespace UE::Audio::Insights
 {
 	enum class EAudioEventLogSoundCategory : uint8
@@ -25,13 +20,11 @@ namespace UE::Audio::Insights
 		SoundCueTemplate
 	};
 
-#if WITH_EDITOR
 	enum class EAudioEventCacheState : uint8
 	{
 		Cached,
 		NextToBeDeleted
 	};
-#endif // WITH_EDITOR
 
 	struct FAudioEventLogMessage : public IAudioCachedMessage
 	{
@@ -74,11 +67,8 @@ namespace UE::Audio::Insights
 		FString CategoryName;
 
 		EAudioEventLogSoundCategory Category = EAudioEventLogSoundCategory::None;
-
-#if WITH_EDITOR
 		EAudioEventCacheState CachedState = EAudioEventCacheState::Cached;
 		bool bCacheStatusIsDirty = false;
-#endif // WITH_EDITOR
 	};
 
 	class FAudioEventLogMessages
diff --git a/Engine/Plugins/AudioInsights/Source/AudioInsights/Public/Messages/SoundTraceMessages.h b/Engine/Plugins/AudioInsights/Source/AudioInsights/Public/Messages/SoundTraceMessages.h
--- a/Engine/Plugins/AudioInsights/Source/AudioInsights/Public/Messages/SoundTraceMessages.h
+++ b/Engine/Plugins/AudioInsights/Source/AudioInsights/Public/Messages/SoundTraceMessages.h
@@ -207,46 +207,6 @@ namespace UE::Audio::Insights
 		TAnalyzerMessageQueue<FSoundStopMessage> StopMessages;
 	};
 
-#if !WITH_EDITOR
-	struct FSoundSessionCachedMessages
-	{
-		FSoundSessionCachedMessages(TraceServices::IAnalysisSession& InSession)
-			: StartCachedMessages(InSession.GetLinearAllocator(), 16384)
-			, SoundWaveStartCachedMessages(InSession.GetLinearAllocator(), 16384)
-			, SoundIsAlivePingCachedMessages(InSession.GetLinearAllocator(), 16384)
-			, SoundWaveIsAlivePingCachedMessages(InSession.GetLinearAllocator(), 16384)
-			, PriorityCachedMessages(InSession.GetLinearAllocator(), 16384)
-			, DistanceCachedMessages(InSession.GetLinearAllocator(), 16384)
-			, DistanceAttenuationCachedMessages(InSession.GetLinearAllocator(), 16384)
-			, HPFFreqCachedMessages(InSession.GetLinearAllocator(), 16384)
-			, LPFFreqCachedMessages(InSession.GetLinearAllocator(), 16384)
-			, AmplitudeCachedMessages(InSession.GetLinearAllocator(), 16384)
-			, VolumeCachedMessages(InSession.GetLinearAllocator(), 16384)
-			, PitchCachedMessages(InSession.GetLinearAllocator(), 16384)
-			, RelativeRenderCostCachedMessages(InSession.GetLinearAllocator(), 16384)
-			, StopCachedMessages(InSession.GetLinearAllocator(), 4096)
-		{
-
-		}
-
-		TraceServices::TPagedArray<FSoundStartMessage> StartCachedMessages;
-		TraceServices::TPagedArray<FSoundWaveStartMessage> SoundWaveStartCachedMessages;
-		TraceServices::TPagedArray<FSoundIsAlivePingMessage> SoundIsAlivePingCachedMessages;
-		TraceServices::TPagedArray<FSoundWaveIsAlivePingMessage> SoundWaveIsAlivePingCachedMessages;
-		TraceServices::TPagedArray<FSoundPriorityMessage> PriorityCachedMessages;
-		TraceServices::TPagedArray<FSoundDistanceMessage> DistanceCachedMessages;
-		TraceServices::TPagedArray<FSoundDistanceAttenuationMessage> DistanceAttenuationCachedMessages;
-		TraceServices::TPagedArray<FSoundHPFFreqMessage> HPFFreqCachedMessages;
-		TraceServices::TPagedArray<FSoundLPFFreqMessage> LPFFreqCachedMessages;
-		TraceServices::TPagedArray<FSoundEnvelopeMessage> AmplitudeCachedMessages;
-		TraceServices::TPagedArray<FSoundVolumeMessage> VolumeCachedMessages;
-		TraceServices::TPagedArray<FSoundPitchMessage> PitchCachedMessages;
-		TraceServices::TPagedArray<FSoundRelativeRenderCostMessage> RelativeRenderCostCachedMessages;
-		TraceServices::TPagedArray<FSoundStopMessage> StopCachedMessages;
-	};
-#endif // !WITH_EDITOR
-
-
 	// Dashboard entry
 	using FDataPoint = TPair<double, float>; // (Timestamp, Value)
 
diff --git a/Engine/Plugins/AudioInsights/Source/AudioInsights/Public/Views/AudioEventLogDashboardViewFactory.h b/Engine/Plugins/AudioInsights/Source/AudioInsights/Public/Views/AudioEventLogDashboardViewFactory.h
--- a/Engine/Plugins/AudioInsights/Source/AudioInsights/Public/Views/AudioEventLogDashboardViewFactory.h
+++ b/Engine/Plugins/AudioInsights/Source/AudioInsights/Public/Views/AudioEventLogDashboardViewFactory.h
@@ -109,14 +109,13 @@ namespace UE::Audio::Insights
 		void OnReadEditorSettings(const FAudioEventLogSettings& InSettings);
 		void OnWriteEditorSettings(FAudioEventLogSettings& OutSettings);
 
-		void OnCacheChunkOverwritten(const double NewCacheStartTimestamp);
-
 		void RemoveDeletedCustomEvents(const TMap<FString, FAudioEventLogCustomEvents>& CustomEventsFromSettings);
 		void AddNewCustomEvents(const TMap<FString, FAudioEventLogCustomEvents>& CustomEventsFromSettings);
 
 		void RefreshFilterBarFromSettings(const TSet<FString>& EventFilters);
 #endif // WITH_EDITOR
 
+		void OnCacheChunkOverwritten(const double NewCacheStartTimestamp);
 		void OnTimingViewTimeMarkerChanged(double InTimeMarker);
 
 		void UpdateCacheMessageProcessMethod();
@@ -151,7 +150,6 @@ namespace UE::Audio::Insights
 		double BeginTimestamp = 0.0;
 		bool bAutoScroll = true;
 
-#if WITH_EDITOR
 		enum class EAutoStopCachingMode
 		{
 			WhenInLastChunk,
@@ -160,7 +158,6 @@ namespace UE::Audio::Insights
 		};
 
 		EAutoStopCachingMode AutoStopCachingMode = EAutoStopCachingMode::WhenInLastChunk;
-#endif // WITH_EDITOR
 	};
 } // namespace UE::Audio::Insights
 
