AudioInsights

* Fix potential nullptr dereference inside FSoundPlotsWidgetView::ProcessPlotData
* BugHawk suggestion:
"error: memory: Dereferencing the result of TMap::Find without a nullptr check can crash when the key is missing. TMap::Find returns a pointer to the value or nullptr if not found; using * immediately on the result is unsafe."

diff --git a/Engine/Plugins/AudioInsights/Source/AudioInsights/Private/Views/SoundPlotsWidgetView.cpp b/Engine/Plugins/AudioInsights/Source/AudioInsights/Private/Views/SoundPlotsWidgetView.cpp
--- a/Engine/Plugins/AudioInsights/Source/AudioInsights/Private/Views/SoundPlotsWidgetView.cpp
+++ b/Engine/Plugins/AudioInsights/Source/AudioInsights/Private/Views/SoundPlotsWidgetView.cpp
@@ -135,12 +135,13 @@ namespace UE::Audio::Insights
 		const double PlotDrawMaxLimitTimestamp = CurrentTimestamp + FAudioInsightsTimingViewExtender::PlottingMarginSeconds /* extra grace time to avoid curve cuts being displayed */;
 #endif // !WITH_EDITOR
 
-		const TSharedPtr<FPointDataPerCurveMap> CurveDataMapPtr = *PlotWidgetCurveIdToPointDataMapPerColumn.Find(SelectedPlotColumnName);
 
-		if (CurveDataMapPtr.IsValid())
+		const TSharedPtr<FPointDataPerCurveMap>* const FoundPointDataPerCurveElement = PlotWidgetCurveIdToPointDataMapPerColumn.Find(SelectedPlotColumnName);
+
+		if (FoundPointDataPerCurveElement && FoundPointDataPerCurveElement->IsValid())
 		{
 			// Remove points that are older than max history limit from the most recent timestamp
-			for (auto& [CurveId, CurvePoints] : *CurveDataMapPtr)
+			for (auto& [CurveId, CurvePoints] : **FoundPointDataPerCurveElement)
 			{
 				const int32 FoundMinIndex = CurvePoints.IndexOfByPredicate([&PlotDrawMinLimitTimestamp](const FDataPoint& InDataPoint)
 				{
