[AudioInsights]

* Fix virtual loop scrolling and selection on/off
* Fixes: https://jira.it.epicgames.com/browse/UE-348728
* Fixes: https://jira.it.epicgames.com/browse/UE-315988

diff --git a/Engine/Plugins/AudioInsights/Source/AudioInsights/Private/Views/VirtualLoopDashboardViewFactory.cpp b/Engine/Plugins/AudioInsights/Source/AudioInsights/Private/Views/VirtualLoopDashboardViewFactory.cpp
--- a/Engine/Plugins/AudioInsights/Source/AudioInsights/Private/Views/VirtualLoopDashboardViewFactory.cpp
+++ b/Engine/Plugins/AudioInsights/Source/AudioInsights/Private/Views/VirtualLoopDashboardViewFactory.cpp
@@ -50,6 +50,21 @@ namespace UE::Audio::Insights
 		return LOCTEXT("AudioDashboard_VirtualLoops_DisplayName", "Virtual Loops");
 	}
 
+	TSharedRef<SWidget> FVirtualLoopDashboardViewFactory::MakeWidget(TSharedRef<SDockTab> InOwnerTab, const FSpawnTabArgs& InSpawnTabArgs)
+	{
+		if (!DashboardWidget.IsValid())
+		{
+			DashboardWidget = FTraceTableDashboardViewFactory::MakeWidget(InOwnerTab, InSpawnTabArgs);
+
+			if (FilteredEntriesListView.IsValid())
+			{
+				FilteredEntriesListView->SetSelectionMode(ESelectionMode::Multi);
+			}
+		}
+
+		return DashboardWidget->AsShared();
+	}
+
 	void FVirtualLoopDashboardViewFactory::ProcessEntries(FTraceTableDashboardViewFactory::EProcessReason Reason)
 	{
 		const FString FilterString = GetSearchFilterText().ToString();
diff --git a/Engine/Plugins/AudioInsights/Source/AudioInsights/Public/Views/VirtualLoopDashboardViewFactory.h b/Engine/Plugins/AudioInsights/Source/AudioInsights/Public/Views/VirtualLoopDashboardViewFactory.h
--- a/Engine/Plugins/AudioInsights/Source/AudioInsights/Public/Views/VirtualLoopDashboardViewFactory.h
+++ b/Engine/Plugins/AudioInsights/Source/AudioInsights/Public/Views/VirtualLoopDashboardViewFactory.h
@@ -34,6 +34,9 @@ namespace UE::Audio::Insights
 		UE_API virtual bool IsDebugDrawEnabled() const override;
 		UE_API virtual void DebugDraw(float InElapsed, const TArray<TSharedPtr<IDashboardDataViewEntry>>& InSelectedItems, ::Audio::FDeviceId InAudioDeviceId) const;
 #endif // WITH_EDITOR
+
+	private:
+		virtual TSharedRef<SWidget> MakeWidget(TSharedRef<SDockTab> InOwnerTab, const FSpawnTabArgs& InSpawnTabArgs) override;
 	};
 } // namespace UE::Audio::Insights
 
