[Audio Insights] Do not pass ownership of editor dashboard view factories to Tab manager OnTabSpawn delegate

diff --git a/Engine/Plugins/AudioInsights/Source/AudioInsightsEditor/Private/AudioInsightsEditorDashboardFactory.cpp b/Engine/Plugins/AudioInsights/Source/AudioInsightsEditor/Private/AudioInsightsEditorDashboardFactory.cpp
--- a/Engine/Plugins/AudioInsights/Source/AudioInsightsEditor/Private/AudioInsightsEditorDashboardFactory.cpp
+++ b/Engine/Plugins/AudioInsights/Source/AudioInsightsEditor/Private/AudioInsightsEditorDashboardFactory.cpp
@@ -1110,15 +1110,16 @@ namespace UE::Audio::Insights
 		for (const auto& KVP : DashboardViewFactories)
 		{
 			const FName& FactoryName = KVP.Key;
-			const TSharedPtr<IDashboardViewFactory>& Factory = KVP.Value;
-
-			DashboardTabManager->RegisterTabSpawner(FactoryName, FOnSpawnTab::CreateLambda([this, Factory](const FSpawnTabArgs& Args)
+			const TSharedPtr<IDashboardViewFactory> Factory = KVP.Value;
+			DashboardTabManager->RegisterTabSpawner(FactoryName, FOnSpawnTab::CreateLambda([this, FactoryWeak = Factory.ToWeakPtr()](const FSpawnTabArgs& Args)
 			{
+				const TSharedPtr<IDashboardViewFactory> FactoryShared = FactoryWeak.Pin();
+
 				TSharedRef<SDockTab> DockTab = SNew(SDockTab)
 					.Clipping(EWidgetClipping::ClipToBounds)
-					.Label(Factory->GetDisplayName());
+					.Label(FactoryShared.IsValid() ? FactoryShared->GetDisplayName() : FText::GetEmpty());
 
-				TSharedRef<SWidget> DashboardView = Factory->MakeWidget(DockTab, Args);
+				TSharedRef<SWidget> DashboardView = FactoryShared.IsValid() ? FactoryShared->MakeWidget(DockTab, Args) : SNullWidget::NullWidget;
 				DockTab->SetContent(DashboardView);
 
 				return DockTab;
