[Audio Insights] Remove Audio Insights's dependency on Audio Insights Runtime.

diff --git a/Engine/Plugins/AudioInsights/AudioInsights.uplugin b/Engine/Plugins/AudioInsights/AudioInsights.uplugin
--- a/Engine/Plugins/AudioInsights/AudioInsights.uplugin
+++ b/Engine/Plugins/AudioInsights/AudioInsights.uplugin
@@ -32,10 +32,6 @@
 		}
 	],
 	"Plugins": [
-		{
-			"Name": "AudioInsightsRuntime",
-			"Enabled": true
-		},
 		{
 			"Name": "AudioWidgets",
 			"Enabled": true
