[Audio Insights] Fix for softlock when opening Audio Insights in standalone. If we load in with an invalid view range, PaintGridLines will enter an infinite loop. Early out in this case.

diff --git a/Engine/Plugins/AudioInsights/Source/AudioInsights/Private/Views/SAudioCurveView.cpp b/Engine/Plugins/AudioInsights/Source/AudioInsights/Private/Views/SAudioCurveView.cpp
--- a/Engine/Plugins/AudioInsights/Source/AudioInsights/Private/Views/SAudioCurveView.cpp
+++ b/Engine/Plugins/AudioInsights/Source/AudioInsights/Private/Views/SAudioCurveView.cpp
@@ -582,6 +582,12 @@ int32 SAudioCurveView::PaintCurves(const FGeometry& AllottedGeometry, const FSla
 	}
 
 	const SSimpleTimeSlider::FScrubRangeToScreen RangeToScreen(ViewRange.Get(), AllottedGeometry.GetLocalSize());
+
+	if (RangeToScreen.ViewInput.IsEmpty())
+	{
+		return LayerId;
+	}
+
 	LayerId = PaintGridLines(AllottedGeometry, MyCullingRect, OutDrawElements, LayerId, InWidgetStyle, bParentEnabled, RangeToScreen);
 	LayerId = PaintCrosshair(AllottedGeometry, OutDrawElements, LayerId, RangeToScreen);
 	
