[Audio Insights] BugHawk suggestion - null check before proceeding to UpdateYDataRange in SAudioCurveView::SetCurvesPointData

diff --git a/Engine/Plugins/AudioInsights/Source/AudioInsights/Private/Views/SAudioCurveView.cpp b/Engine/Plugins/AudioInsights/Source/AudioInsights/Private/Views/SAudioCurveView.cpp
--- a/Engine/Plugins/AudioInsights/Source/AudioInsights/Private/Views/SAudioCurveView.cpp
+++ b/Engine/Plugins/AudioInsights/Source/AudioInsights/Private/Views/SAudioCurveView.cpp
@@ -338,7 +338,7 @@ void SAudioCurveView::SetCurvesPointData(TSharedPtr<TMap<uint64, TArray<FCurvePo
 {
 	PointDataPerCurve = InPointDataPerCurve;
 
-	if (AutoRangeYAxis.Get())
+	if (PointDataPerCurve.IsValid() && AutoRangeYAxis.Get())
 	{
 		UpdateYDataRange();
 	}
