[AudioInsights]

* Fix linker errors when calling AudioInsightsCacheManager functions outside the  AudioInsights Module

* The AudioInsightsCacheManager templated functions e.g. IterateTo() internally call:
FAudioInsightsCache::Capacity()
FAudioInsightsCache::GetNextIndex()
FAudioInsightsCache::GetPreviousIndex()

* When modules outside of AudioInsights e.g. AudioModulationInsights, call these templated functions we encounter linker errors accessing the above functions (see Jira for details)

* Prefixing these with UE_API fixes this
* Intentionally leaving ResetCache() and ReserveNewChunks() which are not called

* Fixes: https://jira.it.epicgames.com/browse/UE-351664
* Unblocks: https://jira.it.epicgames.com/browse/UE-314913

diff --git a/Engine/Plugins/AudioInsights/Source/AudioInsights/Public/Cache/AudioInsightsCacheManager.h b/Engine/Plugins/AudioInsights/Source/AudioInsights/Public/Cache/AudioInsightsCacheManager.h
--- a/Engine/Plugins/AudioInsights/Source/AudioInsights/Public/Cache/AudioInsightsCacheManager.h
+++ b/Engine/Plugins/AudioInsights/Source/AudioInsights/Public/Cache/AudioInsightsCacheManager.h
@@ -45,9 +45,9 @@ namespace UE::Audio::Insights
 		void ResetCache(TraceServices::FSlabAllocator& Allocator);
 		void ReserveNewChunks(TraceServices::FSlabAllocator& Allocator, const uint32 NumToReserve);
 
-		uint32 Capacity() const;
-		uint32 GetNextIndex(uint32 ReadIndex) const;
-		uint32 GetPreviousIndex(uint32 ReadIndex) const;
+		UE_API uint32 Capacity() const;
+		UE_API uint32 GetNextIndex(uint32 ReadIndex) const;
+		UE_API uint32 GetPreviousIndex(uint32 ReadIndex) const;
 
 	private:
 #if !WITH_EDITOR
