[AudioInsights]

* Support FindOrAddDeviceEntry() in TDeviceDataMapTraceProvider's protected API
* Allows cleaner subclass code, where we simply want to find or add without the obligation of executing an Update Ftor
* Clearer intention for callers,
* Doesn't break existing API, as subclasses using UpdateDeviceEntry() still behaves the same
* Additionally, we can still mutate the returned entry similar to UpdateDeviceEntry, but with clearer intent

Example use case:

1. Before (not clear that UpdateDeviceEntry is finding or adding under the hood)
auto BumpEntryFunc = [this](const FControlBusMessageBase& Msg)
{
TSharedPtr<FControlBusDashboardEntry>* ToReturn = nullptr;
UpdateDeviceEntry(Msg.DeviceId, Msg.ControlBusId, [&ToReturn, &Msg](TSharedPtr<FControlBusDashboardEntry>& Entry)
{
if (!Entry.IsValid())
{
Entry = MakeShared<FControlBusDashboardEntry>();
Entry->DeviceId = Msg.DeviceId;
Entry->ControlBusId = Msg.ControlBusId;
}
ToReturn = &Entry;
});

return ToReturn;
};

2. After:
TSharedPtr<FControlBusDashboardEntry>& FoundEntry = FindOrAddDeviceEntry(InControlBusMessage.DeviceId, InControlBusMessage.ControlBusId);
if (!FoundEntry.IsValid())
{
FoundEntry = MakeShared<FControlBusDashboardEntry>();
FoundEntry->DeviceId = InControlBusMessage.DeviceId;
FoundEntry->ControlBusId = InControlBusMessage.ControlBusId;
}
return &FoundEntry;

* Relates to: https://jira.it.epicgames.com/browse/UE-314913

diff --git a/Engine/Plugins/AudioInsights/Source/AudioInsights/Public/AudioInsightsTraceProviderBase.h b/Engine/Plugins/AudioInsights/Source/AudioInsights/Public/AudioInsightsTraceProviderBase.h
--- a/Engine/Plugins/AudioInsights/Source/AudioInsights/Public/AudioInsightsTraceProviderBase.h
+++ b/Engine/Plugins/AudioInsights/Source/AudioInsights/Public/AudioInsightsTraceProviderBase.h
@@ -330,6 +330,13 @@ namespace UE::Audio::Insights
 			return DeviceDataMap.Find(AudioDeviceId);
 		}
 
+		EntryType& FindOrAddDeviceEntry(::Audio::FDeviceId InDeviceId, const EntryKey& InKey)
+		{
+			ensure(InDeviceId != INDEX_NONE);
+			FDeviceData& DeviceData = DeviceDataMap.FindOrAdd(InDeviceId);
+			return DeviceData.FindOrAdd(InKey);
+		}
+
 		bool RemoveDeviceEntry(::Audio::FDeviceId InDeviceId, const EntryKey& InKey)
 		{
 			if (FDeviceData* DeviceData = DeviceDataMap.Find(InDeviceId))
@@ -350,9 +357,7 @@ namespace UE::Audio::Insights
 
 		void UpdateDeviceEntry(::Audio::FDeviceId InDeviceId, const EntryKey& InKey, TFunctionRef<void(EntryType&)> InEntryMutator)
 		{
-			ensure(InDeviceId != INDEX_NONE);
-			FDeviceData& DeviceData = DeviceDataMap.FindOrAdd(InDeviceId);
-			EntryType& Entry = DeviceData.FindOrAdd(InKey);
+			EntryType& Entry = FindOrAddDeviceEntry(InDeviceId, InKey);
 			InEntryMutator(Entry);
 		}
 
