[Audio Insights] Signal Flow part 03: New Signal Flow messages and dashboard entry set up

#Synced-CL 48352470

diff --git a/Engine/Plugins/AudioInsights/Source/AudioInsights/Private/Messages/SignalFlowTraceMessages.cpp b/Engine/Plugins/AudioInsights/Source/AudioInsights/Private/Messages/SignalFlowTraceMessages.cpp
new file mode 100644
--- /dev/null
+++ b/Engine/Plugins/AudioInsights/Source/AudioInsights/Private/Messages/SignalFlowTraceMessages.cpp
@@ -0,0 +1,38 @@
+// Copyright Epic Games, Inc. All Rights Reserved.
+#include "Messages/SignalFlowTraceMessages.h"
+
+namespace UE::Audio::Insights
+{
+	namespace SignalFlowMessageNames
+	{
+		const FName SoundWaveToBusSend = "SoundToBusSend";
+		const FName SoundWaveToSubmixSend = "SoundToSubmixSend";
+		const FName SubmixToSubmixSend = "SubmixToSubmixSend";
+	};
+
+	FSignalFlowMessageBase::FSignalFlowMessageBase(const Trace::IAnalyzer::FOnEventContext& InContext)
+	{
+		const Trace::IAnalyzer::FEventData& EventData = InContext.EventData;
+
+		DeviceId = static_cast<::Audio::FDeviceId>(EventData.GetValue<uint32>("DeviceId"));
+		Timestamp = InContext.EventTime.AsSeconds(EventData.GetValue<uint64>("Timestamp"));
+	}
+
+	FSignalFlowSendMessageBase::FSignalFlowSendMessageBase(const Trace::IAnalyzer::FOnEventContext& InContext)
+		: FSignalFlowMessageBase(InContext)
+	{
+		const Trace::IAnalyzer::FEventData& EventData = InContext.EventData;
+
+		SenderID = EventData.GetValue<uint32>("SenderID");
+		ReceiverID = EventData.GetValue<uint32>("ReceiverID");
+		SendLevel = EventData.GetValue<float>("SendLevel");
+	}
+
+	FSubmixToSubmixSendMessage::FSubmixToSubmixSendMessage(const Trace::IAnalyzer::FOnEventContext& InContext)
+		: FSignalFlowSendMessageBase(InContext)
+	{
+		const Trace::IAnalyzer::FEventData& EventData = InContext.EventData;
+
+		bIsFinalStage = EventData.GetValue<bool>("IsFinalStage");
+	}
+} // namespace UE::Audio::Insights
diff --git a/Engine/Plugins/AudioInsights/Source/AudioInsights/Public/Messages/SignalFlowEntryKey.h b/Engine/Plugins/AudioInsights/Source/AudioInsights/Public/Messages/SignalFlowEntryKey.h
new file mode 100644
--- /dev/null
+++ b/Engine/Plugins/AudioInsights/Source/AudioInsights/Public/Messages/SignalFlowEntryKey.h
@@ -0,0 +1,47 @@
+// Copyright Epic Games, Inc. All Rights Reserved.
+#pragma once
+
+#include "CoreTypes.h"
+#include "HAL/Platform.h"
+#include "Templates/TypeHash.h"
+
+namespace UE::Audio::Insights
+{
+	enum class ESignalFlowEntryType : uint8
+	{
+		OwnerObject,
+		SoundSource,
+		AudioBus,
+		Submix,
+		AudioDevice
+	};
+
+	struct FSignalFlowEntryKey
+	{
+		FSignalFlowEntryKey() = default;
+		FSignalFlowEntryKey(const ESignalFlowEntryType InEntryType, const uint32 InEntryID)
+			: EntryType(InEntryType)
+			, EntryID(InEntryID)
+		{
+		}
+
+		friend uint32 GetTypeHash(const FSignalFlowEntryKey& Key)
+		{
+			return HashCombine(GetTypeHash(static_cast<uint32>(Key.EntryType)), GetTypeHash(Key.EntryID));
+		}
+
+		bool operator==(const FSignalFlowEntryKey& InOther) const
+		{
+			return EntryType == InOther.EntryType && EntryID == InOther.EntryID;
+		}
+
+		bool operator<(const FSignalFlowEntryKey& InOther) const
+		{
+			return EntryType == InOther.EntryType ? EntryID < InOther.EntryID 
+												  : static_cast<uint32>(EntryType) < static_cast<uint32>(InOther.EntryType);
+		}
+
+		ESignalFlowEntryType EntryType = ESignalFlowEntryType::SoundSource;
+		uint32 EntryID = INDEX_NONE;
+	};
+} // namespace UE::Audio::Insights
diff --git a/Engine/Plugins/AudioInsights/Source/AudioInsights/Public/Messages/SignalFlowTraceMessages.h b/Engine/Plugins/AudioInsights/Source/AudioInsights/Public/Messages/SignalFlowTraceMessages.h
new file mode 100644
--- /dev/null
+++ b/Engine/Plugins/AudioInsights/Source/AudioInsights/Public/Messages/SignalFlowTraceMessages.h
@@ -0,0 +1,130 @@
+// Copyright Epic Games, Inc. All Rights Reserved.
+#pragma once
+
+#include "Cache/IAudioCachedMessage.h"
+#include "Messages/AnalyzerMessageQueue.h"
+#include "Messages/AudioBusTraceMessages.h"
+#include "Messages/SignalFlowEntryKey.h"
+#include "Messages/SoundTraceMessages.h"
+#include "Messages/SubmixTraceMessages.h"
+#include "Trace/Analyzer.h"
+#include "Views/TableDashboardViewFactory.h"
+
+namespace UE::Audio::Insights
+{
+	namespace SignalFlowMessageNames
+	{
+		extern const FName SoundWaveToBusSend;
+		extern const FName SoundWaveToSubmixSend;
+		extern const FName SubmixToSubmixSend;
+	};
+
+	struct FSignalFlowMessageBase : public IAudioCachedMessage
+	{
+		FSignalFlowMessageBase() = default;
+		FSignalFlowMessageBase(const Trace::IAnalyzer::FOnEventContext& InContext);
+
+		::Audio::FDeviceId DeviceId = INDEX_NONE;
+	};
+
+	struct FSignalFlowSendMessageBase : public FSignalFlowMessageBase
+	{
+		FSignalFlowSendMessageBase() = default;
+		FSignalFlowSendMessageBase(const Trace::IAnalyzer::FOnEventContext& InContext);
+
+		virtual uint32 GetSizeOf() const override { return sizeof(FSignalFlowSendMessageBase); }
+
+		uint32 SenderID = INDEX_NONE;
+		uint32 ReceiverID = INDEX_NONE;
+		float SendLevel = 0.0f;
+	};
+
+	struct FSoundToBusSendMessage : public FSignalFlowSendMessageBase
+	{
+		FSoundToBusSendMessage() = default;
+		FSoundToBusSendMessage(const Trace::IAnalyzer::FOnEventContext& InContext)
+			: FSignalFlowSendMessageBase(InContext)
+		{
+		}
+
+		virtual uint64 GetID() const override { return GetTypeHash(FSignalFlowEntryKey{ ESignalFlowEntryType::AudioBus, SenderID }); }
+		virtual const FName GetMessageName() const override { return SignalFlowMessageNames::SoundWaveToBusSend; }
+	};
+
+	struct FSoundToSubmixSendMessage : public FSignalFlowSendMessageBase
+	{
+		FSoundToSubmixSendMessage() = default;
+		FSoundToSubmixSendMessage(const Trace::IAnalyzer::FOnEventContext& InContext)
+			: FSignalFlowSendMessageBase(InContext)
+		{
+		}
+
+		virtual uint64 GetID() const override { return GetTypeHash(FSignalFlowEntryKey{ ESignalFlowEntryType::SoundSource, SenderID }); }
+		virtual const FName GetMessageName() const override { return SignalFlowMessageNames::SoundWaveToSubmixSend; }
+	};
+
+	struct FSubmixToSubmixSendMessage : public FSignalFlowSendMessageBase
+	{
+		FSubmixToSubmixSendMessage() = default;
+		FSubmixToSubmixSendMessage(const Trace::IAnalyzer::FOnEventContext& InContext);
+
+		virtual uint64 GetID() const override { return GetTypeHash(FSignalFlowEntryKey{ ESignalFlowEntryType::Submix, SenderID }); }
+		virtual const FName GetMessageName() const override { return SignalFlowMessageNames::SubmixToSubmixSend; }
+
+		bool bIsFinalStage = false;
+	};
+
+	class FSignalFlowMessages
+	{
+		TAnalyzerMessageQueue<FSoundWaveStartMessage> SoundWaveStartMessages;
+		TAnalyzerMessageQueue<FSoundWaveIsAlivePingMessage> SoundWaveIsAlivePingMessages;
+		TAnalyzerMessageQueue<FSoundWaveStopMessage> SoundWaveStopMessages;
+		TAnalyzerMessageQueue<FSoundStopMessage> SoundStopMessages;
+		TAnalyzerMessageQueue<FAudioBusStartMessage> BusStartMessages;
+		TAnalyzerMessageQueue<FAudioBusStartMessage> BusIsAlivePingMessages;
+		TAnalyzerMessageQueue<FAudioBusStopMessage> BusStopMessages;
+		TAnalyzerMessageQueue<FSubmixLoadedMessage> SubmixLoadedMessages;
+		TAnalyzerMessageQueue<FSubmixAlivePingMessage> SubmixAliveMessages;
+		TAnalyzerMessageQueue<FSubmixUnloadedMessage> SubmixUnloadedMessages;
+
+		TAnalyzerMessageQueue<FSoundToBusSendMessage> SoundToBusSendMessages;
+		TAnalyzerMessageQueue<FSoundToSubmixSendMessage> SoundToSubmixSendMessages;
+		TAnalyzerMessageQueue<FSubmixToSubmixSendMessage> SubmixToSubmixSendMessages;
+
+		friend class FSignalFlowTraceProvider;
+	};
+
+	struct FSignalFlowOutputData
+	{
+		float SendLevel = TNumericLimits<float>::Lowest();
+	};
+
+	struct FSignalFlowDashboardEntry : public IObjectDashboardEntry
+	{
+		FSignalFlowDashboardEntry() = default;
+		virtual ~FSignalFlowDashboardEntry() = default;
+
+		virtual FText GetDisplayName() const override 
+		{ 
+			const FString DisplayName = FSoftObjectPath(Name).GetAssetName();
+			return DisplayName.IsEmpty() ? FText::FromString(Name) : FText::FromString(DisplayName);
+		}
+
+		virtual const UObject* GetObject() const override { return FSoftObjectPath(Name).ResolveObject(); }
+		virtual UObject* GetObject() override { return FSoftObjectPath(Name).ResolveObject(); }
+		virtual bool IsValid() const override { return SignalFlowNodeID != static_cast<uint32>(INDEX_NONE); }
+
+		FSignalFlowEntryKey GetSignalFlowEntryKey() const { return { EntryType, SignalFlowNodeID }; }
+
+		::Audio::FDeviceId DeviceId = INDEX_NONE;
+		ESignalFlowEntryType EntryType = ESignalFlowEntryType::SoundSource;
+		uint32 SignalFlowNodeID = INDEX_NONE;
+		double Timestamp = 0.0;
+
+		FString Name;
+		FName IconName = NAME_None;
+
+		TSet<FSignalFlowEntryKey> Inputs;
+		TSortedMap<FSignalFlowEntryKey, FSignalFlowOutputData> Outputs;
+	};
+} // namespace UE::Audio::Insights
