[Audio Insights] Only trace audio channels is now a persistent setting

#Synced-CL 48449847

diff --git a/Engine/Plugins/AudioInsights/Source/AudioInsightsEditor/Private/AudioInsightsEditorDashboardFactory.cpp b/Engine/Plugins/AudioInsights/Source/AudioInsightsEditor/Private/AudioInsightsEditorDashboardFactory.cpp
--- a/Engine/Plugins/AudioInsights/Source/AudioInsightsEditor/Private/AudioInsightsEditorDashboardFactory.cpp
+++ b/Engine/Plugins/AudioInsights/Source/AudioInsightsEditor/Private/AudioInsightsEditorDashboardFactory.cpp
@@ -175,11 +175,17 @@ namespace UE::Audio::Insights
 		// If we haven't started a trace, start the analysis with the currently running trace
 		if (!TraceModule.IsAudioInsightsTrace())
 		{
-			TraceModule.StartTraceAnalysis(bOnlyTraceAudioChannels, ETraceMode::Recording);
+			TraceModule.StartTraceAnalysis(ShouldOnlyTraceAudioChannels(), ETraceMode::Recording);
 			BookmarkIndex = 0;
 		}
 	}
 
+	void FEditorDashboardFactory::OnOnlyTraceAudioChannelsPropertyChanged(const bool bInOnlyTraceAudioChannels)
+	{
+		IAudioInsightsTraceModule& TraceModule = FAudioInsightsEditorModule::GetChecked().GetTraceModule();
+		TraceModule.OnOnlyTraceAudioChannelsStateChanged(bInOnlyTraceAudioChannels);
+	}
+
 	void FEditorDashboardFactory::CreateTraceBookmark()
 	{
 		using namespace EditorDashboardFactoryPrivate;
@@ -291,6 +297,17 @@ namespace UE::Audio::Insights
 			FTraceAuxiliary::OnTraceStarted.Remove(OnTraceStartedHandle);
 			OnTraceStartedHandle.Reset();
 		}
+
+		if (OnOnlyTraceAudioChannelsPropertyChangedHandle.IsValid())
+		{
+			TObjectPtr<UAudioInsightsEditorSettings> AudioInsightsEditorSettings = GetMutableDefault<UAudioInsightsEditorSettings>();
+
+			if (AudioInsightsEditorSettings)
+			{
+				AudioInsightsEditorSettings->OnOnlyTraceAudioChannelsPropertyChanged.Remove(OnOnlyTraceAudioChannelsPropertyChangedHandle);
+				OnOnlyTraceAudioChannelsPropertyChangedHandle.Reset();
+			}
+		}
 	}
 
 	void FEditorDashboardFactory::BindCommands()
@@ -350,7 +367,7 @@ namespace UE::Audio::Insights
 		if (!bIsLayoutReseting)
 		{
 			IAudioInsightsTraceModule& TraceModule = FAudioInsightsEditorModule::GetChecked().GetTraceModule();
-			TraceModule.StartTraceAnalysis(bOnlyTraceAudioChannels, ETraceMode::Monitoring);
+			TraceModule.StartTraceAnalysis(ShouldOnlyTraceAudioChannels(), ETraceMode::Monitoring);
 		}
 
 		const TSharedRef<FTabManager::FLayout> TabLayout = LoadLayoutFromConfig();
@@ -717,13 +734,20 @@ namespace UE::Audio::Insights
 				[
 					SNew(SCheckBox)
 					.ToolTipText(OnlyTraceAudioChannelsTooltipText)
-					.IsChecked_Lambda([this]() { return bOnlyTraceAudioChannels ? ECheckBoxState::Checked : ECheckBoxState::Unchecked; })
+					.IsChecked_Lambda([this]()
+					{
+						return ShouldOnlyTraceAudioChannels() ? ECheckBoxState::Checked : ECheckBoxState::Unchecked;
+					})
 					.OnCheckStateChanged_Lambda([this](ECheckBoxState NewState) 
 					{ 
-						bOnlyTraceAudioChannels = NewState == ECheckBoxState::Checked;
+						TObjectPtr<UAudioInsightsEditorSettings> AudioInsightsEditorSettings = GetMutableDefault<UAudioInsightsEditorSettings>();
+						if (AudioInsightsEditorSettings)
+						{
+							AudioInsightsEditorSettings->bOnlyTraceAudioChannels = NewState == ECheckBoxState::Checked;
 
-						IAudioInsightsTraceModule& TraceModule = FAudioInsightsEditorModule::GetChecked().GetTraceModule();
-						TraceModule.OnOnlyTraceAudioChannelsStateChanged(bOnlyTraceAudioChannels);
+							IAudioInsightsTraceModule& TraceModule = FAudioInsightsEditorModule::GetChecked().GetTraceModule();
+							TraceModule.OnOnlyTraceAudioChannelsStateChanged(AudioInsightsEditorSettings->bOnlyTraceAudioChannels);
+						}
 					})
 				]
 				// Separator
@@ -851,13 +875,13 @@ namespace UE::Audio::Insights
 
 							if (TraceMode == ETraceMode::Monitoring)
 							{
-								StartTraceAnalysisDeferred(bOnlyTraceAudioChannels, ETraceMode::Recording);
+								StartTraceAnalysisDeferred(ShouldOnlyTraceAudioChannels(), ETraceMode::Recording);
 							}
 						}
 					}
 					else
 					{
-						TraceModule.StartTraceAnalysis(bOnlyTraceAudioChannels, InTraceButtonMode);
+						TraceModule.StartTraceAnalysis(ShouldOnlyTraceAudioChannels(), InTraceButtonMode);
 
 						if (InTraceButtonMode == ETraceMode::Recording)
 						{
@@ -962,6 +986,16 @@ namespace UE::Audio::Insights
 		{
 			OnTraceStartedHandle = FTraceAuxiliary::OnTraceStarted.AddSP(this, &FEditorDashboardFactory::OnTraceStarted);
 		}
+
+		if (!OnOnlyTraceAudioChannelsPropertyChangedHandle.IsValid())
+		{
+			TObjectPtr<UAudioInsightsEditorSettings> AudioInsightsEditorSettings = GetMutableDefault<UAudioInsightsEditorSettings>();
+
+			if (AudioInsightsEditorSettings)
+			{
+				OnOnlyTraceAudioChannelsPropertyChangedHandle = AudioInsightsEditorSettings->OnOnlyTraceAudioChannelsPropertyChanged.AddSP(this, &FEditorDashboardFactory::OnOnlyTraceAudioChannelsPropertyChanged);
+			}
+		}
 	}
 
 	TSharedRef<FTabManager::FLayout> FEditorDashboardFactory::GetDefaultTabLayout()
@@ -1168,6 +1202,12 @@ namespace UE::Audio::Insights
 			FLayoutSaveRestore::SaveToConfig(GEditorLayoutIni, DashboardTabManager->PersistLayout());
 		}
 	}
+
+	bool FEditorDashboardFactory::ShouldOnlyTraceAudioChannels() const
+	{
+		const TObjectPtr<const UAudioInsightsEditorSettings> AudioInsightsEditorSettings = GetDefault<UAudioInsightsEditorSettings>();
+		return AudioInsightsEditorSettings && AudioInsightsEditorSettings->bOnlyTraceAudioChannels;
+	}
 } // namespace UE::Audio::Insights
 
 #undef LOCTEXT_NAMESPACE
diff --git a/Engine/Plugins/AudioInsights/Source/AudioInsightsEditor/Private/AudioInsightsEditorDashboardFactory.h b/Engine/Plugins/AudioInsights/Source/AudioInsightsEditor/Private/AudioInsightsEditorDashboardFactory.h
--- a/Engine/Plugins/AudioInsights/Source/AudioInsightsEditor/Private/AudioInsightsEditorDashboardFactory.h
+++ b/Engine/Plugins/AudioInsights/Source/AudioInsightsEditor/Private/AudioInsightsEditorDashboardFactory.h
@@ -39,6 +39,9 @@ namespace UE::Audio::Insights
 		void OnDeviceDestroyed(::Audio::FDeviceId InDeviceId);
 
 		void OnTraceStarted(FTraceAuxiliary::EConnectionType InTraceType, const FString& InTraceDestination);
+
+		void OnOnlyTraceAudioChannelsPropertyChanged(const bool bInOnlyTraceAudioChannels);
+
 		void CreateTraceBookmark();
 
 		void RefreshDeviceSelector();
@@ -61,6 +64,8 @@ namespace UE::Audio::Insights
 		TSharedRef<FTabManager::FLayout> LoadLayoutFromConfig();
 		void SaveLayoutToConfig();
 
+		bool ShouldOnlyTraceAudioChannels() const;
+
 		FDelegateHandle OnWorldRegisteredToAudioDeviceHandle;
 		FDelegateHandle OnWorldUnregisteredFromAudioDeviceHandle;
 
@@ -69,7 +74,8 @@ namespace UE::Audio::Insights
 
 		FDelegateHandle OnTraceStartedHandle;
 
-		bool bOnlyTraceAudioChannels = false;
+		FDelegateHandle OnOnlyTraceAudioChannelsPropertyChangedHandle;
+
 		bool bIsLayoutReseting = false;
 
 		TSharedPtr<FTabManager> DashboardTabManager;
diff --git a/Engine/Plugins/AudioInsights/Source/AudioInsightsEditor/Private/AudioInsightsEditorSettings.cpp b/Engine/Plugins/AudioInsights/Source/AudioInsightsEditor/Private/AudioInsightsEditorSettings.cpp
--- a/Engine/Plugins/AudioInsights/Source/AudioInsightsEditor/Private/AudioInsightsEditorSettings.cpp
+++ b/Engine/Plugins/AudioInsights/Source/AudioInsightsEditor/Private/AudioInsightsEditorSettings.cpp
@@ -98,6 +98,11 @@ void UAudioInsightsEditorSettings::PostEditChangeChainProperty(FPropertyChangedC
 		return;
 	}
 
+	if (ActiveMemberNode->GetValue()->GetFName() == GET_MEMBER_NAME_CHECKED(UAudioInsightsEditorSettings, bOnlyTraceAudioChannels))
+	{
+		OnOnlyTraceAudioChannelsPropertyChanged.Broadcast(bOnlyTraceAudioChannels);
+	}
+
 	if (ActiveMemberNode->GetValue()->GetFName() == GET_MEMBER_NAME_CHECKED(UAudioInsightsEditorSettings, EventLogSettings))
 	{
 		if (VerifyEventLogInput(*ActiveMemberNode, PropertyChangedEvent))
diff --git a/Engine/Plugins/AudioInsights/Source/AudioInsightsEditor/Public/AudioInsightsEditorSettings.h b/Engine/Plugins/AudioInsights/Source/AudioInsightsEditor/Public/AudioInsightsEditorSettings.h
--- a/Engine/Plugins/AudioInsights/Source/AudioInsightsEditor/Public/AudioInsightsEditorSettings.h
+++ b/Engine/Plugins/AudioInsights/Source/AudioInsightsEditor/Public/AudioInsightsEditorSettings.h
@@ -34,6 +34,10 @@ public:
 	UPROPERTY(Config, EditAnywhere, Category="World Filter")
 	bool bWorldFilterDefaultsToFirstClient = false;
 
+	/** Disable all non - command line trace channels apart from Audio, Audio Mixer and Bookmark when starting a live monitoring / live trace.This will reduce the file size of snapshots / trace sessions while using Audio Insights. */
+	UPROPERTY(Config, EditAnywhere, Category = "Trace Channels")
+	bool bOnlyTraceAudioChannels = true;
+
 	/** Settings for analyzer rack spectrogram widget */
 	UPROPERTY(EditAnywhere, config, Category = Spectrogram, meta = (ShowOnlyInnerProperties))
 	FSpectrogramRackUnitSettings SpectrogramSettings;
@@ -54,6 +58,9 @@ public:
 	UPROPERTY(EditAnywhere, config, Category = SoundDashboard, meta = (ShowOnlyInnerProperties))
 	FSoundDashboardSettings SoundDashboardSettings;
 
+	DECLARE_MULTICAST_DELEGATE_OneParam(FOnOnlyTraceAudioChannelsPropertyChanged, const bool);
+	FOnOnlyTraceAudioChannelsPropertyChanged OnOnlyTraceAudioChannelsPropertyChanged;
+
 private:
 	struct EditedEvent
 	{
