[Audio Insights] Crash fix after closing large trace files while they load in standalone
* We appear to be still receiving trace event messages after the session has closed and the module has shutdown. Trying to load the module again from a different thread to the main game thread causes a crash.
* Make sure the module is actually loaded before grabbing it rather than getting and loading it automatically

#Synced-CL 48449847

diff --git a/Engine/Plugins/AudioInsights/Source/AudioInsights/Private/AudioInsightsModule.cpp b/Engine/Plugins/AudioInsights/Source/AudioInsights/Private/AudioInsightsModule.cpp
--- a/Engine/Plugins/AudioInsights/Source/AudioInsights/Private/AudioInsightsModule.cpp
+++ b/Engine/Plugins/AudioInsights/Source/AudioInsights/Private/AudioInsightsModule.cpp
@@ -148,6 +148,8 @@ namespace UE::Audio::Insights
 
 	FAudioInsightsModule& FAudioInsightsModule::GetChecked()
 	{
+		ensureMsgf(IsInGameThread(), TEXT("Do not call this from outside of the Game Thread, if the module is not loaded this will fail - use FAudioInsightsModule::GetModulePtr() instead."));
+
 		return static_cast<FAudioInsightsModule&>(FModuleManager::LoadModuleChecked<IAudioInsightsModule>("AudioInsights"));
 	}
 
diff --git a/Engine/Plugins/AudioInsights/Source/AudioInsights/Private/AudioInsightsTraceProviderBase.cpp b/Engine/Plugins/AudioInsights/Source/AudioInsights/Private/AudioInsightsTraceProviderBase.cpp
--- a/Engine/Plugins/AudioInsights/Source/AudioInsights/Private/AudioInsightsTraceProviderBase.cpp
+++ b/Engine/Plugins/AudioInsights/Source/AudioInsights/Private/AudioInsightsTraceProviderBase.cpp
@@ -168,13 +168,19 @@ namespace UE::Audio::Insights
 
 	ECacheAndProcess FTraceProviderBase::GetMessageCacheAndProcessingStatus() const
 	{
-		const FAudioInsightsModule& AudioInsightsModule = FAudioInsightsModule::GetChecked();
-		return AudioInsightsModule.GetTimingViewExtender().GetMessageCacheAndProcessingStatus();
+		if (const FAudioInsightsModule* AudioInsightsModulePtr = FAudioInsightsModule::GetModulePtr())
+		{
+			return AudioInsightsModulePtr->GetTimingViewExtender().GetMessageCacheAndProcessingStatus();
+		}
+		
+		return ECacheAndProcess::None;
 	}
 
 	void FTraceProviderBase::ResetMessageProcessType()
 	{
-		FAudioInsightsModule& AudioInsightsModule = FAudioInsightsModule::GetChecked();
-		AudioInsightsModule.GetTimingViewExtender().ResetMessageProcessType();
+		if (FAudioInsightsModule* AudioInsightsModulePtr = FAudioInsightsModule::GetModulePtr())
+		{
+			AudioInsightsModulePtr->GetTimingViewExtender().ResetMessageProcessType();
+		}
 	}
 } // namespace UE::Audio::Insights
