[Audio Insights] Crash fix in standalone when opening large .utrace files
* Threading issue can occur if CacheMessageIfLatest is called on the trace analysis thread - if FindClosestMessage is called when the cache is increased inside on the game thread in standalone, we can end up accessing a chunk before it is ready.
* Instead of try to access the cache chunks on the trace analysis thread, move the check for whether this message is the latest to FAudioCachedMessageChunk::AddMessageToChunk, which is called on the game thread. Ideally all messages should be received in order, so make this check the default.

#Synced-CL 48502821

diff --git a/Engine/Plugins/AudioInsights/Source/AudioInsights/Private/Cache/AudioCachedMessageChunk.cpp b/Engine/Plugins/AudioInsights/Source/AudioInsights/Private/Cache/AudioCachedMessageChunk.cpp
--- a/Engine/Plugins/AudioInsights/Source/AudioInsights/Private/Cache/AudioCachedMessageChunk.cpp
+++ b/Engine/Plugins/AudioInsights/Source/AudioInsights/Private/Cache/AudioCachedMessageChunk.cpp
@@ -11,20 +11,24 @@ namespace UE::Audio::Insights
 		}
 
 		const double MessageTimestamp = Message->Timestamp;
-		if (RangeStartTimestamp == INVALID_TIMESTAMP)
+		
+		TSharedPtr<TraceServices::TPagedArray<TSharedPtr<IAudioCachedMessage>>>& CachedMessages = Messages.FindOrAdd(Message->GetMessageName(), MakeShared<TraceServices::TPagedArray<TSharedPtr<IAudioCachedMessage>>>(Allocator, 4096));
+		if (CachedMessages->Num() == 0 || CachedMessages->Last()->Timestamp <= MessageTimestamp)
 		{
-			RangeStartTimestamp = MessageTimestamp;
-		}
+			if (RangeStartTimestamp == INVALID_TIMESTAMP)
+			{
+				RangeStartTimestamp = MessageTimestamp;
+			}
 
-		if (RangeEndTimestamp < MessageTimestamp)
-		{
-			RangeEndTimestamp = MessageTimestamp;
-		}
+			if (RangeEndTimestamp < MessageTimestamp)
+			{
+				RangeEndTimestamp = MessageTimestamp;
+			}
 
-		CurrentChunkSize += Message->GetSizeOf();
+			CurrentChunkSize += Message->GetSizeOf();
 
-		TSharedPtr<TraceServices::TPagedArray<TSharedPtr<IAudioCachedMessage>>>& CachedMessages = Messages.FindOrAdd(Message->GetMessageName(), MakeShared<TraceServices::TPagedArray<TSharedPtr<IAudioCachedMessage>>>(Allocator, 4096));
-		CachedMessages->EmplaceBack(Message);
+			CachedMessages->EmplaceBack(Message);
+		}
     }
 
 	void FAudioCachedMessageChunk::ClearChunk()
diff --git a/Engine/Plugins/AudioInsights/Source/AudioInsights/Private/Providers/SoundTraceProvider.cpp b/Engine/Plugins/AudioInsights/Source/AudioInsights/Private/Providers/SoundTraceProvider.cpp
--- a/Engine/Plugins/AudioInsights/Source/AudioInsights/Private/Providers/SoundTraceProvider.cpp
+++ b/Engine/Plugins/AudioInsights/Source/AudioInsights/Private/Providers/SoundTraceProvider.cpp
@@ -237,50 +237,50 @@ namespace UE::Audio::Insights
 
 					case RouteId_Priority:
 					{
-						CacheMessageIfLatest<FSoundPriorityMessage>(Context, Messages.PriorityMessages);
+						CacheMessage<FSoundPriorityMessage>(Context, Messages.PriorityMessages);
 						break;
 					}
 
 					case RouteId_Distance:
 					{
-						CacheMessageIfLatest<FSoundDistanceMessage>(Context, Messages.DistanceMessages);
+						CacheMessage<FSoundDistanceMessage>(Context, Messages.DistanceMessages);
 						break;
 					}
 
 					case RouteId_DistanceAttenuation:
 					{
-						CacheMessageIfLatest<FSoundDistanceAttenuationMessage>(Context, Messages.DistanceAttenuationMessages);
+						CacheMessage<FSoundDistanceAttenuationMessage>(Context, Messages.DistanceAttenuationMessages);
 						break;
 					}
 
 					case RouteId_Filters:
 					{
-						CacheMessageIfLatest<FSoundLPFFreqMessage>(Context, Messages.LPFFreqMessages);
-						CacheMessageIfLatest<FSoundHPFFreqMessage>(Context, Messages.HPFFreqMessages);
+						CacheMessage<FSoundLPFFreqMessage>(Context, Messages.LPFFreqMessages);
+						CacheMessage<FSoundHPFFreqMessage>(Context, Messages.HPFFreqMessages);
 						break;
 					}
 
 					case RouteId_Amplitude:
 					{
-						CacheMessageIfLatest<FSoundEnvelopeMessage>(Context, Messages.AmplitudeMessages);
+						CacheMessage<FSoundEnvelopeMessage>(Context, Messages.AmplitudeMessages);
 						break;
 					}
 
 					case RouteId_Volume:
 					{
-						CacheMessageIfLatest<FSoundVolumeMessage>(Context, Messages.VolumeMessages);
+						CacheMessage<FSoundVolumeMessage>(Context, Messages.VolumeMessages);
 						break;
 					}
 
 					case RouteId_Pitch:
 					{
-						CacheMessageIfLatest<FSoundPitchMessage>(Context, Messages.PitchMessages);
+						CacheMessage<FSoundPitchMessage>(Context, Messages.PitchMessages);
 						break;
 					}
 
 					case RouteId_RelativeRenderCost:
 					{
-						CacheMessageIfLatest<FSoundRelativeRenderCostMessage>(Context, Messages.RelativeRenderCostMessages);
+						CacheMessage<FSoundRelativeRenderCostMessage>(Context, Messages.RelativeRenderCostMessages);
 						break;
 					}
 
diff --git a/Engine/Plugins/AudioInsights/Source/AudioInsights/Public/AudioInsightsTraceProviderBase.h b/Engine/Plugins/AudioInsights/Source/AudioInsights/Public/AudioInsightsTraceProviderBase.h
--- a/Engine/Plugins/AudioInsights/Source/AudioInsights/Public/AudioInsightsTraceProviderBase.h
+++ b/Engine/Plugins/AudioInsights/Source/AudioInsights/Public/AudioInsightsTraceProviderBase.h
@@ -126,17 +126,6 @@ namespace UE::Audio::Insights
 				AddMessageToCache(MoveTemp(Message), OutMessageQueue);
 			}
 
-			template <TIsCacheableMessage T>
-			void CacheMessageIfLatest(const FOnEventContext& Context, TAnalyzerMessageQueue<T>& OutMessageQueue)
-			{
-				TSharedPtr<T> Message = MakeShared<T>(Context);
-				const T* LastMessage = CacheManager.FindClosestMessage<T>(Message->GetMessageName(), CacheManager.GetCacheEndTimeStamp(), Message->GetID(), Message->Timestamp);
-				if (LastMessage == nullptr || LastMessage->Timestamp < Message->Timestamp)
-				{
-					AddMessageToCache(MoveTemp(Message), OutMessageQueue);
-				}
-			}
-
 		private:
 			AUDIOINSIGHTS_API virtual bool OnEvent(uint16 RouteId, EStyle Style, const FOnEventContext& Context) override;
 
