[Audio Insights] Merge shared functionality of Submix is alive pings and loaded messages - allows us to reduce duplicated processing code

#Synced-CL 48646572

diff --git a/Engine/Plugins/AudioInsights/Source/AudioInsights/Private/Messages/SubmixTraceMessages.cpp b/Engine/Plugins/AudioInsights/Source/AudioInsights/Private/Messages/SubmixTraceMessages.cpp
--- a/Engine/Plugins/AudioInsights/Source/AudioInsights/Private/Messages/SubmixTraceMessages.cpp
+++ b/Engine/Plugins/AudioInsights/Source/AudioInsights/Private/Messages/SubmixTraceMessages.cpp
@@ -45,28 +45,6 @@ namespace UE::Audio::Insights
 		return MemorySize;
 	};
 
-	// FSubmixAlivePingMessage
-	FSubmixAlivePingMessage::FSubmixAlivePingMessage(const Trace::IAnalyzer::FOnEventContext& InContext)
-		: FSubmixMessageBase(InContext)
-	{
-		const Trace::IAnalyzer::FEventData& EventData = InContext.EventData;
-
-		EventData.GetString("Name", Name);
-		NumChannels = EventData.GetValue<int32>("NumChannels");
-		bIsMainSubmix = EventData.GetValue<bool>("IsMainSubmix");
-	}
-
-	uint32 FSubmixAlivePingMessage::GetSizeOf() const
-	{
-		// Start with the fixed size of the struct
-		uint32 MemorySize = sizeof(FSubmixAlivePingMessage);
-
-		// Add any dynamically sized members
-		MemorySize += Name.GetAllocatedSize();
-
-		return MemorySize;
-	};
-
 	// FSubmixHasActivityMessage
 	FSubmixHasActivityMessage::FSubmixHasActivityMessage(const Trace::IAnalyzer::FOnEventContext& InContext)
 		: FSubmixMessageBase(InContext)
diff --git a/Engine/Plugins/AudioInsights/Source/AudioInsights/Public/Messages/SubmixTraceMessages.h b/Engine/Plugins/AudioInsights/Source/AudioInsights/Public/Messages/SubmixTraceMessages.h
--- a/Engine/Plugins/AudioInsights/Source/AudioInsights/Public/Messages/SubmixTraceMessages.h
+++ b/Engine/Plugins/AudioInsights/Source/AudioInsights/Public/Messages/SubmixTraceMessages.h
@@ -49,17 +49,15 @@ namespace UE::Audio::Insights
 		bool bIsMainSubmix = false;
 	};
 
-	struct FSubmixAlivePingMessage : public FSubmixMessageBase
+	struct FSubmixAlivePingMessage : public FSubmixLoadedMessage
 	{
 		FSubmixAlivePingMessage() = default;
-		FSubmixAlivePingMessage(const Trace::IAnalyzer::FOnEventContext& InContext);
+		FSubmixAlivePingMessage(const Trace::IAnalyzer::FOnEventContext& InContext)
+			: FSubmixLoadedMessage(InContext)
+		{
+		}
 
-		virtual uint32 GetSizeOf() const override;
 		virtual const FName GetMessageName() const override { return SubmixMessageNames::IsAlivePing; }
-
-		FString Name;
-		int32 NumChannels = 0;
-		bool bIsMainSubmix = false;
 	};
 
 	struct FSubmixHasActivityMessage : public FSubmixMessageBase
