[Audio Insights] Providers can add a message to an analyzer message queue without adding the message to the Audio Insights cache, while still stopping/starting processing when the cache manager instructs.

Predominant use case is messages shared across multiple providers - we only need to cache one of these messages to retrieve it later, additional caches only fill up the cache uneccessarily

#Synced-CL 48717847

diff --git a/Engine/Plugins/AudioInsights/Source/AudioInsights/Public/AudioInsightsTraceProviderBase.h b/Engine/Plugins/AudioInsights/Source/AudioInsights/Public/AudioInsightsTraceProviderBase.h
--- a/Engine/Plugins/AudioInsights/Source/AudioInsights/Public/AudioInsightsTraceProviderBase.h
+++ b/Engine/Plugins/AudioInsights/Source/AudioInsights/Public/AudioInsightsTraceProviderBase.h
@@ -119,6 +119,7 @@ namespace UE::Audio::Insights
 				return *StaticCastSharedRef<TProviderType>(Provider);
 			}
 
+			// Will cache a message, and add it to the provided OutMessageQueue for processing
 			template <TIsCacheableMessage T>
 			void CacheMessage(const FOnEventContext& Context, TAnalyzerMessageQueue<T>& OutMessageQueue)
 			{
@@ -126,6 +127,17 @@ namespace UE::Audio::Insights
 				AddMessageToCache(MoveTemp(Message), OutMessageQueue);
 			}
 
+
+			// If processing is enabled, will add a message to the provided OutMessageQueue for processing, but will not cache it.
+			template <TIsCacheableMessage T>
+			void AddMessageToQueue(const FOnEventContext& Context, TAnalyzerMessageQueue<T>& OutMessageQueue)
+			{
+				if (ShouldProcessNewEvents())
+				{
+					OutMessageQueue.Enqueue(T(Context));
+				}
+			}
+
 		private:
 			AUDIOINSIGHTS_API virtual bool OnEvent(uint16 RouteId, EStyle Style, const FOnEventContext& Context) override;
 
