[Audio Insights] Signal Flow part 04 - Signal Flow Provider
* Initial submit of the signal flow provider, including processing messages for submixes, sound sources and audio buses, along with linking their outputs and send levels.

#Synced-CL 48842273

diff --git a/Engine/Plugins/AudioInsights/Source/AudioInsights/Private/Providers/SignalFlowTraceProvider.cpp b/Engine/Plugins/AudioInsights/Source/AudioInsights/Private/Providers/SignalFlowTraceProvider.cpp
new file mode 100644
--- /dev/null
+++ b/Engine/Plugins/AudioInsights/Source/AudioInsights/Private/Providers/SignalFlowTraceProvider.cpp
@@ -0,0 +1,585 @@
+// Copyright Epic Games, Inc. All Rights Reserved.
+#include "Providers/SignalFlowTraceProvider.h"
+
+#include "AudioInsightsModule.h"
+#include "Cache/AudioInsightsCacheManager.h"
+
+#define LOCTEXT_NAMESPACE "AudioInsights"
+
+namespace UE::Audio::Insights
+{
+	namespace FSignalFlowTraceProviderPrivate
+	{
+		const FName AudioDeviceIconName = "AudioInsights.Icon.Dashboard";
+		const FName AudioBusIconName = "AudioInsights.Icon";
+		const FName SubmixIconName = "AudioInsights.Icon.Submix";
+
+		const FName MetasoundIconName = "AudioInsights.Icon.SoundDashboard.MetaSound";
+		const FName ProceduralSourceIconName = "AudioInsights.Icon.SoundDashboard.ProceduralSource";
+		const FName SoundWaveIconName = "AudioInsights.Icon.SoundDashboard.SoundWave";
+		const FName SoundCueIconName = "AudioInsights.Icon.SoundDashboard.SoundCue";
+		const FName SoundCueTemplateIconName = "AudioInsights.Icon.SoundDashboard.SoundCue";
+
+		FName GetSourceIconName(const ESoundDashboardEntryType SoundEntryType)
+		{
+			switch (SoundEntryType)
+			{
+				case ESoundDashboardEntryType::MetaSound:
+					return MetasoundIconName;
+				case ESoundDashboardEntryType::SoundCue:
+					return SoundCueIconName;
+				case ESoundDashboardEntryType::ProceduralSource:
+					return ProceduralSourceIconName;
+				case ESoundDashboardEntryType::SoundWave:
+					return SoundWaveIconName;
+				case ESoundDashboardEntryType::SoundCueTemplate:
+					return SoundCueTemplateIconName;
+			}
+			
+			return SoundWaveIconName;
+		}
+	} // FSignalFlowTraceProviderPrivate
+
+	FSignalFlowTraceProvider::FSignalFlowTraceProvider()
+		: TDeviceDataMapTraceProvider<FSignalFlowEntryKey, TSharedPtr<FSignalFlowDashboardEntry>>(GetName_Static())
+	{
+	}
+	
+	FName FSignalFlowTraceProvider::GetName_Static()
+	{
+		static const FLazyName SignalFlowTraceProviderName = "SignalFlowProvider";
+		return SignalFlowTraceProviderName;
+	}
+
+	UE::Trace::IAnalyzer* FSignalFlowTraceProvider::ConstructAnalyzer(TraceServices::IAnalysisSession& InSession)
+	{
+		class FSignalFlowTraceAnalyzer : public FTraceAnalyzerBase
+		{
+		public:
+			FSignalFlowTraceAnalyzer(TSharedRef<FSignalFlowTraceProvider> InProvider, TraceServices::IAnalysisSession& InSession)
+				: FTraceAnalyzerBase(InProvider)
+				, Session(InSession)
+			{
+			}
+
+			virtual void OnAnalysisBegin(const FOnAnalysisContext& Context) override
+			{
+				FTraceAnalyzerBase::OnAnalysisBegin(Context);
+
+				UE::Trace::IAnalyzer::FInterfaceBuilder& Builder = Context.InterfaceBuilder;
+
+				Builder.RouteEvent(RouteID_SoundWaveStart, "Audio", "SoundWaveStart");
+				Builder.RouteEvent(RouteID_SoundWaveIsAlivePing, "Audio", "SoundWaveIsAlivePing");
+				Builder.RouteEvent(RouteID_SoundWaveStop, "Audio", "SoundWaveStop");
+				Builder.RouteEvent(RouteID_ActiveSoundStop, "Audio", "SoundStop");
+
+				Builder.RouteEvent(RouteId_BusStart, "Audio", "AudioBusStart");
+				Builder.RouteEvent(RouteId_BusIsAlivePing, "Audio", "AudioBusIsAlivePing");
+				Builder.RouteEvent(RouteId_BusStop, "Audio", "AudioBusStop");
+
+				Builder.RouteEvent(RouteId_SubmixLoaded, "Audio", "SubmixLoaded");
+				Builder.RouteEvent(RouteId_SubmixAlivePing, "Audio", "SubmixAlivePing");
+				Builder.RouteEvent(RouteId_SubmixUnloaded, "Audio", "SubmixUnloaded");
+
+				Builder.RouteEvent(RouteId_SoundToBusSend, "Audio", "SoundToBusSend");
+				Builder.RouteEvent(RouteId_SoundToSubmixSend, "Audio", "SoundToSubmixSend");
+				Builder.RouteEvent(RouteId_SubmixToSubmixSend, "Audio", "SubmixToSubmixSend");
+			}
+
+			virtual bool OnHandleEvent(uint16 RouteId, EStyle Style, const FOnEventContext& Context) override
+			{
+				LLM_SCOPE_BYNAME(TEXT("Insights/FSignalFlowTraceAnalyzer"));
+
+				FSignalFlowMessages& Messages = GetProvider<FSignalFlowTraceProvider>().TraceMessages;
+
+				switch (RouteId)
+				{
+					// The first set of messages here are also used in other providers.
+					// We do not need to add these messages to the cache multiple times.
+					// Call AddMessageToQueue instead to just add them to the processing queue.
+					case RouteID_SoundWaveStart:
+					{
+						AddMessageToQueue<FSoundWaveStartMessage>(Context, Messages.SoundWaveStartMessages);
+						break;
+					}
+
+					case RouteID_SoundWaveIsAlivePing:
+					{
+						AddMessageToQueue<FSoundWaveIsAlivePingMessage>(Context, Messages.SoundWaveIsAlivePingMessages);
+						break;
+					}
+
+					case RouteID_ActiveSoundStop:
+					{
+						AddMessageToQueue<FSoundStopMessage>(Context, Messages.SoundStopMessages);
+						break;
+					}
+
+					case RouteId_BusStart:
+					{
+						AddMessageToQueue<FAudioBusStartMessage>(Context, Messages.BusStartMessages);
+						break;
+					}
+
+					case RouteId_BusStop:
+					{
+						AddMessageToQueue<FAudioBusStopMessage>(Context, Messages.BusStopMessages);
+						break;
+					}
+
+					case RouteId_SubmixLoaded:
+					{
+						AddMessageToQueue<FSubmixLoadedMessage>(Context, Messages.SubmixLoadedMessages);
+						break;
+					}
+
+					case RouteId_SubmixAlivePing:
+					{
+						AddMessageToQueue<FSubmixAlivePingMessage>(Context, Messages.SubmixAliveMessages);
+						break;
+					}
+
+					case RouteId_SubmixUnloaded:
+					{
+						AddMessageToQueue<FSubmixUnloadedMessage>(Context, Messages.SubmixUnloadedMessages);
+						break;
+					}
+
+					// The second set of messages here are unique to signal flow.
+					// Make sure these are added to the cache.
+					case RouteID_SoundWaveStop:
+					{
+						CacheMessage<FSoundWaveStopMessage>(Context, Messages.SoundWaveStopMessages);
+						break;
+					}
+
+					case RouteId_BusIsAlivePing:
+					{
+						CacheMessage<FAudioBusStartMessage>(Context, Messages.BusIsAlivePingMessages);
+						break;
+					}
+
+					case RouteId_SoundToBusSend:
+					{
+						CacheMessage<FSoundToBusSendMessage>(Context, Messages.SoundToBusSendMessages);
+						break;
+					}
+
+					case RouteId_SoundToSubmixSend:
+					{
+						CacheMessage<FSoundToSubmixSendMessage>(Context, Messages.SoundToSubmixSendMessages);
+						break;
+					}
+
+					case RouteId_SubmixToSubmixSend:
+					{
+						CacheMessage<FSubmixToSubmixSendMessage>(Context, Messages.SubmixToSubmixSendMessages);
+						break;
+					}
+
+					default:
+					{
+						return OnEventFailure(RouteId, Style, Context);
+					}
+				}
+
+				const double Timestamp = Context.EventTime.AsSeconds(Context.EventData.GetValue<uint64>("Timestamp"));
+
+				{
+					TraceServices::FAnalysisSessionEditScope SessionEditScope(Session);
+					Session.UpdateDurationSeconds(Timestamp);
+				}
+
+				return OnEventSuccess(RouteId, Style, Context);
+			}
+
+		private:
+			enum : uint16
+			{
+				RouteID_SoundWaveStart,
+				RouteID_SoundWaveIsAlivePing,
+				RouteID_SoundWaveStop,
+				RouteID_ActiveSoundStop,
+				RouteId_BusStart,
+				RouteId_BusIsAlivePing,
+				RouteId_BusStop,
+				RouteId_SubmixLoaded,
+				RouteId_SubmixAlivePing,
+				RouteId_SubmixUnloaded,
+
+				RouteId_SoundToBusSend,
+				RouteId_SoundToSubmixSend,
+				RouteId_SubmixToSubmixSend
+			};
+
+			TraceServices::IAnalysisSession& Session;
+		};
+
+		return new FSignalFlowTraceAnalyzer(AsShared(), InSession);
+	}
+
+	bool FSignalFlowTraceProvider::ProcessMessages()
+	{
+		auto GetOwnerEntry = [this](const FSoundStopMessage& Msg)
+		{
+			return FindDeviceEntry(Msg.DeviceId, { ESignalFlowEntryType::OwnerObject, Msg.ActiveSoundPlayOrder });
+		};
+
+		auto GetOwnerEntryFromSoundWaveStop = [this](const FSoundWaveStopMessage& Msg)
+		{
+			return FindDeviceEntry(Msg.DeviceId, { ESignalFlowEntryType::OwnerObject, Msg.ActiveSoundPlayOrder });
+		};
+
+		auto CreateSoundwaveEntry = [this](const FSoundWaveStartMessage& Msg)
+		{
+			CreateAudioDeviceEntry(Msg, Msg.DeviceId);
+
+			return CreateEntry(Msg, Msg.DeviceId, ESignalFlowEntryType::SoundSource, Msg.WaveInstancePlayOrder);
+		};
+
+		auto GetSoundwaveEntry = [this](const ::Audio::FDeviceId DeviceId, const uint32 WaveInstancePlayOrder)
+		{
+			return FindDeviceEntry(DeviceId, { ESignalFlowEntryType::SoundSource, WaveInstancePlayOrder });
+		};
+
+		auto GetSoundwaveEntryFromMessage = [this](const FSoundWaveStopMessage& Msg)
+		{
+			return FindDeviceEntry(Msg.DeviceId, { ESignalFlowEntryType::SoundSource, Msg.WaveInstancePlayOrder });
+		};
+
+		auto GetSoundwaveEntryForSubmixSendMessage = [&GetSoundwaveEntry](const FSignalFlowSendMessageBase& Msg)
+		{
+			return GetSoundwaveEntry(Msg.DeviceId, Msg.SenderID);
+		};
+
+		auto CreateBusEntry = [this](const FAudioBusStartMessage& Msg)
+		{
+			CreateAudioDeviceEntry(Msg, Msg.DeviceId);
+
+			return CreateEntry(Msg, Msg.DeviceId, ESignalFlowEntryType::AudioBus, Msg.AudioBusId);
+		};
+
+		auto GetBusEntry = [this](const FAudioBusMessageBase& Msg)
+		{
+			return FindDeviceEntry(Msg.DeviceId, { ESignalFlowEntryType::AudioBus, Msg.AudioBusId });
+		};
+
+		auto CreateSubmixEntry = [this](const FSubmixMessageBase& Msg)
+		{
+			CreateAudioDeviceEntry(Msg, Msg.DeviceId);
+
+			return CreateEntry(Msg, Msg.DeviceId, ESignalFlowEntryType::Submix, Msg.SubmixId);
+		};
+
+		auto GetSubmixEntry = [this](const FSubmixMessageBase& Msg)
+		{
+			return FindDeviceEntry(Msg.DeviceId, { ESignalFlowEntryType::Submix, Msg.SubmixId });
+		};
+
+		auto GetSubmixToSubmixEntry = [this](const FSubmixToSubmixSendMessage& Msg)
+		{
+			return FindDeviceEntry(Msg.DeviceId, { ESignalFlowEntryType::Submix, Msg.SenderID });
+		};
+
+		TSet<FSignalFlowEntryKey> RemovedKeys;
+
+		ProcessMessageQueue<FAudioBusStartMessage>(TraceMessages.BusStartMessages, CreateBusEntry,
+		[this](const FAudioBusStartMessage& Msg, TSharedPtr<FSignalFlowDashboardEntry>* OutEntry)
+		{
+			ProcessNewBusEntry(Msg, OutEntry);
+		});
+
+		ProcessMessageQueue<FAudioBusStartMessage>(TraceMessages.BusIsAlivePingMessages, CreateBusEntry,
+		[this](const FAudioBusStartMessage& Msg, TSharedPtr<FSignalFlowDashboardEntry>* OutEntry)
+		{
+			ProcessNewBusEntry(Msg, OutEntry);
+		});
+
+		ProcessMessageQueue<FAudioBusStopMessage>(TraceMessages.BusStopMessages, GetBusEntry,
+		[this, &RemovedKeys](const FAudioBusStopMessage& Msg, TSharedPtr<FSignalFlowDashboardEntry>* OutEntry)
+		{
+			ProcessStopMessage(Msg.DeviceId, Msg.Timestamp, OutEntry, RemovedKeys);
+		});
+
+		ProcessMessageQueue<FSubmixLoadedMessage>(TraceMessages.SubmixLoadedMessages, CreateSubmixEntry,
+		[this](const FSubmixLoadedMessage& Msg, TSharedPtr<FSignalFlowDashboardEntry>* OutEntry)
+		{
+			ProcessNewSubmixEntry(Msg, OutEntry);
+		});
+
+		ProcessMessageQueue<FSubmixAlivePingMessage>(TraceMessages.SubmixAliveMessages, CreateSubmixEntry,
+		[this](const FSubmixAlivePingMessage& Msg, TSharedPtr<FSignalFlowDashboardEntry>* OutEntry)
+		{
+			ProcessNewSubmixEntry(Msg, OutEntry);
+		});
+
+		ProcessMessageQueue<FSubmixUnloadedMessage>(TraceMessages.SubmixUnloadedMessages, GetSubmixEntry,
+		[this, &RemovedKeys](const FSubmixUnloadedMessage& Msg, TSharedPtr<FSignalFlowDashboardEntry>* OutEntry)
+		{
+			ProcessStopMessage(Msg.DeviceId, Msg.Timestamp, OutEntry, RemovedKeys);
+		});
+
+		ProcessMessageQueue<FSoundWaveStartMessage>(TraceMessages.SoundWaveStartMessages, CreateSoundwaveEntry,
+		[this](const FSoundWaveStartMessage& Msg, TSharedPtr<FSignalFlowDashboardEntry>* OutEntry)
+		{
+			ProcessNewSoundwaveEntry(Msg, OutEntry);
+		});
+
+		ProcessMessageQueue<FSoundWaveIsAlivePingMessage>(TraceMessages.SoundWaveIsAlivePingMessages, CreateSoundwaveEntry,
+		[this](const FSoundWaveIsAlivePingMessage& Msg, TSharedPtr<FSignalFlowDashboardEntry>* OutEntry)
+		{
+			ProcessNewSoundwaveEntry(Msg, OutEntry);
+		});
+
+		ProcessMessageQueue<FSoundToBusSendMessage>(TraceMessages.SoundToBusSendMessages, GetSoundwaveEntryForSubmixSendMessage,
+		[this](const FSoundToBusSendMessage& Msg, TSharedPtr<FSignalFlowDashboardEntry>* OutEntry)
+		{
+			const FSignalFlowEntryKey OutputEntryKey{ ESignalFlowEntryType::AudioBus, Msg.ReceiverID };
+			ProcessOutputSendMessage(OutEntry, OutputEntryKey, Msg.SendLevel);
+		});
+
+		ProcessMessageQueue<FSoundToSubmixSendMessage>(TraceMessages.SoundToSubmixSendMessages, GetSoundwaveEntryForSubmixSendMessage,
+		[this](const FSoundToSubmixSendMessage& Msg, TSharedPtr<FSignalFlowDashboardEntry>* OutEntry)
+		{
+			const FSignalFlowEntryKey OutputEntryKey{ ESignalFlowEntryType::Submix, Msg.ReceiverID };
+			ProcessOutputSendMessage(OutEntry, OutputEntryKey, Msg.SendLevel);
+		});
+
+		ProcessMessageQueue<FSubmixToSubmixSendMessage>(TraceMessages.SubmixToSubmixSendMessages, GetSubmixToSubmixEntry,
+		[this](const FSubmixToSubmixSendMessage& Msg, TSharedPtr<FSignalFlowDashboardEntry>* OutEntry)
+		{
+			const FSignalFlowEntryKey OutputEntryKey{ Msg.bIsFinalStage ? ESignalFlowEntryType::AudioDevice : ESignalFlowEntryType::Submix, Msg.ReceiverID };
+			ProcessOutputSendMessage(OutEntry, OutputEntryKey, Msg.SendLevel);
+		});
+
+		ProcessMessageQueue<FSoundWaveStopMessage>(TraceMessages.SoundWaveStopMessages, GetSoundwaveEntryFromMessage,
+		[this, &RemovedKeys, &GetOwnerEntryFromSoundWaveStop](const FSoundWaveStopMessage& Msg, TSharedPtr<FSignalFlowDashboardEntry>* OutEntry)
+		{
+			if (OutEntry && OutEntry->IsValid())
+			{
+				if (TSharedPtr<FSignalFlowDashboardEntry>* OwnerObjectEntry = GetOwnerEntryFromSoundWaveStop(Msg))
+				{
+					if (OwnerObjectEntry->IsValid())
+					{
+						(*OwnerObjectEntry)->Outputs.Remove((*OutEntry)->GetSignalFlowEntryKey());
+					}
+				}
+			}
+
+			ProcessStopMessage(Msg.DeviceId, Msg.Timestamp, OutEntry, RemovedKeys);
+		});
+
+		const TArray<FSoundStopMessage> StopMessages = TraceMessages.SoundStopMessages.DequeueAll();
+		for (const FSoundStopMessage& SoundStop : StopMessages)
+		{
+			// Try find owning object
+			if (TSharedPtr<FSignalFlowDashboardEntry>* OwnerObjectEntry = GetOwnerEntry(SoundStop))
+			{
+				if (OwnerObjectEntry->IsValid())
+				{
+					RemovedKeys.Add((*OwnerObjectEntry)->GetSignalFlowEntryKey());
+				}
+
+				RemoveDeviceEntry(SoundStop.DeviceId, { ESignalFlowEntryType::OwnerObject, SoundStop.ActiveSoundPlayOrder });
+			}
+
+			if (TSortedMap<FSignalFlowEntryKey, TSharedPtr<FSignalFlowDashboardEntry>>* DeviceData = DeviceDataMap.Find(SoundStop.DeviceId))
+			{
+				// Remove the audio device entry if it's the only remaining entry in the map
+				if (DeviceData->Num() == 1 && DeviceData->Contains({ ESignalFlowEntryType::AudioDevice, SoundStop.DeviceId }))
+				{
+					DeviceDataMap.Remove(SoundStop.DeviceId);
+				}
+			}
+		}
+
+		if (!RemovedKeys.IsEmpty())
+		{
+			RemoveInputs(RemovedKeys);
+			bForceGraphRefresh = true;
+		}
+
+		if (bForceGraphRefresh)
+		{
+			OnRequestGraphRefresh.Broadcast();
+			bForceGraphRefresh = false;
+		}
+
+		return true;
+	}
+
+	void FSignalFlowTraceProvider::OnTimingViewTimeMarkerChanged(double TimeMarker)
+	{
+		// Call parent method to update LastMessageId
+		FTraceProviderBase::OnTimingViewTimeMarkerChanged(TimeMarker);
+
+		// TODO - Implement scrubbing
+	}
+
+	TSharedPtr<FSignalFlowDashboardEntry>* FSignalFlowTraceProvider::CreateEntry(const IAudioCachedMessage& Msg, const ::Audio::FDeviceId DeviceID, const ESignalFlowEntryType EntryType, const uint32 NodeID)
+	{
+		TSharedPtr<FSignalFlowDashboardEntry>* ToReturn = nullptr;
+
+		UpdateDeviceEntry(DeviceID, { EntryType, NodeID }, [&ToReturn, &Msg, EntryType, DeviceID, NodeID](TSharedPtr<FSignalFlowDashboardEntry>& Entry)
+		{
+			if (!Entry.IsValid())
+			{
+				Entry = MakeShared<FSignalFlowDashboardEntry>();
+				Entry->SignalFlowNodeID = NodeID;
+				Entry->DeviceId = DeviceID;
+				Entry->EntryType = EntryType;
+			}
+
+			Entry->Timestamp = Msg.Timestamp;
+
+			ToReturn = &Entry;
+		});
+
+		return ToReturn;
+	}
+
+	TSharedPtr<FSignalFlowDashboardEntry>* FSignalFlowTraceProvider::CreateAudioDeviceEntry(const IAudioCachedMessage& Msg, const::Audio::FDeviceId DeviceID)
+	{
+		TSharedPtr<FSignalFlowDashboardEntry>* AudioDeviceEntry = CreateEntry(Msg, DeviceID, ESignalFlowEntryType::AudioDevice, DeviceID);
+		if (AudioDeviceEntry && AudioDeviceEntry->IsValid())
+		{
+			(*AudioDeviceEntry)->Name = FText::Format(LOCTEXT("SignalFlow_AudioDevice", "Audio Device {0}"), static_cast<uint32>(DeviceID)).ToString();
+			(*AudioDeviceEntry)->IconName = FSignalFlowTraceProviderPrivate::AudioDeviceIconName;
+		}
+
+		return AudioDeviceEntry;
+	}
+
+	void FSignalFlowTraceProvider::LinkInput(const FSignalFlowDashboardEntry& Entry, const FSignalFlowEntryKey& OutputKey, const ::Audio::FDeviceId DeviceID)
+	{
+		if (const TSharedPtr<FSignalFlowDashboardEntry>* OutputEntry = FindDeviceEntry(DeviceID, OutputKey))
+		{
+			if (OutputEntry->IsValid())
+			{
+				bool bIsAlreadySet = false;
+				(*OutputEntry)->Inputs.FindOrAdd(Entry.GetSignalFlowEntryKey(), &bIsAlreadySet);
+
+				if (!bIsAlreadySet)
+				{
+					bForceGraphRefresh = true;
+				}
+			}
+		}
+	}
+
+	void FSignalFlowTraceProvider::ProcessNewBusEntry(const FAudioBusStartMessage& Msg, TSharedPtr<FSignalFlowDashboardEntry>* OutEntry)
+	{
+		if (OutEntry && OutEntry->IsValid())
+		{
+			FSignalFlowDashboardEntry& EntryRef = *OutEntry->Get();
+			EntryRef.Name = Msg.Name;
+			EntryRef.IconName = FSignalFlowTraceProviderPrivate::AudioBusIconName;
+			EntryRef.Timestamp = Msg.Timestamp;
+		}
+	}
+
+	void FSignalFlowTraceProvider::ProcessNewSubmixEntry(const FSubmixLoadedMessage& Msg, TSharedPtr<FSignalFlowDashboardEntry>* OutEntry)
+	{
+		if (OutEntry && OutEntry->IsValid())
+		{
+			FSignalFlowDashboardEntry& EntryRef = *OutEntry->Get();
+			EntryRef.Name = Msg.Name;
+			EntryRef.IconName = FSignalFlowTraceProviderPrivate::SubmixIconName;
+			EntryRef.Timestamp = Msg.Timestamp;
+
+			if (Msg.bIsMainSubmix)
+			{
+				MainSubmixKey = { ESignalFlowEntryType::Submix, Msg.SubmixId };
+
+				const FSignalFlowEntryKey OutputAudioDeviceKey = { ESignalFlowEntryType::AudioDevice, Msg.DeviceId };
+
+				EntryRef.Outputs.Add(OutputAudioDeviceKey);
+				LinkInput(EntryRef, OutputAudioDeviceKey, Msg.DeviceId);
+			}
+		}
+	}
+
+	void FSignalFlowTraceProvider::ProcessNewSoundwaveEntry(const FSoundWaveStartMessage& Msg, TSharedPtr<FSignalFlowDashboardEntry>* OutEntry)
+	{
+		auto CreateOwnerEntry = [this](const FSoundWaveStartMessage& Msg)
+		{
+			CreateAudioDeviceEntry(Msg, Msg.DeviceId);
+
+			return CreateEntry(Msg, Msg.DeviceId, ESignalFlowEntryType::OwnerObject, Msg.ActiveSoundPlayOrder);
+		};
+
+		if (OutEntry && OutEntry->IsValid())
+		{
+			FSignalFlowDashboardEntry& EntryRef = *OutEntry->Get();
+			EntryRef.Name = Msg.Name;
+			EntryRef.IconName = FSignalFlowTraceProviderPrivate::GetSourceIconName(Msg.EntryType);
+
+			if (TSharedPtr<FSignalFlowDashboardEntry>* OwnerObjectEntry = CreateOwnerEntry(Msg))
+			{
+				if (OwnerObjectEntry->IsValid())
+				{
+					(*OwnerObjectEntry)->Name = Msg.ActorLabel;
+					(*OwnerObjectEntry)->EntryType = ESignalFlowEntryType::OwnerObject;
+					(*OwnerObjectEntry)->Name = Msg.ActorLabel;
+					(*OwnerObjectEntry)->IconName = FName(Msg.ActorIconName);
+
+					LinkInput((**OwnerObjectEntry), EntryRef.GetSignalFlowEntryKey(), Msg.DeviceId);
+					FSignalFlowOutputData& OutputInfo = (*OwnerObjectEntry)->Outputs.FindOrAdd(EntryRef.GetSignalFlowEntryKey());
+				}
+			}
+		}
+	}
+
+	void FSignalFlowTraceProvider::ProcessStopMessage(const ::Audio::FDeviceId DeviceID, const double MsgTimestamp, TSharedPtr<FSignalFlowDashboardEntry>* OutEntry, TSet<FSignalFlowEntryKey>& OutRemovedKeys)
+	{
+		if (OutEntry && OutEntry->IsValid() && (*OutEntry)->Timestamp < MsgTimestamp)
+		{
+			const FSignalFlowEntryKey EntryKey = (*OutEntry)->GetSignalFlowEntryKey();
+
+			OutRemovedKeys.Add(EntryKey);
+			RemoveDeviceEntry(DeviceID, EntryKey);
+		}
+	}
+
+	void FSignalFlowTraceProvider::ProcessOutputSendMessage(TSharedPtr<FSignalFlowDashboardEntry>* OutEntry, const FSignalFlowEntryKey& OutputEntryKey, const float SendLevel)
+	{
+		if (OutEntry && OutEntry->IsValid())
+		{
+			if (FSignalFlowOutputData* SendData = (*OutEntry)->Outputs.Find(OutputEntryKey))
+			{
+				SendData->SendLevel = SendLevel;
+			}
+			else
+			{
+				(*OutEntry)->Outputs.Add(OutputEntryKey, { SendLevel });
+				bForceGraphRefresh = true;
+			}
+
+			LinkInput(**OutEntry, OutputEntryKey, (*OutEntry)->DeviceId);
+		}
+	}
+
+	void FSignalFlowTraceProvider::RemoveInputs(const TSet<FSignalFlowEntryKey>& RemovedKeys)
+	{
+		for (auto& [DeviceID, Entries] : DeviceDataMap)
+		{
+			for (auto& [EntryKey, Entry] : Entries)
+			{
+				if (!Entry.IsValid())
+				{
+					continue;
+				}
+
+				for (auto Itr = Entry->Inputs.CreateIterator(); Itr; ++Itr)
+				{
+					if (RemovedKeys.Contains(*Itr))
+					{
+						Itr.RemoveCurrent();
+					}
+				}
+			}
+		}
+	}
+} // namespace UE::Audio::Insights
+
+#undef LOCTEXT_NAMESPACE
\ No newline at end of file
diff --git a/Engine/Plugins/AudioInsights/Source/AudioInsights/Private/Providers/SignalFlowTraceProvider.h b/Engine/Plugins/AudioInsights/Source/AudioInsights/Private/Providers/SignalFlowTraceProvider.h
new file mode 100644
--- /dev/null
+++ b/Engine/Plugins/AudioInsights/Source/AudioInsights/Private/Providers/SignalFlowTraceProvider.h
@@ -0,0 +1,45 @@
+// Copyright Epic Games, Inc. All Rights Reserved.
+
+#pragma once
+
+#include "AudioInsightsTraceProviderBase.h"
+#include "Messages/SignalFlowTraceMessages.h"
+
+namespace UE::Audio::Insights
+{
+	class FSignalFlowTraceProvider : public TDeviceDataMapTraceProvider<FSignalFlowEntryKey, TSharedPtr<FSignalFlowDashboardEntry>>, public TSharedFromThis<FSignalFlowTraceProvider>
+	{
+	public:
+		FSignalFlowTraceProvider();
+		virtual ~FSignalFlowTraceProvider() = default;
+
+		static FName GetName_Static();
+		
+		virtual UE::Trace::IAnalyzer* ConstructAnalyzer(TraceServices::IAnalysisSession& InSession) override;
+
+		DECLARE_MULTICAST_DELEGATE(FOnRequestGraphRefresh);
+		FOnRequestGraphRefresh OnRequestGraphRefresh;
+
+	private:
+		virtual bool ProcessMessages() override;
+		virtual void OnTimingViewTimeMarkerChanged(double TimeMarker) override;
+
+		TSharedPtr<FSignalFlowDashboardEntry>* CreateEntry(const IAudioCachedMessage& Msg, const ::Audio::FDeviceId DeviceID, const ESignalFlowEntryType EntryType, const uint32 NodeID);
+		TSharedPtr<FSignalFlowDashboardEntry>* CreateAudioDeviceEntry(const IAudioCachedMessage& Msg, const ::Audio::FDeviceId DeviceID);
+
+		void ProcessNewBusEntry(const FAudioBusStartMessage& Msg, TSharedPtr<FSignalFlowDashboardEntry>* OutEntry);
+		void ProcessNewSubmixEntry(const FSubmixLoadedMessage& Msg, TSharedPtr<FSignalFlowDashboardEntry>* OutEntry);
+		void ProcessNewSoundwaveEntry(const FSoundWaveStartMessage& Msg, TSharedPtr<FSignalFlowDashboardEntry>* OutEntry);
+		void ProcessStopMessage(const ::Audio::FDeviceId DeviceID, const double MsgTimestamp, TSharedPtr<FSignalFlowDashboardEntry>* OutEntry, TSet<FSignalFlowEntryKey>& OutRemovedKeys);
+
+		void ProcessOutputSendMessage(TSharedPtr<FSignalFlowDashboardEntry>* OutEntry, const FSignalFlowEntryKey& OutputEntryKey, const float SendLevel);
+
+		void LinkInput(const FSignalFlowDashboardEntry& Entry, const FSignalFlowEntryKey& OutputKey, const ::Audio::FDeviceId DeviceID);
+		void RemoveInputs(const TSet<FSignalFlowEntryKey>& RemovedKeys);
+
+		FSignalFlowMessages TraceMessages;
+
+		TOptional<FSignalFlowEntryKey> MainSubmixKey;
+		bool bForceGraphRefresh = false;
+	};
+} // namespace UE::Audio::Insights
