[Audio Insights] Create mechanism for dashboards and providers to request the currently filtered audio device ID

#Synced-CL 48901678

diff --git a/Engine/Plugins/AudioInsights/Source/AudioInsights/Public/AudioInsightsTraceProviderBase.h b/Engine/Plugins/AudioInsights/Source/AudioInsights/Public/AudioInsightsTraceProviderBase.h
--- a/Engine/Plugins/AudioInsights/Source/AudioInsights/Public/AudioInsightsTraceProviderBase.h
+++ b/Engine/Plugins/AudioInsights/Source/AudioInsights/Public/AudioInsightsTraceProviderBase.h
@@ -239,15 +239,19 @@ namespace UE::Audio::Insights
 			return DeviceDataMap;
 		}
 
-		const FDeviceData* FindFilteredDeviceData() const
+		::Audio::FDeviceId GetFilteredAudioDeviceID() const
 		{
 #if WITH_EDITOR
 			const IAudioInsightsModule& AudioInsightsModule = IAudioInsightsModule::GetEditorChecked();
-			const ::Audio::FDeviceId AudioDeviceId = AudioInsightsModule.GetDeviceId();
+			return AudioInsightsModule.GetDeviceId();
 #else
-			const ::Audio::FDeviceId AudioDeviceId = GetAudioDeviceIdFromDeviceDataMap();
+			return GetAudioDeviceIdFromDeviceDataMap();
 #endif // WITH_EDITOR
+		}
 
+		const FDeviceData* FindFilteredDeviceData() const
+		{
+			const ::Audio::FDeviceId AudioDeviceId = GetFilteredAudioDeviceID();
 			return DeviceDataMap.Find(AudioDeviceId);
 		}
 
