[Audio Insights] Move Tick function to FTraceDashboardViewFactoryBase - reduces duplication of code inside Table, Tree and now Signal Flow dashboards

#Synced-CL 48901678

diff --git a/Engine/Plugins/AudioInsights/Source/AudioInsights/Private/Views/DashboardViewFactory.cpp b/Engine/Plugins/AudioInsights/Source/AudioInsights/Private/Views/DashboardViewFactory.cpp
new file mode 100644
--- /dev/null
+++ b/Engine/Plugins/AudioInsights/Source/AudioInsights/Private/Views/DashboardViewFactory.cpp
@@ -0,0 +1,74 @@
+// Copyright Epic Games, Inc. All Rights Reserved.
+#include "Views/DashboardViewFactory.h"
+
+namespace UE::Audio::Insights
+{
+	FTraceDashboardViewFactoryBase::FTraceDashboardViewFactoryBase()
+	{
+		TickerHandle = FTSTicker::GetCoreTicker().AddTicker(TEXT("TraceDashboardViewFactoryBase"), 0.0f, [this](float DeltaTime)
+		{
+			Tick(DeltaTime);
+			return true;
+		});
+	}
+
+	FTraceDashboardViewFactoryBase::~FTraceDashboardViewFactoryBase()
+	{
+		FTSTicker::RemoveTicker(TickerHandle);
+	}
+
+	void FTraceDashboardViewFactoryBase::Tick(float InElapsed)
+	{
+		for (const TSharedPtr<FTraceProviderBase>& Provider : Providers)
+		{
+			if (!Provider.IsValid())
+			{
+				continue;
+			}
+
+			if (Provider->ShouldForceUpdate())
+			{
+				Provider->ResetShouldForceUpdate();
+ 
+				UpdateFilterReason = EProcessReason::EntriesUpdated;
+			}
+			else if (const uint64* CurrentUpdateId = UpdateIds.Find(Provider->GetName()))
+			{
+				if (*CurrentUpdateId != Provider->GetLastUpdateId())
+				{
+					UpdateFilterReason = EProcessReason::EntriesUpdated;
+				}
+			}
+			else
+			{
+				UpdateFilterReason = EProcessReason::EntriesUpdated;
+			}
+		}
+
+		if (UpdateFilterReason != EProcessReason::None)
+		{
+			ProcessEntries(UpdateFilterReason);
+			if (UpdateFilterReason == EProcessReason::EntriesUpdated)
+			{
+				for (TSharedPtr<FTraceProviderBase> Provider : Providers)
+				{
+					if (!Provider.IsValid())
+					{
+						continue;
+					}
+
+					const FName ProviderName = Provider->GetName();
+					const uint64 LastUpdateId = Provider->GetLastUpdateId();
+					UpdateIds.FindOrAdd(ProviderName) = LastUpdateId;
+				}
+			}
+
+			RefreshFilteredEntriesListView();
+			UpdateFilterReason = EProcessReason::None;
+		}
+
+#if WITH_EDITOR
+		UpdateDebugDraw(InElapsed);
+#endif // WITH_EDITOR
+	}
+}
\ No newline at end of file
diff --git a/Engine/Plugins/AudioInsights/Source/AudioInsights/Private/Views/TableDashboardViewFactory.cpp b/Engine/Plugins/AudioInsights/Source/AudioInsights/Private/Views/TableDashboardViewFactory.cpp
--- a/Engine/Plugins/AudioInsights/Source/AudioInsights/Private/Views/TableDashboardViewFactory.cpp
+++ b/Engine/Plugins/AudioInsights/Source/AudioInsights/Private/Views/TableDashboardViewFactory.cpp
@@ -38,6 +38,11 @@ namespace UE::Audio::Insights
 		return Factory->GenerateWidgetForColumn(Data->AsShared(), Column);
 	}
 
+	FTraceTableDashboardViewFactory::FTraceTableDashboardViewFactory()
+		: FTraceDashboardViewFactoryBase()
+	{
+	}
+
 	TSharedRef<SWidget> FTraceTableDashboardViewFactory::GenerateWidgetForColumn(TSharedRef<IDashboardDataViewEntry> InRowData, const FName& Column)
 	{
 		const FColumnData& ColumnData = GetColumns()[Column];
@@ -372,62 +377,9 @@ namespace UE::Audio::Insights
 		}
 	}
 
-	const FText& FTraceTableDashboardViewFactory::GetSearchFilterText() const
-	{
-		return SearchBoxFilterText;
-	}
-
-	void FTraceTableDashboardViewFactory::Tick(float InElapsed)
-	{
-		for (TSharedPtr<FTraceProviderBase> Provider : Providers)
-		{
-			if (!Provider.IsValid())
-			{
-				continue;
-			}
-
-			if (Provider->ShouldForceUpdate())
-			{
-				Provider->ResetShouldForceUpdate();
-
-				UpdateFilterReason = EProcessReason::EntriesUpdated;
-			}
-			else if (const uint64* CurrentUpdateId = UpdateIds.Find(Provider->GetName()))
-			{
-				if (*CurrentUpdateId != Provider->GetLastUpdateId())
-				{
-					UpdateFilterReason = EProcessReason::EntriesUpdated;
-				}
-			}
-			else
-			{
-				UpdateFilterReason = EProcessReason::EntriesUpdated;
-			}
-		}
-
-		if (UpdateFilterReason != EProcessReason::None)
-		{
-			ProcessEntries(UpdateFilterReason);
-			if (UpdateFilterReason == EProcessReason::EntriesUpdated)
-			{
-				for (TSharedPtr<FTraceProviderBase> Provider : Providers)
-				{
-					if (!Provider.IsValid())
-					{
-						continue;
-					}
-
-					const FName ProviderName = Provider->GetName();
-					const uint64 LastUpdateId = Provider->GetLastUpdateId();
-					UpdateIds.FindOrAdd(ProviderName) = LastUpdateId;
-				}
-			}
-
-			RefreshFilteredEntriesListView();
-			UpdateFilterReason = EProcessReason::None;
-		}
-
 #if WITH_EDITOR
+	void FTraceTableDashboardViewFactory::UpdateDebugDraw(const float DeltaTime)
+	{
 		if (IsDebugDrawEnabled())
 		{
 			if (FilteredEntriesListView.IsValid())
@@ -436,28 +388,19 @@ namespace UE::Audio::Insights
 
 				if (FAudioDeviceManager* AudioDeviceManager = FAudioDeviceManager::Get())
 				{
-					AudioDeviceManager->IterateOverAllDevices([this, &SelectedItems, InElapsed](::Audio::FDeviceId DeviceId, FAudioDevice* Device)
+					AudioDeviceManager->IterateOverAllDevices([this, &SelectedItems, DeltaTime](::Audio::FDeviceId DeviceId, FAudioDevice* Device)
 					{
-						DebugDraw(InElapsed, SelectedItems, DeviceId);
+						DebugDraw(DeltaTime, SelectedItems, DeviceId);
 					});
 				}
 			}
 		}
-#endif // WITH_EDITOR
-	}
-
-	FTraceTableDashboardViewFactory::FTraceTableDashboardViewFactory()
-	{
-		TickerHandle = FTSTicker::GetCoreTicker().AddTicker(TEXT("TraceTableDashboardViewFactory"), 0.0f, [this](float DeltaTime)
-		{
-			Tick(DeltaTime);
-			return true;
-		});
 	}
+#endif // WITH_EDITOR
 
-	FTraceTableDashboardViewFactory::~FTraceTableDashboardViewFactory()
+	const FText& FTraceTableDashboardViewFactory::GetSearchFilterText() const
 	{
-		FTSTicker::RemoveTicker(TickerHandle);
+		return SearchBoxFilterText;
 	}
 
 	FText FSoundAssetDashboardEntry::GetDisplayName() const
@@ -703,4 +646,4 @@ namespace UE::Audio::Insights
 #endif // WITH_EDITOR
 } // namespace UE::Audio::Insights
 
-#undef LOCTEXT_NAMESPACE
+#undef LOCTEXT_NAMESPACE
\ No newline at end of file
diff --git a/Engine/Plugins/AudioInsights/Source/AudioInsights/Private/Views/TreeDashboardViewFactory.cpp b/Engine/Plugins/AudioInsights/Source/AudioInsights/Private/Views/TreeDashboardViewFactory.cpp
--- a/Engine/Plugins/AudioInsights/Source/AudioInsights/Private/Views/TreeDashboardViewFactory.cpp
+++ b/Engine/Plugins/AudioInsights/Source/AudioInsights/Private/Views/TreeDashboardViewFactory.cpp
@@ -45,23 +45,14 @@ namespace UE::Audio::Insights
 	}
 
 	// FTraceTreeDashboardViewFactory
-	FTraceTreeDashboardViewFactory::FTraceTreeDashboardViewFactory()
-	{
-		TickerHandle = FTSTicker::GetCoreTicker().AddTicker(TEXT("TraceTreeDashboardViewFactory"), 0.0f, [this](float DeltaTime)
-		{
-			Tick(DeltaTime);
-			return true;
-		});
-	}
-
-	FTraceTreeDashboardViewFactory::~FTraceTreeDashboardViewFactory()
+	TSharedPtr<SWidget> FTraceTreeDashboardViewFactory::GetFilterBarButtonWidget()
 	{
-		FTSTicker::RemoveTicker(TickerHandle);
+		return nullptr;
 	}
 
-	TSharedPtr<SWidget> FTraceTreeDashboardViewFactory::GetFilterBarButtonWidget()
+	FTraceTreeDashboardViewFactory::FTraceTreeDashboardViewFactory()
+		: FTraceDashboardViewFactoryBase()
 	{
-		return nullptr;
 	}
 
 	TSharedPtr<SWidget> FTraceTreeDashboardViewFactory::GetFilterBarWidget()
@@ -320,52 +311,9 @@ namespace UE::Audio::Insights
 		}
 	}
 
-	void FTraceTreeDashboardViewFactory::Tick(float InElapsed)
-	{
-		for (const TSharedPtr<FTraceProviderBase>& Provider : Providers)
-		{
-			if (!Provider.IsValid())
-			{
-				continue;
-			}
-
-			if (Provider->ShouldForceUpdate())
-			{
-				Provider->ResetShouldForceUpdate();
-
-				UpdateFilterReason = EProcessReason::EntriesUpdated;
-			}
-			else if (const uint64* FoundCurrentUpdateId = UpdateIds.Find(Provider->GetName()))
-			{
-				if (*FoundCurrentUpdateId != Provider->GetLastUpdateId())
-				{
-					UpdateFilterReason = EProcessReason::EntriesUpdated;
-				}
-			}
-			else
-			{
-				UpdateFilterReason = EProcessReason::EntriesUpdated;
-			}
-		}
-
-		if (UpdateFilterReason != EProcessReason::None)
-		{
-			ProcessEntries(UpdateFilterReason);
-
-			if (UpdateFilterReason == EProcessReason::EntriesUpdated)
-			{
-				for (const TSharedPtr<FTraceProviderBase>& Provider : Providers)
-				{
-					UpdateIds.FindOrAdd(Provider->GetName()) = Provider->GetLastUpdateId();
-				}
-			}
-
-			RefreshFilteredEntriesListView();
-
-			UpdateFilterReason = EProcessReason::None;
-		}
-
 #if WITH_EDITOR
+	void FTraceTreeDashboardViewFactory::UpdateDebugDraw(const float DeltaTime)
+	{
 		if (IsDebugDrawEnabled())
 		{
 			if (FilteredEntriesListView.IsValid())
@@ -374,15 +322,15 @@ namespace UE::Audio::Insights
 
 				if (FAudioDeviceManager* AudioDeviceManager = FAudioDeviceManager::Get())
 				{
-					AudioDeviceManager->IterateOverAllDevices([this, &SelectedItems, InElapsed](::Audio::FDeviceId DeviceId, FAudioDevice* Device)
+					AudioDeviceManager->IterateOverAllDevices([this, &SelectedItems, DeltaTime](::Audio::FDeviceId DeviceId, FAudioDevice* Device)
 					{
-						DebugDraw(InElapsed, SelectedItems, DeviceId);
+						DebugDraw(DeltaTime, SelectedItems, DeviceId);
 					});
 				}
 			}
 		}
-#endif // WITH_EDITOR
 	}
+#endif // WITH_EDITOR
 } // namespace UE::Audio::Insights
 
-#undef LOCTEXT_NAMESPACE
+#undef LOCTEXT_NAMESPACE
\ No newline at end of file
diff --git a/Engine/Plugins/AudioInsights/Source/AudioInsights/Public/Views/DashboardViewFactory.h b/Engine/Plugins/AudioInsights/Source/AudioInsights/Public/Views/DashboardViewFactory.h
--- a/Engine/Plugins/AudioInsights/Source/AudioInsights/Public/Views/DashboardViewFactory.h
+++ b/Engine/Plugins/AudioInsights/Source/AudioInsights/Public/Views/DashboardViewFactory.h
@@ -12,6 +12,8 @@
 class FSpawnTabArgs;
 class SDockTab;
 
+#define UE_API AUDIOINSIGHTS_API
+
 namespace UE::Audio::Insights
 {
 	enum class EDefaultDashboardTabStack : uint8
@@ -39,6 +41,9 @@ namespace UE::Audio::Insights
 	class FTraceDashboardViewFactoryBase : public IDashboardViewFactory
 	{
 	public:
+		UE_API FTraceDashboardViewFactoryBase();
+		UE_API virtual ~FTraceDashboardViewFactoryBase();
+
 		const TArray<TSharedPtr<FTraceProviderBase>>& GetProviders() const
 		{
 			return Providers;
@@ -64,6 +69,30 @@ namespace UE::Audio::Insights
 		}
 
 	protected:
+		enum class EProcessReason : uint8
+		{
+			None,
+			FilterUpdated,
+			EntriesUpdated
+		};
+
+		UE_API virtual void Tick(float InElapsed);
+
+		virtual void ProcessEntries(EProcessReason Reason) { }
+		virtual void RefreshFilteredEntriesListView() { }
+
+#if WITH_EDITOR
+		virtual void UpdateDebugDraw(const float DeltaTime) { }
+#endif // WITH_EDITOR
+
 		TArray<TSharedPtr<FTraceProviderBase>> Providers;
+
+		EProcessReason UpdateFilterReason = EProcessReason::None;
+		FTSTicker::FDelegateHandle TickerHandle;
+
+	private:
+		TMap<FName, uint64> UpdateIds;
 	};
 } // namespace UE::Audio::Insights
+
+#undef UE_API
\ No newline at end of file
diff --git a/Engine/Plugins/AudioInsights/Source/AudioInsights/Public/Views/TableDashboardViewFactory.h b/Engine/Plugins/AudioInsights/Source/AudioInsights/Public/Views/TableDashboardViewFactory.h
--- a/Engine/Plugins/AudioInsights/Source/AudioInsights/Public/Views/TableDashboardViewFactory.h
+++ b/Engine/Plugins/AudioInsights/Source/AudioInsights/Public/Views/TableDashboardViewFactory.h
@@ -65,7 +65,7 @@ namespace UE::Audio::Insights
 	{
 	public:
 		UE_API FTraceTableDashboardViewFactory();
-		UE_API virtual ~FTraceTableDashboardViewFactory();
+		virtual ~FTraceTableDashboardViewFactory() = default;
 
 	protected:
 		struct SRowWidget : public SMultiColumnTableRow<TSharedPtr<IDashboardDataViewEntry>>
@@ -87,18 +87,9 @@ namespace UE::Audio::Insights
 		UE_API virtual TSharedRef<SWidget> GenerateWidgetForColumn(TSharedRef<IDashboardDataViewEntry> InRowData, const FName& Column);
 		UE_API virtual FReply OnDataRowKeyInput(const FGeometry& Geometry, const FKeyEvent& KeyEvent) const;
 		UE_API virtual TSharedRef<SWidget> MakeWidget(TSharedRef<SDockTab> OwnerTab, const FSpawnTabArgs& SpawnTabArgs) override;
-		UE_API virtual void RefreshFilteredEntriesListView();
-
-		enum class EProcessReason : uint8
-		{
-			None,
-			FilterUpdated,
-			EntriesUpdated
-		};
+		UE_API virtual void RefreshFilteredEntriesListView() override;
 
 	protected:
-		UE_API void Tick(float InElapsed);
-
 		template<typename TableProviderType>
 		bool FilterEntries(TFunctionRef<bool(const IDashboardDataViewEntry&)> IsFiltered)
 		{
@@ -190,7 +181,6 @@ namespace UE::Audio::Insights
 		};
 
 		virtual const TMap<FName, FColumnData>& GetColumns() const = 0;
-		virtual void ProcessEntries(EProcessReason Reason) = 0;
 
 		UE_API const FText& GetSearchFilterText() const;
 
@@ -217,13 +207,9 @@ namespace UE::Audio::Insights
 		virtual const FTableRowStyle* GetRowStyle() const { return nullptr; }
 		UE_API virtual FSlateColor GetRowColor(const TSharedPtr<IDashboardDataViewEntry>& InRowDataPtr);
 
-		EProcessReason UpdateFilterReason = EProcessReason::None;
-		FTSTicker::FDelegateHandle TickerHandle;
-
 		FDelegateHandle OnEntriesUpdatedHandle;
 
 		TArray<TSharedPtr<IDashboardDataViewEntry>> DataViewEntries;
-		TMap<FName, uint64> UpdateIds;
 
 		TSharedPtr<SWidget> DashboardWidget;
 		TSharedPtr<SHeaderRow> HeaderRowWidget;
@@ -234,6 +220,10 @@ namespace UE::Audio::Insights
 		EColumnSortMode::Type SortMode = EColumnSortMode::None;
 
 	private:
+#if WITH_EDITOR
+		UE_API virtual void UpdateDebugDraw(const float DeltaTime) override;
+#endif // WITH_EDITOR
+
 		UE_API TSharedRef<SHeaderRow> MakeHeaderRowWidget();
 		UE_API void SetSearchBoxFilterText(const FText& NewText);
 
@@ -268,4 +258,4 @@ namespace UE::Audio::Insights
 	};
 } // namespace UE::Audio::Insights
 
-#undef UE_API
+#undef UE_API
\ No newline at end of file
diff --git a/Engine/Plugins/AudioInsights/Source/AudioInsights/Public/Views/TreeDashboardViewFactory.h b/Engine/Plugins/AudioInsights/Source/AudioInsights/Public/Views/TreeDashboardViewFactory.h
--- a/Engine/Plugins/AudioInsights/Source/AudioInsights/Public/Views/TreeDashboardViewFactory.h
+++ b/Engine/Plugins/AudioInsights/Source/AudioInsights/Public/Views/TreeDashboardViewFactory.h
@@ -37,14 +37,7 @@ namespace UE::Audio::Insights
 	{
 	public:
 		UE_API FTraceTreeDashboardViewFactory();
-		UE_API virtual ~FTraceTreeDashboardViewFactory();
-
-		enum class EProcessReason : uint8
-		{
-			None,
-			FilterUpdated,
-			EntriesUpdated
-		};
+		virtual ~FTraceTreeDashboardViewFactory() = default;
 
 	protected:
 		struct SRowWidget : public SMultiColumnTableRow<TSharedPtr<IDashboardDataTreeViewEntry>>
@@ -98,13 +91,12 @@ namespace UE::Audio::Insights
 
 		UE_API const FText& GetSearchFilterText() const;
 
-		UE_API virtual void RefreshFilteredEntriesListView();
+		UE_API virtual void RefreshFilteredEntriesListView() override;
 
 		virtual const TMap<FName, FHeaderRowColumnData>& GetHeaderRowColumns() const = 0;
 		virtual void OnHiddenColumnsListChanged() {};
 		virtual const TMap<FName, FColumnData>& GetColumns() const = 0;
 
-		virtual void ProcessEntries(EProcessReason Reason) = 0;
 		virtual void SortTable() = 0;
 		virtual bool IsColumnSortable(const FName& ColumnId) const { return true; }
 
@@ -125,8 +117,6 @@ namespace UE::Audio::Insights
 			return false;
 		}
 
-		UE_API void Tick(float InElapsed);
-
 #if WITH_EDITOR
 		virtual bool IsDebugDrawEnabled() const { return false; }
 		virtual void DebugDraw(float InElapsed, const TArray<TSharedPtr<IDashboardDataTreeViewEntry>>& InSelectedItems, ::Audio::FDeviceId InAudioDeviceId) const {};
@@ -180,11 +170,7 @@ namespace UE::Audio::Insights
 			return false;
 		}
 
-		EProcessReason UpdateFilterReason = EProcessReason::None;
-		FTSTicker::FDelegateHandle TickerHandle;
-
 		TArray<TSharedPtr<IDashboardDataTreeViewEntry>> DataViewEntries;
-		TMap<FName, uint64> UpdateIds;
 
 		TSharedPtr<SWidget> DashboardWidget;
 		TSharedPtr<SHeaderRow> HeaderRowWidget;
@@ -194,6 +180,10 @@ namespace UE::Audio::Insights
 		EColumnSortMode::Type SortMode = EColumnSortMode::None;
 
 	private:
+#if WITH_EDITOR
+		virtual void UpdateDebugDraw(const float DeltaTime) override;
+#endif // WITH_EDITOR
+
 		UE_API TSharedRef<SHeaderRow> MakeHeaderRowWidget();
 		UE_API void SetSearchBoxFilterText(const FText& NewText);
 
@@ -209,4 +199,4 @@ namespace UE::Audio::Insights
 	};
 } // namespace UE::Audio::Insights
 
-#undef UE_API
+#undef UE_API
\ No newline at end of file
