[Audio Insights] Signal Flow part 05: Signal flow graph widgets

#Synced-CL 49327150

diff --git a/Engine/Plugins/AudioInsights/Source/AudioInsights/Private/Views/SSignalFlowGraph.cpp b/Engine/Plugins/AudioInsights/Source/AudioInsights/Private/Views/SSignalFlowGraph.cpp
new file mode 100644
--- /dev/null
+++ b/Engine/Plugins/AudioInsights/Source/AudioInsights/Private/Views/SSignalFlowGraph.cpp
@@ -0,0 +1,569 @@
+// Copyright Epic Games, Inc. All Rights Reserved.
+
+#include "Views/SSignalFlowGraph.h"
+
+#include "AudioInsightsStyle.h"
+#include "Internationalization/Text.h"
+#include "Messages/SignalFlowTraceMessages.h"
+#include "Styling/StyleColors.h"
+#include "Views/SignalFlowNodes.h"
+#include "Widgets/Colors/SColorBlock.h"
+#include "Widgets/SBoxPanel.h"
+#include "Widgets/Text/STextBlock.h"
+
+#define LOCTEXT_NAMESPACE "AudioInsights"
+
+namespace UE::Audio::Insights
+{
+	namespace FSignalFlowGraphPrivate
+	{
+		float GetNodeMaxWidth(const EOrientation Orientation)
+		{
+			return Orientation == EOrientation::Orient_Vertical ? 160.0f : 160.0f;
+		}
+
+		float GetNodePaddingX(const EOrientation Orientation)
+		{
+			return Orientation == EOrientation::Orient_Vertical ? 24.0f : 48.0f;
+		}
+
+		constexpr float NodePaddingY = 24.0f;
+		constexpr float NodeMaxHeight = 100.0f;
+	}
+
+	void SSignalFlowGraphNode::Construct(const SSignalFlowGraphNode::FArguments& InArgs)
+	{
+		using namespace FSignalFlowGraphPrivate;
+
+		Entry = InArgs._Entry;
+		PositionInGraph = InArgs._PositionInGraph;
+		XPos = InArgs._XPos;
+		YPos = InArgs._YPos;
+		IsSelected = InArgs._IsSelected;
+		GraphOrientation = InArgs._GraphOrientation;
+		OnNodeSelected = InArgs._OnNodeSelected;
+
+		if (!Entry.IsValid())
+		{
+			return;
+		}
+		// Todo - move the send levels to the connection lines
+		TSharedPtr<SHorizontalBox> SendLevels = SNew(SHorizontalBox);
+
+		for (const auto& [Key, SendInfo] : InArgs._Entry->Outputs)
+		{
+			if (SendInfo.SendLevel == TNumericLimits<float>::Lowest())
+			{
+				continue;
+			}
+
+			if (SendLevels->NumSlots() == 0)
+			{
+				SendLevels->AddSlot()
+				[
+					SNullWidget::NullWidget
+				];
+			}
+
+			SendLevels->AddSlot()
+			.AutoWidth()
+			.Padding(12.0f, 0.0f)
+			.HAlign(HAlign_Center)
+			[
+				SNew(STextBlock)
+				.Text_Lambda([Key, Entry = InArgs._Entry]()
+				{ 
+					if (Entry.IsValid())
+					{
+						if (FSignalFlowOutputData* FoundSendInfo = Entry->Outputs.Find(Key))
+						{
+							return FText::AsNumber(FoundSendInfo->SendLevel);
+						}
+					}
+					return FText::GetEmpty();
+				})
+			];
+		}
+
+		if (SendLevels->NumSlots() > 0)
+		{
+			SendLevels->AddSlot()
+			[
+				SNullWidget::NullWidget
+			];
+		}
+
+		ChildSlot
+		[
+			SNew(SVerticalBox)
+
+			+ SVerticalBox::Slot()
+			.FillContentHeight(1.0f)
+			[
+				SNullWidget::NullWidget
+			]
+
+			+ SVerticalBox::Slot()
+			.AutoHeight()
+			.HAlign(HAlign_Center)
+			[
+				SNew(SBorder)
+				.ColorAndOpacity_Lambda([this]()
+				{
+					return IsSelected.IsBound() && IsSelected.Get() ? FLinearColor::Green : FLinearColor::White;
+				})
+				.Content()
+				[
+					SNew(SOverlay)
+
+					+ SOverlay::Slot()
+					[
+						SNew(SColorBlock)
+						.Color(FStyleColors::Panel.GetSpecifiedColor())
+					]
+
+					+ SOverlay::Slot()
+					[
+						SNew(SVerticalBox)
+
+						+ SVerticalBox::Slot()
+						.Padding(4.0f)
+						.AutoHeight()
+						[
+							SNew(SHorizontalBox)
+
+							+ SHorizontalBox::Slot()
+							.AutoWidth()
+							.Padding(4.0f, 0.0f)
+							[
+								SNew(SImage)
+								.Image(FSlateStyle::Get().CreateIcon(InArgs._Entry->IconName).GetIcon())
+							]
+
+							+ SHorizontalBox::Slot()
+							.FillWidth(1.0f)
+							[
+								SNew(STextBlock)
+								.Text(InArgs._DisplayName)
+								.Justification(ETextJustify::Left)
+								.OverflowPolicy(ETextOverflowPolicy::Ellipsis)
+								.Clipping(EWidgetClipping::ClipToBounds)
+							]
+						]
+
+						// TODO - Move to lines
+						+ SVerticalBox::Slot()
+						.Padding(4.0f)
+						[
+							SendLevels.IsValid() && SendLevels->NumSlots() > 0 ? SendLevels.ToSharedRef() : SNullWidget::NullWidget
+						]
+					]
+				]
+			]
+
+			+ SVerticalBox::Slot()
+			.FillContentHeight(1.0f)
+			[
+				SNullWidget::NullWidget
+			]
+		];
+	}
+
+	FReply SSignalFlowGraphNode::OnMouseButtonDown(const FGeometry& MyGeometry, const FPointerEvent& MouseEvent)
+	{
+		if (MouseEvent.GetEffectingButton() == EKeys::LeftMouseButton)
+		{
+			if (Entry.IsValid())
+			{
+				OnNodeSelected.ExecuteIfBound(IsSelected.IsBound() && IsSelected.Get() ? nullptr : Entry);
+			}
+
+			return FReply::Handled();
+		}
+		else
+		{
+			return SCompoundWidget::OnMouseButtonDown(MyGeometry, MouseEvent);
+		}
+	}
+
+	FVector2D SSignalFlowGraphNode::ComputeDesiredSize(float InLayoutScaleMultiplier) const
+	{
+		using namespace FSignalFlowGraphPrivate;
+
+		const FVector2D Size = SCompoundWidget::ComputeDesiredSize(InLayoutScaleMultiplier);
+		const FVector2D ClampedSize = { FMath::Min(Size.X, GetNodeMaxWidth(GraphOrientation.Get()) - GetNodePaddingX(GraphOrientation.Get())), Size.Y};
+		return ClampedSize;
+	}
+
+	FVector2D SSignalFlowGraphNode::GetLocalInputPos() const
+	{
+		const FVector2D Size = GetDesiredSize();
+
+		if (GraphOrientation.IsBound() && GraphOrientation.Get() == EOrientation::Orient_Horizontal)
+		{
+			// Input pos = left center of node
+			return { CachedPositionInGraph.X, CachedPositionInGraph.Y + (Size.Y / 2.0f) };
+		}
+		else
+		{
+			// Input pos = top center of node
+			return { CachedPositionInGraph.X + (Size.X / 2.0f), CachedPositionInGraph.Y };
+		}
+	}
+	
+	FVector2D SSignalFlowGraphNode::GetLocalOutputPos() const
+	{
+		const FVector2D Size = GetDesiredSize();
+
+		if (GraphOrientation.IsBound() && GraphOrientation.Get() == EOrientation::Orient_Horizontal)
+		{
+			// Out pos = right center of node
+			return { CachedPositionInGraph.X + Size.X, CachedPositionInGraph.Y + (Size.Y / 2.0f) };
+		}
+		else
+		{
+			// Output pos = bottom center of node
+			return { CachedPositionInGraph.X + (Size.X / 2.0f), CachedPositionInGraph.Y + Size.Y };
+		}
+	}
+
+	FVector2D SSignalFlowGraphNode::UpdatePositionInGraph()
+	{
+		CachedPositionInGraph = PositionInGraph.Get();
+		return CachedPositionInGraph;
+	}
+
+	FVector2D SSignalFlowGraphNode::GetCachedPositionInGraph() const
+	{
+		return CachedPositionInGraph;
+	}
+
+	void SSignalFlowGraph::Construct(const SSignalFlowGraph::FArguments& InArgs)
+	{
+		ListSource = InArgs._ListSource;
+		SelectedItem = InArgs._SelectedItem;
+		OnNodeSelected = InArgs._OnNodeSelected;
+
+		Clipping = EWidgetClipping::ClipToBounds;
+
+		Canvas = SNew(SCanvas);
+
+		RefreshGraph();
+
+		ChildSlot
+		[
+			Canvas.ToSharedRef()
+		];
+	}
+
+	void SSignalFlowGraph::RefreshGraph()
+	{
+		using namespace FSignalFlowGraphPrivate;
+
+		if (ListSource == nullptr)
+		{
+			return;
+		}
+
+		// Keep track of any node widgets that were not detected in the new graph
+		// We will clean these up at the end
+		TArray<FSignalFlowEntryKey> RemainingNodeKeys;
+		NodeWidgets.GenerateKeyArray(RemainingNodeKeys);
+
+		DummyEntryToOutNode.Reset();
+
+		TreeDepthPair CurrentTreeDepth{ ESignalFlowEntryType::InvalidLow, TNumericLimits<int32>::Lowest() };
+
+		float MinXPos = TNumericLimits<float>::Max();
+		float MaxXPos = TNumericLimits<float>::Lowest();
+
+		// Gather stats on the width of the graph
+		for (const TSharedPtr<ISignalFlowNode>& Node : *ListSource)
+		{
+			const float XPos = Node->XPos;
+
+			MinXPos = FMath::Min(MinXPos, XPos);
+			MaxXPos = FMath::Max(MaxXPos, XPos);
+		}
+
+		// DepthTracker is used to assign a YPosition to each node
+		// Each time we cross a layer, add one to the depth tracker
+		int32 DepthTracker = 0;
+
+		for (const TSharedPtr<ISignalFlowNode>& Node : *ListSource)
+		{
+			if (CurrentTreeDepth != Node->GetTreeDepth())
+			{
+				CurrentTreeDepth = Node->GetTreeDepth();
+				DepthTracker++;
+			}
+
+			CreateGraphNode(Node, MinXPos, MaxXPos, DepthTracker, RemainingNodeKeys);
+		}
+
+		// Clean up and remove any old nodes that we did not detect in the new graph
+		for (const FSignalFlowEntryKey& Key : RemainingNodeKeys)
+		{
+			TSharedPtr<SSignalFlowGraphNode> RemovedNode = nullptr;
+			NodeWidgets.RemoveAndCopyValue(Key, RemovedNode);
+			
+			if (RemovedNode.IsValid())
+			{
+				Canvas->RemoveSlot(RemovedNode.ToSharedRef());
+			}
+		}
+	}
+
+	int32 SSignalFlowGraph::OnPaint(const FPaintArgs& Args, const FGeometry& AllottedGeometry, const FSlateRect& MyCullingRect, FSlateWindowElementList& OutDrawElements, int32 LayerId, const FWidgetStyle& InWidgetStyle, bool bParentEnabled) const
+	{
+		LayerId = OnPaintConnections(AllottedGeometry, OutDrawElements, LayerId);
+
+		return SCompoundWidget::OnPaint(Args, AllottedGeometry, MyCullingRect, OutDrawElements, LayerId, InWidgetStyle, ShouldBeEnabled(bParentEnabled));
+	}
+
+	int32 SSignalFlowGraph::OnPaintConnections(const FGeometry& AllottedGeometry, FSlateWindowElementList& OutDrawElements, int32 LayerId) const
+	{
+		if (!ListSource)
+		{
+			return LayerId;
+		}
+
+		for (const TSharedPtr<ISignalFlowNode>& Node : *ListSource)
+		{
+			if (!Node->IsRealNode())
+			{
+				continue;
+			}
+
+			const TSharedPtr<SSignalFlowGraphNode>* InputNode = NodeWidgets.Find(Node->GetEntryKey());
+
+			if (InputNode == nullptr || !InputNode->IsValid())
+			{
+				continue;
+			}
+
+			const FVector2D ConnectionStartPos = (*InputNode)->GetLocalOutputPos();
+
+			for (const FSignalFlowEntryKey& NodeOutputKey : Node->FilteredOutputs)
+			{
+				TSharedPtr<SSignalFlowGraphNode> OutputNode = nullptr;
+
+				if (const TSharedPtr<SSignalFlowGraphNode>* FoundOutputNode = NodeWidgets.Find(NodeOutputKey))
+				{
+					if (FoundOutputNode->IsValid())
+					{
+						OutputNode = *FoundOutputNode;
+					}
+				}
+				else
+				{
+					const FSignalFlowEntryKey* FoundDummyOutputKey = DummyEntryToOutNode.Find(NodeOutputKey);
+					if (FoundDummyOutputKey)
+					{
+						const TSharedPtr<SSignalFlowGraphNode>* FoundDummyOutputNode = NodeWidgets.Find(*FoundDummyOutputKey);
+
+						if (FoundDummyOutputNode && FoundDummyOutputNode->IsValid())
+						{
+							OutputNode = *FoundDummyOutputNode;
+						}
+					}
+				}
+
+				if (OutputNode.IsValid())
+				{
+					const FVector2D ConnectionEndPos = OutputNode->GetLocalInputPos();
+					if (Orientation == EOrientation::Orient_Vertical && ConnectionEndPos.Y <= ConnectionStartPos.Y)
+					{
+						continue;
+					}
+					else if (Orientation == EOrientation::Orient_Horizontal && ConnectionEndPos.X <= ConnectionStartPos.X)
+					{
+						continue;
+					}
+
+					constexpr float Radius = 120.0f;
+					const FVector2d StartDirection = Orientation == EOrientation::Orient_Horizontal ? FVector2d{ Radius, 0.0f } : FVector2d{ 0.0f, Radius };
+					const FVector2d EndDirection = Orientation == EOrientation::Orient_Horizontal ? FVector2d{ Radius, 0.0f } : FVector2d{ 0.0f, Radius };
+
+					FSlateDrawElement::MakeSpline(
+						OutDrawElements,
+						++LayerId,
+						AllottedGeometry.ToPaintGeometry(),
+						ConnectionStartPos,
+						StartDirection,
+						ConnectionEndPos,
+						EndDirection,
+						2.0f,
+						ESlateDrawEffect::NoPixelSnapping,
+						FStyleColors::AccentGray.GetSpecifiedColor()
+					);
+				}
+			}
+		}
+
+		return LayerId;
+	}
+
+	FReply SSignalFlowGraph::OnKeyDown(const FGeometry& MyGeometry, const FKeyEvent& InKeyEvent)
+	{
+		if (InKeyEvent.GetKey() == EKeys::F)
+		{
+			Orientation = (Orientation == EOrientation::Orient_Vertical) ? EOrientation::Orient_Horizontal : EOrientation::Orient_Vertical;
+			RefreshGraph();
+			return FReply::Handled();
+		}
+
+		return SCompoundWidget::OnKeyDown(MyGeometry, InKeyEvent);
+	}
+
+	FVector2D SSignalFlowGraph::ComputeDesiredSize(float InLayoutScaleMultiplier) const
+	{
+		using namespace FSignalFlowGraphPrivate;
+
+		float GraphWidth = 0.0f;
+		float GraphHeight = 0.0f;
+		for (const auto& [Key, NodeWidget] : NodeWidgets)
+		{
+			if (!NodeWidget.IsValid())
+			{
+				continue;
+			}
+
+			if (Orientation == Orient_Vertical)
+			{
+				GraphWidth = FMath::Max(GraphWidth, (NodeWidget->GetXPos() + 1) * GetNodeMaxWidth(Orientation));
+				GraphHeight = FMath::Max(GraphHeight, (NodeWidget->GetYPos() + 1) * NodeMaxHeight);
+			}
+			else
+			{
+				GraphWidth = FMath::Max(GraphWidth, (NodeWidget->GetYPos() + 1) * GetNodeMaxWidth(Orientation));
+				GraphHeight = FMath::Max(GraphHeight, (NodeWidget->GetXPos() + 1) * NodeMaxHeight);
+			}
+		}
+
+		return FVector2D{ GraphWidth, GraphHeight };
+	}
+
+	void SSignalFlowGraph::CreateGraphNode(const TSharedPtr<ISignalFlowNode>& Node, const float MinXPos, const float MaxXPos, const int32 DepthTracker, TArray<FSignalFlowEntryKey>& OutRemainingNodeKeys)
+	{
+		using namespace FSignalFlowGraphPrivate;
+
+		if (!Node.IsValid())
+		{
+			return;
+		}
+
+		if (!Node->IsRealNode())
+		{
+			// If this is a fake node, keep a note of it's key and connection end point for fast lookup later
+			const TSharedPtr<FDummyConnectionNode> FoundDummyCast = StaticCastSharedPtr<FDummyConnectionNode>(Node);
+			if (FoundDummyCast.IsValid())
+			{
+				DummyEntryToOutNode.Add(FoundDummyCast->GetEntryKey(), FoundDummyCast->ConnectionOutputKey);
+				return;
+			}
+		}
+		
+
+		const FSignalFlowEntryKey EntryKey = Node->GetEntryKey();
+
+		// Avoid creating this node again if it already exists
+		if (TSharedPtr<SSignalFlowGraphNode>* NodeWidget = NodeWidgets.Find(EntryKey))
+		{
+			if (NodeWidget->IsValid())
+			{
+				// Update the X and Y positions of the node
+				(*NodeWidget)->SetXPos(Node->XPos - MinXPos);
+				(*NodeWidget)->SetYPos(static_cast<float>(DepthTracker));
+
+				OutRemainingNodeKeys.RemoveSwap(EntryKey);
+				return;
+			}
+			else
+			{
+				NodeWidgets.Remove(EntryKey);
+			}
+		}
+
+		const TSharedPtr<FSignalFlowEntryNode> EntryNode = StaticCastSharedPtr<FSignalFlowEntryNode>(Node);
+
+		TSharedPtr<SSignalFlowGraphNode> NodeWidget = SNew(SSignalFlowGraphNode)
+			.Entry(EntryNode->Entry)
+			.ListSource(ListSource)
+			.DisplayName(EntryNode->Entry.IsValid() ? EntryNode->Entry->GetDisplayName() : FText::GetEmpty())
+			.PositionInGraph_Lambda([this, EntryKey]() -> FVector2D
+			{
+				const TSharedPtr<SSignalFlowGraphNode>* Node = NodeWidgets.Find(EntryKey);
+				if (Node && Node->IsValid())
+				{
+					const FVector2D NodeSize = (*Node)->GetDesiredSize();
+					const float XOffset = (GetNodeMaxWidth(Orientation) - NodeSize.X) / 2.0f;
+					const float YOffset = (NodeMaxHeight - NodeSize.Y) / 2.0f;
+
+					if (Orientation == EOrientation::Orient_Vertical)
+					{
+						const float XPos = ((*Node)->GetXPos() * GetNodeMaxWidth(Orientation)) + XOffset;
+						const float YPos = (*Node)->GetYPos() * NodeMaxHeight + YOffset;
+
+						const FVector2D NodePosition = { XPos, YPos };
+						return NodePosition;
+					}
+					else
+					{
+						const float XPos = ((*Node)->GetYPos() * GetNodeMaxWidth(Orientation)) + XOffset;
+						const float YPos = (*Node)->GetXPos() * NodeMaxHeight + YOffset;
+
+						const FVector2D NodePosition = { XPos, YPos };
+						return NodePosition;
+					}
+				}
+
+				return FVector2D::ZeroVector;
+			})
+			.XPos(Node->XPos - MinXPos)
+			.YPos(static_cast<float>(DepthTracker))
+			.IsSelected_Lambda([this, EntryKey]()
+			{
+				return (SelectedItem.IsBound() && SelectedItem.Get().IsValid() && SelectedItem.Get()->GetSignalFlowEntryKey() == EntryKey);
+			})
+			.GraphOrientation_Lambda([this]()
+			{
+				return Orientation;
+			})
+			.OnNodeSelected_Lambda([this](TSharedPtr<FSignalFlowDashboardEntry> SelectedEntry)
+			{
+				OnNodeSelected.ExecuteIfBound(SelectedEntry);
+			});
+
+		NodeWidgets.Add(Node->GetEntryKey(), NodeWidget);
+
+		Canvas->AddSlot()
+			.Position_Lambda([this, EntryKey]() -> FVector2D
+			{
+				const TSharedPtr<SSignalFlowGraphNode>* Node = NodeWidgets.Find(EntryKey);
+				if (Node && Node->IsValid())
+				{
+					return (*Node)->UpdatePositionInGraph();
+				}
+
+				return FVector2D::ZeroVector;
+			})
+			.Size_Lambda([this, EntryKey]() -> FVector2D
+			{
+				const TSharedPtr<SSignalFlowGraphNode>* Node = NodeWidgets.Find(EntryKey);
+				if (Node && Node->IsValid())
+				{
+					return (*Node)->GetDesiredSize();
+				}
+
+				return FVector2D::ZeroVector;
+			})
+			.HAlign(HAlign_Left)
+			.VAlign(VAlign_Top)
+			[
+				NodeWidget.ToSharedRef()
+			];
+	}
+} // namespace UE::Audio::Insights
+
+#undef LOCTEXT_NAMESPACE
diff --git a/Engine/Plugins/AudioInsights/Source/AudioInsights/Private/Views/SignalFlowNodes.cpp b/Engine/Plugins/AudioInsights/Source/AudioInsights/Private/Views/SignalFlowNodes.cpp
new file mode 100644
--- /dev/null
+++ b/Engine/Plugins/AudioInsights/Source/AudioInsights/Private/Views/SignalFlowNodes.cpp
@@ -0,0 +1,13 @@
+// Copyright Epic Games, Inc. All Rights Reserved.
+#include "Views/SignalFlowNodes.h"
+
+#include "Messages/SignalFlowTraceMessages.h"
+
+namespace UE::Audio::Insights
+{
+	FSignalFlowEntryKey FSignalFlowEntryNode::GetEntryKey() const
+	{
+		ensure(Entry.IsValid());
+		return Entry.IsValid() ? Entry->GetSignalFlowEntryKey() : FSignalFlowEntryKey();
+	}
+} // namespace UE::Audio::Insights
\ No newline at end of file
diff --git a/Engine/Plugins/AudioInsights/Source/AudioInsights/Public/Messages/SignalFlowEntryKey.h b/Engine/Plugins/AudioInsights/Source/AudioInsights/Public/Messages/SignalFlowEntryKey.h
--- a/Engine/Plugins/AudioInsights/Source/AudioInsights/Public/Messages/SignalFlowEntryKey.h
+++ b/Engine/Plugins/AudioInsights/Source/AudioInsights/Public/Messages/SignalFlowEntryKey.h
@@ -7,13 +7,19 @@
 
 namespace UE::Audio::Insights
 {
-	enum class ESignalFlowEntryType : uint8
+	// Which render stage does this node belong to.
+	// This enum is in the order the render stages occur - from root to audio device
+	enum class ESignalFlowEntryType : int32
 	{
-		OwnerObject,
+		OwnerObject = 0,
 		SoundSource,
 		AudioBus,
 		Submix,
-		AudioDevice
+		AudioDevice,
+
+		///
+		InvalidLow = TNumericLimits<int32>::Lowest(),
+		InvalidMax = TNumericLimits<int32>::Max()
 	};
 
 	struct FSignalFlowEntryKey
@@ -27,7 +33,7 @@ namespace UE::Audio::Insights
 
 		friend uint32 GetTypeHash(const FSignalFlowEntryKey& Key)
 		{
-			return HashCombine(GetTypeHash(static_cast<uint32>(Key.EntryType)), GetTypeHash(Key.EntryID));
+			return HashCombine(GetTypeHash(static_cast<int32>(Key.EntryType)), GetTypeHash(Key.EntryID));
 		}
 
 		bool operator==(const FSignalFlowEntryKey& InOther) const
@@ -38,7 +44,7 @@ namespace UE::Audio::Insights
 		bool operator<(const FSignalFlowEntryKey& InOther) const
 		{
 			return EntryType == InOther.EntryType ? EntryID < InOther.EntryID 
-												  : static_cast<uint32>(EntryType) < static_cast<uint32>(InOther.EntryType);
+												  : static_cast<int32>(EntryType) < static_cast<int32>(InOther.EntryType);
 		}
 
 		ESignalFlowEntryType EntryType = ESignalFlowEntryType::SoundSource;
diff --git a/Engine/Plugins/AudioInsights/Source/AudioInsights/Public/Views/SSignalFlowGraph.h b/Engine/Plugins/AudioInsights/Source/AudioInsights/Public/Views/SSignalFlowGraph.h
new file mode 100644
--- /dev/null
+++ b/Engine/Plugins/AudioInsights/Source/AudioInsights/Public/Views/SSignalFlowGraph.h
@@ -0,0 +1,118 @@
+// Copyright Epic Games, Inc. All Rights Reserved.
+#pragma once
+
+#include "Delegates/Delegate.h"
+#include "Messages/SignalFlowEntryKey.h"
+#include "Widgets/SCanvas.h"
+#include "Widgets/SCompoundWidget.h"
+
+#define UE_API AUDIOINSIGHTS_API
+
+namespace UE::Audio::Insights
+{
+	struct FSignalFlowDashboardEntry;
+	class ISignalFlowNode;
+
+	DECLARE_DELEGATE_OneParam(FOnSignalFlowNodeSelected, TSharedPtr<FSignalFlowDashboardEntry>);
+
+	class SSignalFlowGraphNode : public SCompoundWidget
+	{
+	public:
+		SLATE_BEGIN_ARGS(SSignalFlowGraphNode)
+			: _Entry(nullptr),
+			_ListSource(nullptr),
+			_DisplayName(),
+			_PositionInGraph(FVector2d::ZeroVector),
+			_XPos(0.0f),
+			_YPos(0.0f),
+			_IsSelected(false),
+			_GraphOrientation(EOrientation::Orient_Vertical)
+		{
+		}
+
+		SLATE_ARGUMENT(TSharedPtr<FSignalFlowDashboardEntry>, Entry)
+		SLATE_ARGUMENT(TArray<TSharedPtr<ISignalFlowNode>>*, ListSource)
+		SLATE_ARGUMENT(FText, DisplayName)
+		SLATE_ATTRIBUTE(FVector2d, PositionInGraph)
+		SLATE_ATTRIBUTE(float, XPos)
+		SLATE_ATTRIBUTE(float, YPos)
+		SLATE_ATTRIBUTE(bool, IsSelected)
+		SLATE_ATTRIBUTE(EOrientation, GraphOrientation)
+		SLATE_EVENT(FOnSignalFlowNodeSelected, OnNodeSelected)
+
+		SLATE_END_ARGS()
+
+		UE_API void Construct(const FArguments& InArgs);
+
+		UE_API virtual FReply OnMouseButtonDown(const FGeometry& MyGeometry, const FPointerEvent& MouseEvent) override;
+		UE_API virtual FVector2D ComputeDesiredSize(float InLayoutScaleMultiplier) const override;
+
+		UE_API FVector2D GetLocalInputPos() const;
+		UE_API FVector2D GetLocalOutputPos() const;
+
+		UE_API FVector2D UpdatePositionInGraph();
+		UE_API FVector2D GetCachedPositionInGraph() const;
+
+		float GetXPos() const { return XPos.Get(); }
+		float GetYPos() const { return YPos.Get(); }
+
+		void SetXPos(const float Pos) { XPos = Pos; }
+		void SetYPos(const float Pos) { YPos = Pos; }
+
+	private:
+		TSharedPtr<FSignalFlowDashboardEntry> Entry;
+		TAttribute<FVector2d> PositionInGraph;
+		TAttribute<float> XPos;
+		TAttribute<float> YPos;
+		TAttribute<bool> IsSelected;
+		TAttribute<EOrientation> GraphOrientation;
+
+		FOnSignalFlowNodeSelected OnNodeSelected;
+		FVector2d CachedPositionInGraph = FVector2d::Zero();
+	};
+
+	class SSignalFlowGraph : public SCompoundWidget
+	{
+	public:
+		SLATE_BEGIN_ARGS(SSignalFlowGraph)
+			: _ListSource(nullptr)
+		{
+		}
+
+		SLATE_ARGUMENT(TArray<TSharedPtr<ISignalFlowNode>>*, ListSource)
+		SLATE_ATTRIBUTE(TSharedPtr<FSignalFlowDashboardEntry>, SelectedItem)
+		SLATE_EVENT(FOnSignalFlowNodeSelected, OnNodeSelected)
+
+		SLATE_END_ARGS()
+
+		UE_API void Construct(const FArguments& InArgs);
+		UE_API void RefreshGraph();
+
+		UE_API virtual int32 OnPaint(const FPaintArgs& Args, const FGeometry& AllottedGeometry, const FSlateRect& MyCullingRect, FSlateWindowElementList& OutDrawElements, int32 LayerId, const FWidgetStyle& InWidgetStyle, bool bParentEnabled) const override;
+
+		UE_API virtual FReply OnKeyDown(const FGeometry& MyGeometry, const FKeyEvent& InKeyEvent) override;
+		UE_API virtual bool SupportsKeyboardFocus() const override { return true; }
+
+	protected:
+		UE_API virtual FVector2D ComputeDesiredSize(float InLayoutScaleMultiplier) const override;
+
+	private:
+		int32 OnPaintConnections(const FGeometry& AllottedGeometry, FSlateWindowElementList& OutDrawElements, int32 LayerId) const;
+		void CreateGraphNode(const TSharedPtr<ISignalFlowNode>& Node, const float MinXPos, const float MaxXPos, const int32 DepthTracker, TArray<FSignalFlowEntryKey>& OutRemainingNodeKeys);
+
+		TSharedPtr<SCanvas> Canvas;
+
+		TMap<FSignalFlowEntryKey, TSharedPtr<SSignalFlowGraphNode>> NodeWidgets;
+		TMap<FSignalFlowEntryKey, FSignalFlowEntryKey> DummyEntryToOutNode;
+
+		TArray<TSharedPtr<ISignalFlowNode>>* ListSource;
+		TAttribute<TSharedPtr<FSignalFlowDashboardEntry>> SelectedItem;
+
+		FOnSignalFlowNodeSelected OnNodeSelected;
+		EOrientation Orientation = EOrientation::Orient_Vertical;
+
+		bool bNeedsRefresh = false;
+	};
+} // namespace UE::Audio::Insights
+
+#undef UE_API
\ No newline at end of file
diff --git a/Engine/Plugins/AudioInsights/Source/AudioInsights/Public/Views/SignalFlowNodes.h b/Engine/Plugins/AudioInsights/Source/AudioInsights/Public/Views/SignalFlowNodes.h
new file mode 100644
--- /dev/null
+++ b/Engine/Plugins/AudioInsights/Source/AudioInsights/Public/Views/SignalFlowNodes.h
@@ -0,0 +1,127 @@
+// Copyright Epic Games, Inc. All Rights Reserved.
+#pragma once
+
+#include "Containers/Array.h"
+#include "CoreTypes.h"
+#include "Messages/SignalFlowEntryKey.h"
+#include "Templates/SharedPointer.h"
+
+#define UE_API AUDIOINSIGHTS_API
+
+namespace UE::Audio::Insights
+{
+	struct FSignalFlowDashboardEntry;
+
+	// TreeDepthPair - Simple helper used to manage tree depth inside the signal flow graph
+	// LayerID - the layer the node sits in
+	// TreeDepth - what row inside the layer the node sits in
+	struct TreeDepthPair
+	{
+		ESignalFlowEntryType LayerID = ESignalFlowEntryType::OwnerObject;
+		int32 TreeDepth = 0;
+
+		bool operator==(const TreeDepthPair& Other) const
+		{
+			return LayerID == Other.LayerID && TreeDepth == Other.TreeDepth;
+		}
+
+		bool operator<(const TreeDepthPair& Other) const
+		{
+			if (LayerID != Other.LayerID)
+			{
+				return static_cast<int32>(LayerID) < static_cast<int32>(Other.LayerID);
+			}
+
+			return TreeDepth < Other.TreeDepth;
+		}
+
+		friend int32 GetTypeHash(const TreeDepthPair& Key)
+		{
+			return HashCombine(GetTypeHash(static_cast<int32>(Key.LayerID)), GetTypeHash(Key.TreeDepth));
+		}
+	};
+
+	// ISignalFlowNode : Class that holds data on a node's position, order, inputs, outputs etc.
+	// Specifically holds data on the filtered nodes in the graph
+	class ISignalFlowNode
+	{
+	public:
+		ISignalFlowNode() = delete;
+		ISignalFlowNode(const int32 InTreeDepth, const ESignalFlowEntryType InTreeDepthLayerGroup, const double InTimestamp)
+			: TreeDepthPair({ InTreeDepthLayerGroup, InTreeDepth })
+			, Timestamp(InTimestamp)
+		{
+		}
+
+		virtual ~ISignalFlowNode() = default;
+
+		virtual FSignalFlowEntryKey GetEntryKey() const = 0;
+		virtual bool IsRealNode() const = 0;
+
+		void SetTreeDepth(const ESignalFlowEntryType InTreeDepthLayerGroup, const int32 InTreeDepth) { TreeDepthPair = { InTreeDepthLayerGroup, InTreeDepth }; }
+		const TreeDepthPair& GetTreeDepth() const { return TreeDepthPair; }
+		TreeDepthPair& GetTreeDepth() { return TreeDepthPair; }
+
+		void SetNodeOrderID(const int32 InNodeOrderID) { NodeOrderID = InNodeOrderID; }
+		void ResetOrderID() { NodeOrderID = INVALID_NODE_ORDER_ID; }
+		int32 GetNodeOrderID() const { return NodeOrderID; }
+		bool NodeOrderIsValid() const { return NodeOrderID != INVALID_NODE_ORDER_ID; }
+
+		double GetTimestamp() const { return Timestamp; }
+
+		TArray<FSignalFlowEntryKey> FilteredInputs;
+		TArray<FSignalFlowEntryKey> FilteredOutputs;
+
+		float XPos = 0.0f;
+
+	private:
+		static constexpr int32 INVALID_NODE_ORDER_ID = -1;
+
+		TreeDepthPair TreeDepthPair;
+		int32 NodeOrderID = INVALID_NODE_ORDER_ID;
+
+		const double Timestamp = 0.0;
+	};
+
+	// FSignalFlowEntryNode : A node with a valid entry - directly maps to a FSignalFlowDashboardEntry created inside FSignalFlowTraceProvider
+	class FSignalFlowEntryNode : public ISignalFlowNode
+	{
+	public:
+		FSignalFlowEntryNode() = delete;
+		FSignalFlowEntryNode(TSharedPtr<FSignalFlowDashboardEntry> InEntry, const int32 InTreeDepth, const ESignalFlowEntryType InTreeDepthLayerGroup, const double InTimestamp)
+			: ISignalFlowNode(InTreeDepth, InTreeDepthLayerGroup, InTimestamp)
+			, Entry(InEntry)
+		{
+		}
+
+		UE_API virtual FSignalFlowEntryKey GetEntryKey() const override;
+
+		virtual bool IsRealNode() const override { return true; }
+
+		TSharedPtr<FSignalFlowDashboardEntry> Entry;
+	};
+
+	// FDummyConnectionNode : A fake invisible node placed between the input and output of two other nodes
+	// Used to help position and order real nodes whilst avoiding nodes getting in each other's way and crossing connections
+	class FDummyConnectionNode : public ISignalFlowNode
+	{
+	public:
+		FDummyConnectionNode() = delete;
+		FDummyConnectionNode(const FSignalFlowEntryKey& InEntryKey, const int32 InTreeDepth, const ESignalFlowEntryType InTreeDepthLayerGroup, const double InTimestamp)
+			: ISignalFlowNode(InTreeDepth, InTreeDepthLayerGroup, InTimestamp)
+			, EntryKey(InEntryKey)
+		{
+		}
+
+		virtual FSignalFlowEntryKey GetEntryKey() const override { return EntryKey; }
+		virtual bool IsRealNode() const override { return false; }
+	
+		FSignalFlowEntryKey ConnectionInputKey;
+		FSignalFlowEntryKey ConnectionOutputKey;
+
+	private:
+		const FSignalFlowEntryKey EntryKey;
+	};
+} // namespace UE::Audio::Insights
+
+#undef UE_API
\ No newline at end of file
