[Audio Insights] Fix GetNodeMaxWidth having unrequired ternary operator

#Synced-CL 49327150

diff --git a/Engine/Plugins/AudioInsights/Source/AudioInsights/Private/Views/SSignalFlowGraph.cpp b/Engine/Plugins/AudioInsights/Source/AudioInsights/Private/Views/SSignalFlowGraph.cpp
--- a/Engine/Plugins/AudioInsights/Source/AudioInsights/Private/Views/SSignalFlowGraph.cpp
+++ b/Engine/Plugins/AudioInsights/Source/AudioInsights/Private/Views/SSignalFlowGraph.cpp
@@ -17,9 +17,9 @@ namespace UE::Audio::Insights
 {
 	namespace FSignalFlowGraphPrivate
 	{
-		float GetNodeMaxWidth(const EOrientation Orientation)
+		float GetNodeMaxWidth()
 		{
-			return Orientation == EOrientation::Orient_Vertical ? 160.0f : 160.0f;
+			return 160.0f;
 		}
 
 		float GetNodePaddingX(const EOrientation Orientation)
@@ -191,7 +191,7 @@ namespace UE::Audio::Insights
 		using namespace FSignalFlowGraphPrivate;
 
 		const FVector2D Size = SCompoundWidget::ComputeDesiredSize(InLayoutScaleMultiplier);
-		const FVector2D ClampedSize = { FMath::Min(Size.X, GetNodeMaxWidth(GraphOrientation.Get()) - GetNodePaddingX(GraphOrientation.Get())), Size.Y};
+		const FVector2D ClampedSize = { FMath::Min(Size.X, GetNodeMaxWidth() - GetNodePaddingX(GraphOrientation.Get())), Size.Y};
 		return ClampedSize;
 	}
 
@@ -431,12 +431,12 @@ namespace UE::Audio::Insights
 
 			if (Orientation == Orient_Vertical)
 			{
-				GraphWidth = FMath::Max(GraphWidth, (NodeWidget->GetXPos() + 1) * GetNodeMaxWidth(Orientation));
+				GraphWidth = FMath::Max(GraphWidth, (NodeWidget->GetXPos() + 1) * GetNodeMaxWidth());
 				GraphHeight = FMath::Max(GraphHeight, (NodeWidget->GetYPos() + 1) * NodeMaxHeight);
 			}
 			else
 			{
-				GraphWidth = FMath::Max(GraphWidth, (NodeWidget->GetYPos() + 1) * GetNodeMaxWidth(Orientation));
+				GraphWidth = FMath::Max(GraphWidth, (NodeWidget->GetYPos() + 1) * GetNodeMaxWidth());
 				GraphHeight = FMath::Max(GraphHeight, (NodeWidget->GetXPos() + 1) * NodeMaxHeight);
 			}
 		}
@@ -497,12 +497,12 @@ namespace UE::Audio::Insights
 				if (Node && Node->IsValid())
 				{
 					const FVector2D NodeSize = (*Node)->GetDesiredSize();
-					const float XOffset = (GetNodeMaxWidth(Orientation) - NodeSize.X) / 2.0f;
+					const float XOffset = (GetNodeMaxWidth() - NodeSize.X) / 2.0f;
 					const float YOffset = (NodeMaxHeight - NodeSize.Y) / 2.0f;
 
 					if (Orientation == EOrientation::Orient_Vertical)
 					{
-						const float XPos = ((*Node)->GetXPos() * GetNodeMaxWidth(Orientation)) + XOffset;
+						const float XPos = ((*Node)->GetXPos() * GetNodeMaxWidth()) + XOffset;
 						const float YPos = (*Node)->GetYPos() * NodeMaxHeight + YOffset;
 
 						const FVector2D NodePosition = { XPos, YPos };
@@ -510,7 +510,7 @@ namespace UE::Audio::Insights
 					}
 					else
 					{
-						const float XPos = ((*Node)->GetYPos() * GetNodeMaxWidth(Orientation)) + XOffset;
+						const float XPos = ((*Node)->GetYPos() * GetNodeMaxWidth()) + XOffset;
 						const float YPos = (*Node)->GetXPos() * NodeMaxHeight + YOffset;
 
 						const FVector2D NodePosition = { XPos, YPos };
