[Audio Insights] Signal Flow part 06: Dashboard
* Dashboard takes the entries created by the provider and generates three selectable lists for active sources/buses/submixes.
* It also creates a list of ISignalFlowNodes based off the current filter, and generates the ordering and positioning for the graph.
* These nodes are then passed into SSignalFlowGraph to render the node graph.

#Synced-CL 49327150

diff --git a/Engine/Plugins/AudioInsights/Source/AudioInsights/Private/Views/SignalFlowDashboardViewFactory.cpp b/Engine/Plugins/AudioInsights/Source/AudioInsights/Private/Views/SignalFlowDashboardViewFactory.cpp
new file mode 100644
--- /dev/null
+++ b/Engine/Plugins/AudioInsights/Source/AudioInsights/Private/Views/SignalFlowDashboardViewFactory.cpp
@@ -0,0 +1,1144 @@
+// Copyright Epic Games, Inc. All Rights Reserved.
+#include "Views/SignalFlowDashboardViewFactory.h"
+
+#include "AudioInsightsModule.h"
+#include "AudioInsightsStyle.h"
+#include "AudioInsightsTimingViewExtender.h"
+#include "IAudioInsightsModule.h"
+#include "Internationalization/Text.h"
+#include "Providers/SignalFlowTraceProvider.h"
+#include "Styling/StyleColors.h"
+#include "Views/SignalFlowNodes.h"
+#include "Widgets/Colors/SColorBlock.h"
+#include "Widgets/Layout/SScrollBox.h"
+
+#if WITH_EDITOR
+#include "Editor.h"
+#include "Subsystems/AssetEditorSubsystem.h"
+#else
+#include "AudioInsightsComponent.h"
+#endif // WITH_EDITOR
+
+#define LOCTEXT_NAMESPACE "AudioInsights"
+
+namespace UE::Audio::Insights
+{
+	namespace FSignalFlowDashboardViewFactoryPrivate
+	{
+		constexpr float GraphWidthProp = 0.8f;
+		constexpr float SourcesListHeightProp = 0.5f;
+		constexpr float BusListHeightProp = 0.2f;
+		constexpr float SubmixesListHeightProp = 0.3f;
+	}
+
+	FSignalFlowDashboardViewFactory::FSignalFlowDashboardViewFactory()
+	{
+		IAudioInsightsTraceModule& AudioInsightsTraceModule = IAudioInsightsModule::GetChecked().GetTraceModule();
+
+		SignalFlowProvider = MakeShared<FSignalFlowTraceProvider>();
+
+		AudioInsightsTraceModule.AddTraceProvider(SignalFlowProvider);
+
+		Providers = TArray<TSharedPtr<FTraceProviderBase>>
+		{
+			SignalFlowProvider
+		};
+
+		SignalFlowProvider->OnRequestGraphRefresh.AddRaw(this, &FSignalFlowDashboardViewFactory::OnRequestGraphRefresh);
+	}
+
+	FSignalFlowDashboardViewFactory::~FSignalFlowDashboardViewFactory()
+	{
+		const TSharedPtr<FSignalFlowTraceProvider> Provider = FindProvider<FSignalFlowTraceProvider>();
+		if (Provider.IsValid())
+		{
+			Provider->OnRequestGraphRefresh.RemoveAll(this);
+		}
+	}
+
+	FName FSignalFlowDashboardViewFactory::GetName() const
+	{
+		return "SignalFlow";
+	}
+
+	FText FSignalFlowDashboardViewFactory::GetDisplayName() const
+	{
+		return LOCTEXT("AudioDashboard_SignalFlow_DisplayName", "Signal Flow");
+	}
+	
+	FSlateIcon FSignalFlowDashboardViewFactory::GetIcon() const
+	{
+		return FSlateStyle::Get().CreateIcon("AudioInsights.Icon.Submix"); // TODO: we should request a unique icon for this dashboard
+	}
+
+	EDefaultDashboardTabStack FSignalFlowDashboardViewFactory::GetDefaultTabStack() const
+	{
+		return EDefaultDashboardTabStack::Analysis;
+	}
+
+	TSharedRef<SWidget> FSignalFlowDashboardViewFactory::MakeWidget(TSharedRef<SDockTab> OwnerTab, const FSpawnTabArgs& SpawnTabArgs)
+	{
+		using namespace FSignalFlowDashboardViewFactoryPrivate;
+
+		if (!DashboardWidget.IsValid())
+		{
+			SAssignNew(DashboardWidget, SHorizontalBox)
+
+			+ SHorizontalBox::Slot()
+			.FillWidth(1.0f - GraphWidthProp)
+			[
+				CreateSelectionPanel()
+			]
+			
+			+ SHorizontalBox::Slot()
+			.FillWidth(GraphWidthProp)
+			[
+				CreateGraphSection()
+			];
+		}
+
+		return DashboardWidget->AsShared();
+	}
+
+	void FSignalFlowDashboardViewFactory::SRowWidget::Construct(const FArguments& InArgs, const TSharedRef<STableViewBase>& InOwnerTable, TSharedPtr<FSignalFlowDashboardEntry> Entry)
+	{
+		if (!Entry.IsValid())
+		{
+			return;
+		}
+
+		RowEntry = Entry;
+
+		STableRow<TSharedPtr<FSignalFlowDashboardEntry>>::Construct
+		(
+			STableRow<TSharedPtr<FSignalFlowDashboardEntry>>::FArguments()
+			.Style(InArgs._Style)
+			.Content()
+			[
+				SNew(STextBlock)
+					.Text(RowEntry->GetDisplayName())
+			]
+			, InOwnerTable
+		);
+	}
+
+	const FTableRowStyle* FSignalFlowDashboardViewFactory::GetRowStyle() const
+	{
+		return &FSlateStyle::Get().GetWidgetStyle<FTableRowStyle>("TreeDashboard.TableViewRow");
+	}
+
+	TSharedRef<SWidget> FSignalFlowDashboardViewFactory::CreateSelectionPanel()
+	{
+		using namespace FSignalFlowDashboardViewFactoryPrivate;
+
+		ActiveSourcesListView = CreateSelectableList(LOCTEXT("AudioDashboard_SignalFlow_ActiveSourcesTitle", "Active Sources"), ActiveSources);
+		ActiveBusesListView = CreateSelectableList(LOCTEXT("AudioDashboard_SignalFlow_ActiveBusesTitle", "Active Buses"), ActiveBuses);
+		ActiveSubmixesListView = CreateSelectableList(LOCTEXT("AudioDashboard_SignalFlow_ActiveSubmixesTitle", "Active Submixes"), ActiveSubmixes);
+
+		return SNew(SVerticalBox)
+
+		+ SVerticalBox::Slot()
+		.FillHeight(SourcesListHeightProp)
+		[
+			ActiveSourcesListView.ToSharedRef()
+		]
+			
+		+ SVerticalBox::Slot()
+		.FillHeight(BusListHeightProp)
+		[
+			ActiveBusesListView.ToSharedRef()
+		]
+		
+		+ SVerticalBox::Slot()
+		.FillHeight(SubmixesListHeightProp)
+		[
+			ActiveSubmixesListView.ToSharedRef()
+		];
+	}
+
+	TSharedRef<SWidget> FSignalFlowDashboardViewFactory::CreateGraphSection()
+	{
+		return SNew(SScrollBox)
+		.Orientation(EOrientation::Orient_Vertical)
+
+		+ SScrollBox::Slot()
+		.FillSize(1.0f)
+		.HAlign(HAlign_Fill)
+		.VAlign(VAlign_Fill)
+		[
+			SNew(SScrollBox)
+			.Orientation(EOrientation::Orient_Horizontal)
+
+			+ SScrollBox::Slot()
+			.FillSize(1.0f)
+			.HAlign(HAlign_Fill)
+			.VAlign(VAlign_Fill)
+			[
+				SNew(SOverlay)
+
+				// Graph background
+				+ SOverlay::Slot()
+				[
+					SNew(SColorBlock)
+					.Color(FStyleColors::Background.GetSpecifiedColor())
+				]
+
+				// Graph
+				+ SOverlay::Slot()
+				.HAlign(HAlign_Center)
+				.VAlign(VAlign_Center)
+				[
+					SAssignNew(GraphWidgetContents, SSignalFlowGraph)
+					.ListSource(&FilteredGraphNodes)
+					.SelectedItem_Lambda([this]() { return SelectedItem; })
+					.OnNodeSelected_Lambda([this](TSharedPtr<FSignalFlowDashboardEntry> SelectedNodeEntry)
+					{
+						if (SelectedNodeEntry.IsValid())
+						{
+							switch (SelectedNodeEntry->EntryType)
+							{
+								case ESignalFlowEntryType::SoundSource:
+									if (ActiveSourcesListView.IsValid())
+									{
+										ActiveSourcesListView->SetSelection(SelectedNodeEntry);
+									}
+									break;
+
+								case ESignalFlowEntryType::AudioBus:
+									if (ActiveBusesListView.IsValid())
+									{
+										ActiveBusesListView->SetSelection(SelectedNodeEntry);
+									}
+									break;
+
+								case ESignalFlowEntryType::Submix:
+									if (ActiveSubmixesListView.IsValid())
+									{
+										ActiveSubmixesListView->SetSelection(SelectedNodeEntry);
+									}
+									break;
+
+								default:
+									if (ActiveSourcesListView.IsValid())
+									{
+										ActiveSourcesListView->ClearSelection();
+									}
+									
+									if (ActiveBusesListView.IsValid())
+									{
+										ActiveBusesListView->ClearSelection();
+									}
+									
+									if (ActiveSubmixesListView.IsValid())
+									{
+										ActiveSubmixesListView->ClearSelection();
+									}
+									break;
+							}
+						}
+
+						OnSelectionChanged(SelectedNodeEntry, ESelectInfo::OnMouseClick);
+					})
+				]
+			]
+		];
+	}
+
+	TSharedRef<SListView<TSharedPtr<FSignalFlowDashboardEntry>>> FSignalFlowDashboardViewFactory::CreateSelectableList(const FText& HeadingText, const TArray<TSharedPtr<FSignalFlowDashboardEntry>>& ListSource)
+	{
+		TSharedRef<SHeaderRow> HeaderRowWidget = SNew(SHeaderRow)
+												 .CanSelectGeneratedColumn(false);
+
+		static const FLazyName HeadingColumnID("Heading");
+		const SHeaderRow::FColumn::FArguments NameColumnArgs = SHeaderRow::Column(HeadingColumnID)
+																		   .DefaultLabel(HeadingText)
+																		   .HAlignCell(HAlign_Left)
+																		   .FillWidth(1.0f);
+
+		HeaderRowWidget->AddColumn(NameColumnArgs);
+
+		return SNew(SListView<TSharedPtr<FSignalFlowDashboardEntry>>)
+			.ListItemsSource(&ListSource)
+			.SelectionMode(ESelectionMode::SingleToggle)
+			.OnGenerateRow_Lambda([this](TSharedPtr<FSignalFlowDashboardEntry> Item, const TSharedRef<STableViewBase>& OwnerTable)
+			{
+				return SNew(SRowWidget, OwnerTable, Item)
+					   .Style(GetRowStyle());
+			})
+			.OnSelectionChanged(this, &FSignalFlowDashboardViewFactory::OnSelectionChanged)
+			.HeaderRow
+			(
+				HeaderRowWidget
+			);
+	}
+
+	void FSignalFlowDashboardViewFactory::OnSelectionChanged(TSharedPtr<FSignalFlowDashboardEntry> InSelectedItem, ESelectInfo::Type SelectInfo)
+	{
+		if (!ActiveSourcesListView.IsValid() || !ActiveBusesListView.IsValid() || !ActiveSubmixesListView.IsValid())
+		{
+			return;
+		}
+
+		if (SelectInfo == ESelectInfo::Direct)
+		{
+			return;
+		}
+
+		FilteredEntries.Reset();
+
+		if (InSelectedItem.IsValid())
+		{
+			switch (InSelectedItem->EntryType)
+			{
+				case ESignalFlowEntryType::SoundSource:
+				{
+					ActiveBusesListView->ClearSelection();
+					ActiveSubmixesListView->ClearSelection();
+					break;
+				}
+				case ESignalFlowEntryType::AudioBus:
+				{
+					ActiveSourcesListView->ClearSelection();
+					ActiveSubmixesListView->ClearSelection();
+					break;
+				}
+				case ESignalFlowEntryType::Submix:
+				{
+					ActiveSourcesListView->ClearSelection();
+					ActiveBusesListView->ClearSelection();
+					break;
+				}
+			}
+
+			SelectedItem = InSelectedItem;
+			if (SelectedItem.IsValid())
+			{
+				if (TSharedPtr<ISignalFlowNode>* SelectedNode = FilteredEntries.Find(SelectedItem->GetSignalFlowEntryKey()))
+				{
+					if (SelectedNode->IsValid())
+					{
+						(*SelectedNode)->SetTreeDepth((*SelectedNode)->GetTreeDepth().LayerID, 0);
+					}
+				}
+			}
+		}
+		else
+		{
+			SelectedItem.Reset();
+			ActiveSourcesListView->ClearSelection();
+			ActiveBusesListView->ClearSelection();
+			ActiveSubmixesListView->ClearSelection();
+		}
+
+		OnRequestGraphRefresh();
+	}
+
+	void FSignalFlowDashboardViewFactory::ProcessEntries(EProcessReason Reason)
+	{
+		if (Reason != EProcessReason::EntriesUpdated)
+		{
+			return;
+		}
+
+		const TSharedPtr<const FSignalFlowTraceProvider> Provider = FindProvider<const FSignalFlowTraceProvider>();
+		if (!Provider.IsValid())
+		{
+			return;
+		}
+
+		if (const FSignalFlowTraceProvider::FDeviceData* DeviceData = Provider->FindFilteredDeviceData())
+		{
+			BuildSelectionMenus(*DeviceData, Reason);
+		}
+	}
+
+	void FSignalFlowDashboardViewFactory::BuildSelectionMenus(const TSortedMap<FSignalFlowEntryKey, TSharedPtr<FSignalFlowDashboardEntry>>& DeviceData, const EProcessReason Reason)
+	{
+		ActiveSources.Reset();
+		ActiveBuses.Reset();
+		ActiveSubmixes.Reset();
+
+		// Build the left-hand menus of selectable nodes from this device's data
+		for (const auto& [Key, Entry] : DeviceData)
+		{
+			if (!Entry.IsValid())
+			{
+				continue;
+			}
+
+			switch (Entry->EntryType)
+			{
+				case ESignalFlowEntryType::SoundSource:
+					ActiveSources.Add(Entry);
+					break;
+				case ESignalFlowEntryType::AudioBus:
+					ActiveBuses.Add(Entry);
+					break;
+				case ESignalFlowEntryType::Submix:
+					ActiveSubmixes.Add(Entry);
+					break;
+				default:
+					break;
+			}
+		}
+	}
+
+	void FSignalFlowDashboardViewFactory::RefreshFilteredEntriesListView()
+	{
+		if (ActiveSourcesListView.IsValid())
+		{
+			ActiveSourcesListView->RequestListRefresh();
+		}
+
+		if (ActiveBusesListView.IsValid())
+		{
+			ActiveBusesListView->RequestListRefresh();
+		}
+
+		if (ActiveSubmixesListView.IsValid())
+		{
+			ActiveSubmixesListView->RequestListRefresh();
+		}
+
+		if (GraphWidgetContents.IsValid() && bGraphRequiresRefresh)
+		{
+			GraphWidgetContents->RefreshGraph();
+			bGraphRequiresRefresh = false;
+		}
+	}
+
+	void FSignalFlowDashboardViewFactory::OnRequestGraphRefresh()
+	{
+		const TSharedPtr<const FSignalFlowTraceProvider> Provider = FindProvider<const FSignalFlowTraceProvider>();
+
+		if (!Provider.IsValid())
+		{
+			return;
+		}
+
+		if (SelectedItem.IsValid())
+		{
+			CreateFilteredNodeGraph(*Provider, SelectedItem);
+		}
+		else
+		{
+			AddAllNodesToGraph();
+		}
+	}
+
+	void FSignalFlowDashboardViewFactory::CreateFilteredNodeGraph(const FSignalFlowTraceProvider& Provider, TSharedPtr<FSignalFlowDashboardEntry> FilterFromEntry)
+	{
+		if (!FilterFromEntry.IsValid())
+		{
+			return;
+		}
+
+		const FSignalFlowTraceProvider::FDeviceData* DeviceData = Provider.FindFilteredDeviceData();
+		if (DeviceData == nullptr)
+		{
+			return;
+		}
+
+		bGraphRequiresRefresh = true;
+
+		FilteredEntries.Reset();
+
+		AddFilteredInputsRecursive(*DeviceData, FilterFromEntry, nullptr, 0);
+		AddFilteredOutputs(*DeviceData, FilterFromEntry, 1);
+
+		// Gather stats (such as root entry and tree depth structure)
+		TOptional<FSignalFlowEntryKey> RootEntryKey;
+		TArray<TreeDepthPair> SortedTreeDepths;
+		CalculateFilteredGraphDepthStructure(RootEntryKey, SortedTreeDepths);
+
+		if (!RootEntryKey.IsSet())
+		{
+			return;
+		}
+
+		const TSharedPtr<ISignalFlowNode>* RootNode = FilteredEntries.Find(RootEntryKey.GetValue());
+
+		if (RootNode && RootNode->IsValid())
+		{
+			// Grab a handle to the shared ptr directly - the structure of the FilteredEntries map may change inside CreateDummyInputs
+			// so a raw ptr handle may point to a different entry (or garbage) afterwards
+			TSharedPtr<ISignalFlowNode> RootNodeSharedPtr = *RootNode;
+
+			// Create dummy nodes in places where a connection between two nodes spans multiple rows.
+			// This helps when trying to maintain good X-positioning and ordering
+			CreateDummyInputs(SortedTreeDepths);
+
+			// Now start at the root node and iterate upwards via inputs to assign an order to each node
+			// This helps keep related nodes together
+			int32 NodeProcessingOrder = 0;
+			AssignNodeOrderRecursive(RootNodeSharedPtr, 0, NodeProcessingOrder);
+
+			// Run through again and find any nodes that did not get an order assigned
+			// This can happen if a node is not the root node but has no outputs
+			FixUnassignedNodeOrders();
+
+			// Generate the final node array that will be passed on to the SSignalFlowGraph widget
+			FilteredEntries.GenerateValueArray(FilteredGraphNodes);
+
+			// Sort the new array to be in tree depth/node connection order
+			SortNodeOrder();
+
+			// Assign X positions based off the new order
+			// X positions given in units of 1.0 that will be scaled up by SSignalFlowGraph later
+			CalcXPositions();
+		}
+	}
+
+	void FSignalFlowDashboardViewFactory::AddAllNodesToGraph()
+	{
+		const TSharedPtr<const FSignalFlowTraceProvider> Provider = FindProvider<const FSignalFlowTraceProvider>();
+		if (!Provider.IsValid())
+		{
+			return;
+		}
+
+		if (const FSignalFlowTraceProvider::FDeviceData* DeviceData = Provider->FindFilteredDeviceData())
+		{
+			if (!DeviceData->IsEmpty())
+			{
+				const ::Audio::FDeviceId DeviceID = Provider->GetFilteredAudioDeviceID();
+				const TSharedPtr<FSignalFlowDashboardEntry>* DeviceEntry = DeviceData->Find({ ESignalFlowEntryType::AudioDevice, DeviceID });
+
+				if (DeviceEntry && DeviceEntry->IsValid())
+				{
+					CreateFilteredNodeGraph(*Provider, *DeviceEntry);
+				}
+			}
+		}
+	}
+
+	void FSignalFlowDashboardViewFactory::AddFilteredInputsRecursive(const TSortedMap<FSignalFlowEntryKey, TSharedPtr<FSignalFlowDashboardEntry>>& DeviceData, TSharedPtr<FSignalFlowDashboardEntry> Entry, TSharedPtr<FSignalFlowDashboardEntry> ChildEntry, int32 TreeDepth)
+	{
+		if (!Entry.IsValid())
+		{
+			return;
+		}
+
+		TSharedPtr<ISignalFlowNode> EntryNode = nullptr;
+
+		TSharedPtr<ISignalFlowNode>* FoundNode = FilteredEntries.Find(Entry->GetSignalFlowEntryKey());
+		if (FoundNode && FoundNode->IsValid())
+		{
+			const ESignalFlowEntryType EntryType = (*FoundNode)->GetEntryKey().EntryType;
+			const bool bCanBeLayered = EntryType == ESignalFlowEntryType::AudioBus || EntryType == ESignalFlowEntryType::Submix;
+
+			TreeDepthPair& NodeTreeDepth = (*FoundNode)->GetTreeDepth();
+
+			if (bCanBeLayered && NodeTreeDepth.TreeDepth >= TreeDepth)
+			{
+				NodeTreeDepth.TreeDepth = TreeDepth - 1;
+			}
+
+			EntryNode = *FoundNode;
+		}
+		else
+		{
+			EntryNode = MakeShared<FSignalFlowEntryNode>(Entry, TreeDepth, Entry->EntryType, Entry->Timestamp);
+
+			FilteredEntries.Add(Entry->GetSignalFlowEntryKey(), EntryNode);
+		}
+
+		if (EntryNode.IsValid() && ChildEntry.IsValid())
+		{
+			EntryNode->FilteredOutputs.AddUnique(ChildEntry->GetSignalFlowEntryKey());
+		}
+
+		for (const FSignalFlowEntryKey& InputKey : Entry->Inputs)
+		{
+			const TSharedPtr<FSignalFlowDashboardEntry>* InputEntry = DeviceData.Find(InputKey);
+
+			if (InputEntry && InputEntry->IsValid())
+			{
+				if (EntryNode.IsValid())
+				{
+					EntryNode->FilteredInputs.AddUnique((*InputEntry)->GetSignalFlowEntryKey());
+				}
+
+				bool bResetTreeDepth = (*InputEntry)->EntryType != Entry->EntryType;
+				AddFilteredInputsRecursive(DeviceData, *InputEntry, Entry, bResetTreeDepth ? 0 : TreeDepth - 1);
+			}
+		}
+	}
+
+	void FSignalFlowDashboardViewFactory::AddFilteredOutputs(const TSortedMap<FSignalFlowEntryKey, TSharedPtr<FSignalFlowDashboardEntry>>& DeviceData, TSharedPtr<FSignalFlowDashboardEntry> Entry, int32 TreeDepth)
+	{
+		for (const auto& [OutputKey, SendInfo] : Entry->Outputs)
+		{
+			const TSharedPtr<FSignalFlowDashboardEntry>* OutputEntry = DeviceData.Find(OutputKey);
+			if (OutputEntry && OutputEntry->IsValid())
+			{
+				AddFilteredOutputsRecursive(DeviceData, *OutputEntry, Entry, TreeDepth);
+
+				TSharedPtr<ISignalFlowNode>* EntryNode = FilteredEntries.Find(Entry->GetSignalFlowEntryKey());
+				TSharedPtr<ISignalFlowNode>* OutputNode = FilteredEntries.Find(OutputKey);
+				if (EntryNode && OutputNode && EntryNode->IsValid() && OutputNode->IsValid())
+				{
+					(*EntryNode)->FilteredOutputs.AddUnique(OutputKey);
+				}
+			}
+		}
+	}
+
+	void FSignalFlowDashboardViewFactory::AddFilteredOutputsRecursive(const TSortedMap<FSignalFlowEntryKey, TSharedPtr<FSignalFlowDashboardEntry>>& DeviceData, TSharedPtr<FSignalFlowDashboardEntry> Entry, TSharedPtr<FSignalFlowDashboardEntry> ParentEntry, int32 TreeDepth)
+	{
+		if (!Entry.IsValid())
+		{
+			return;
+		}
+
+		const FSignalFlowEntryKey& EntryKey = Entry->GetSignalFlowEntryKey();
+
+		TSharedPtr<ISignalFlowNode> Node = nullptr;
+		TSharedPtr<ISignalFlowNode>* FoundNode = FilteredEntries.Find(EntryKey);
+
+		if (FoundNode && FoundNode->IsValid())
+		{
+			TreeDepthPair& NodeTreeDepth = (*FoundNode)->GetTreeDepth();
+			if (NodeTreeDepth.TreeDepth <= TreeDepth)
+			{
+				NodeTreeDepth.TreeDepth = TreeDepth + 1;
+			}
+
+			Node = (*FoundNode);
+		}
+		else
+		{
+			Node = MakeShared<FSignalFlowEntryNode>(Entry, TreeDepth, Entry->EntryType, Entry->Timestamp);
+
+			FilteredEntries.Add(EntryKey, Node);
+		}
+
+		if (Node.IsValid() && ParentEntry.IsValid())
+		{
+			Node->FilteredInputs.AddUnique(ParentEntry->GetSignalFlowEntryKey());
+		}
+
+		for (const auto& [OutputKey, SendInfo] : Entry->Outputs)
+		{
+			const TSharedPtr<FSignalFlowDashboardEntry>* OutputEntry = DeviceData.Find(OutputKey);
+			if (OutputEntry && OutputEntry->IsValid())
+			{
+				if (Node.IsValid())
+				{
+					Node->FilteredOutputs.AddUnique(OutputKey);
+				}
+
+				bool bResetTreeDepth = (*OutputEntry)->EntryType != Entry->EntryType;
+
+				AddFilteredOutputsRecursive(DeviceData, *OutputEntry, Entry, bResetTreeDepth ? 0 : TreeDepth + 1);
+			}
+		}
+	}
+
+	void FSignalFlowDashboardViewFactory::CalculateFilteredGraphDepthStructure(TOptional<FSignalFlowEntryKey>& OutRootEntryKey, TArray<TreeDepthPair>& OutSortedTreeDepths) const
+	{
+		OutSortedTreeDepths.Reset();
+
+		TreeDepthPair HighestTreeDepth = { ESignalFlowEntryType::InvalidLow, TNumericLimits<int32>::Lowest() };
+
+		for (auto& [Key, Node] : FilteredEntries)
+		{
+			if (!Node.IsValid())
+			{
+				continue;
+			}
+
+			Node->ResetOrderID();
+
+			const TreeDepthPair& NodeTreeDepth = Node->GetTreeDepth();
+			if (!OutSortedTreeDepths.Contains(NodeTreeDepth))
+			{
+				OutSortedTreeDepths.Add(NodeTreeDepth);
+			}
+
+			if (HighestTreeDepth < NodeTreeDepth)
+			{
+				OutRootEntryKey = Key;
+				HighestTreeDepth = NodeTreeDepth;
+			}
+		}
+
+		auto SortByTreeDepth = [this](const TreeDepthPair& First, const TreeDepthPair& Second)
+		{
+			return First < Second;
+		};
+
+		if (!OutSortedTreeDepths.IsEmpty())
+		{
+			OutSortedTreeDepths.Sort(SortByTreeDepth);
+		}
+	}
+
+	void FSignalFlowDashboardViewFactory::CreateDummyInputs(const TArray<TreeDepthPair>& TreeDepthStructure)
+	{
+		TArray<TSharedPtr<ISignalFlowNode>> Nodes;
+		FilteredEntries.GenerateValueArray(Nodes);
+
+		for (const TSharedPtr<ISignalFlowNode>& Node : Nodes)
+		{
+			const TreeDepthPair& CurrentTreeDepthPair = Node->GetTreeDepth();
+
+			for (const FSignalFlowEntryKey& InputKey : Node->FilteredInputs)
+			{
+				const TSharedPtr<ISignalFlowNode>* FoundInput = FilteredEntries.Find(InputKey);
+				if (FoundInput && FoundInput->IsValid())
+				{
+					ensure((*FoundInput)->IsRealNode());
+
+					const TSharedPtr<FSignalFlowEntryNode> FoundInputSharedPtr = StaticCastSharedPtr<FSignalFlowEntryNode>(*FoundInput);
+					if (!FoundInputSharedPtr.IsValid())
+					{
+						continue;
+					}
+
+					const TreeDepthPair& TargetTreeDepthPair = (*FoundInput)->GetTreeDepth();
+
+					const int32 TargetTreeDepthIndex = TreeDepthStructure.IndexOfByKey(TargetTreeDepthPair);
+					int32 CurrentTreeDepthIndex = TreeDepthStructure.IndexOfByKey(CurrentTreeDepthPair) - 1;
+
+					ensure(CurrentTreeDepthIndex >= TargetTreeDepthIndex);
+
+					const FSignalFlowEntryKey ConnectionOutputKey = Node->GetEntryKey();
+					TSharedPtr<ISignalFlowNode> OutputNode = Node;
+
+					// Add invisible "dummy connection" nodes between the input and output for every row between them
+					while (CurrentTreeDepthIndex > TargetTreeDepthIndex)
+					{
+						OutputNode = CreateDummyConnectionNode(FoundInputSharedPtr, OutputNode, TreeDepthStructure[CurrentTreeDepthIndex], ConnectionOutputKey);
+
+						CurrentTreeDepthIndex--;
+					}
+				}
+			}
+		}
+	}
+
+	TSharedPtr<FDummyConnectionNode> FSignalFlowDashboardViewFactory::CreateDummyConnectionNode(const TSharedPtr<FSignalFlowEntryNode>& InputNode, const TSharedPtr<ISignalFlowNode>& OutputNode, const TreeDepthPair& TreeDepthPair, const FSignalFlowEntryKey& ConnectionOutputNodeKey)
+	{
+		// Input nodes should only be real nodes in the graph, and should have valid entries
+		// Output nodes can be real, or be dummy connections
+		if (!InputNode.IsValid() || !InputNode->Entry.IsValid() || !OutputNode.IsValid())
+		{
+			return nullptr;
+		}
+
+		const FSignalFlowEntryKey InputKey = InputNode->GetEntryKey();
+		const FSignalFlowEntryKey OutputKey = OutputNode->GetEntryKey();
+
+		// It should be guarenteed for the OutputNode to contain the InputKey in it's filtered inputs
+		if (!ensure(OutputNode->FilteredInputs.Contains(InputKey)))
+		{
+			return nullptr;
+		}
+
+		// Create a dummy node in between the input and output nodes
+		// We can use these dummy connection nodes to help order and position real nodes in the graph and reduce connections crossing
+		const FSignalFlowEntryKey DummyConnectionEntryKey = FSignalFlowEntryKey(TreeDepthPair.LayerID, HashCombine(GetTypeHash(InputKey), TreeDepthPair.TreeDepth));
+		TSharedPtr<FDummyConnectionNode> DummyConnectionNode = MakeShared<FDummyConnectionNode>(DummyConnectionEntryKey, TreeDepthPair.TreeDepth, TreeDepthPair.LayerID, InputNode->Entry->Timestamp);
+		
+		DummyConnectionNode->ConnectionOutputKey = ConnectionOutputNodeKey;
+		DummyConnectionNode->ConnectionInputKey = InputKey;
+
+		// Link the dummy node into the graph
+		DummyConnectionNode->FilteredInputs.AddUnique(InputKey);
+		DummyConnectionNode->FilteredOutputs.AddUnique(OutputKey);
+
+		// Unlink the original nodes and point them towards the dummy
+		InputNode->FilteredOutputs.Remove(OutputKey);
+		InputNode->FilteredOutputs.AddUnique(DummyConnectionEntryKey);
+		OutputNode->FilteredInputs[OutputNode->FilteredInputs.IndexOfByKey(InputKey)] = DummyConnectionEntryKey;
+
+		// Add the new dummy node to the filtered entries map
+		FilteredEntries.Add(DummyConnectionEntryKey, DummyConnectionNode);
+
+		return DummyConnectionNode;
+	}
+
+	void FSignalFlowDashboardViewFactory::AssignNodeOrderRecursive(TSharedPtr<ISignalFlowNode> OutNode, const int32 Depth, int32& OutNodeProcessingOrder)
+	{
+		// Keep track of the order nodes were reached when iterating through tree inputs recursively
+		// This helps keep nodes that are linked to the same input together in the graph
+		// We sort the inputs before proceeding forward to keep a good order in the graph
+		
+		// Only assign a new processing order if the node has not been reached yet
+		if (!OutNode->NodeOrderIsValid())
+		{
+			OutNode->SetNodeOrderID(OutNodeProcessingOrder);
+			OutNodeProcessingOrder++;
+		}
+
+		SortNodeInputs(OutNode);
+
+		for (const FSignalFlowEntryKey InputKey : OutNode->FilteredInputs)
+		{
+			const TSharedPtr<ISignalFlowNode>* FoundInput = FilteredEntries.Find(InputKey);
+		
+			if (FoundInput && FoundInput->IsValid())
+			{
+				AssignNodeOrderRecursive(*FoundInput, Depth + 1, OutNodeProcessingOrder);
+			}
+		}
+	}
+
+	void FSignalFlowDashboardViewFactory::SortNodeInputs(TSharedPtr<ISignalFlowNode> Node)
+	{
+		if (!Node.IsValid())
+		{
+			return;
+		}
+
+		auto SortInput = [this](const FSignalFlowEntryKey& First, const FSignalFlowEntryKey& Second)
+		{
+			const TSharedPtr<ISignalFlowNode>* FirstNode = FilteredEntries.Find(First);
+			const TSharedPtr<ISignalFlowNode>* SecondNode = FilteredEntries.Find(Second);
+
+			if (FirstNode && SecondNode && FirstNode->IsValid() && SecondNode->IsValid())
+			{
+				const TreeDepthPair FirstNodeTreeDepth = GetOriginTreeDepthForNode((*FirstNode));
+				const TreeDepthPair SecondNodeTreeDepth = GetOriginTreeDepthForNode((*SecondNode));
+
+				// First, sort by tree depth
+				if (FirstNodeTreeDepth != SecondNodeTreeDepth)
+				{
+					return FirstNodeTreeDepth < SecondNodeTreeDepth;
+				}
+
+				// Next, try to keep nodes that have smaller gaps in tree depth on one side
+				TreeDepthPair FirstMinTreeDepth = { ESignalFlowEntryType::InvalidMax, TNumericLimits<int32>::Max() };
+				TreeDepthPair SecondMinTreeDepth = { ESignalFlowEntryType::InvalidMax, TNumericLimits<int32>::Max() };
+
+				double FirstInputTimestamp = TNumericLimits<double>::Max();
+				double SecondInputTimestamp = TNumericLimits<double>::Max();
+
+				GetMinTreeDepthAndTimestamps((*FirstNode)->FilteredInputs, FirstMinTreeDepth, FirstInputTimestamp);
+				GetMinTreeDepthAndTimestamps((*SecondNode)->FilteredInputs, SecondMinTreeDepth, SecondInputTimestamp);
+
+				// Sort by the minimum tree depth of the parent nodes
+				if (FirstMinTreeDepth != SecondMinTreeDepth)
+				{
+					return FirstMinTreeDepth < SecondMinTreeDepth;
+				}
+
+				// Sort by timestamp of the parent nodes
+				if (FirstInputTimestamp != SecondInputTimestamp)
+				{
+					return FirstInputTimestamp < SecondInputTimestamp;
+				}
+
+				// Try timestamp of this node
+				if ((*FirstNode)->GetTimestamp() != (*SecondNode)->GetTimestamp())
+				{
+					return (*FirstNode)->GetTimestamp() < (*SecondNode)->GetTimestamp();
+				}
+			}
+
+			// Fallback on entry IDs
+			return First.EntryID < Second.EntryID;
+		};
+
+		Node->FilteredInputs.Sort(SortInput);
+	}
+
+	TreeDepthPair FSignalFlowDashboardViewFactory::GetOriginTreeDepthForNode(const TSharedPtr<ISignalFlowNode>& Node) const
+	{
+		// Either grab the tree depth for this node,
+		// or if it's a dummy connection node, grab the tree depth of the original input to the dummy connection
+		ensure(Node.IsValid());
+
+		if (!Node->IsRealNode())
+		{
+			const TSharedPtr<FDummyConnectionNode> DummyNode = StaticCastSharedPtr<FDummyConnectionNode>(Node);
+			const TSharedPtr<ISignalFlowNode>* DummyInputNode = FilteredEntries.Find(DummyNode->ConnectionInputKey);
+			if (DummyInputNode && DummyInputNode->IsValid())
+			{
+				return (*DummyInputNode)->GetTreeDepth();
+			}
+		}
+
+		return Node->GetTreeDepth();
+	}
+
+	void FSignalFlowDashboardViewFactory::GetMinTreeDepthAndTimestamps(const TArray<FSignalFlowEntryKey>& NodeKeys, TreeDepthPair& OutMinTreeDepth, double& OutMinTimestamp) const
+	{
+		for (const FSignalFlowEntryKey& NodeKey : NodeKeys)
+		{
+			const TSharedPtr<ISignalFlowNode>* Node = FilteredEntries.Find(NodeKey);
+
+			if (Node && Node->IsValid())
+			{
+				OutMinTreeDepth = FMath::Min(GetOriginTreeDepthForNode(*Node), OutMinTreeDepth);
+				OutMinTimestamp = FMath::Min((*Node)->GetTimestamp(), OutMinTimestamp);
+			}
+		}
+	}
+
+	void FSignalFlowDashboardViewFactory::FixUnassignedNodeOrders()
+	{
+		// Run through again and find any nodes that did not get an order assigned
+		// This can happen if a node is not the root node but has no outputs
+		for (auto& [Key, Node] : FilteredEntries)
+		{
+			if (!Node.IsValid())
+			{
+				continue;
+			}
+
+			if (!Node->NodeOrderIsValid())
+			{
+				int32 LowestNodeOrder = TNumericLimits<int32>::Max();
+				for (const FSignalFlowEntryKey& InputKey : Node->FilteredInputs)
+				{
+					const TSharedPtr<ISignalFlowNode>* InputNode = FilteredEntries.Find(InputKey);
+					if (InputNode && InputNode->IsValid() && (*InputNode)->NodeOrderIsValid())
+					{
+						LowestNodeOrder = FMath::Min((*InputNode)->GetNodeOrderID(), LowestNodeOrder);
+					}
+				}
+
+				if (LowestNodeOrder != TNumericLimits<int32>::Max())
+				{
+					Node->SetNodeOrderID(LowestNodeOrder);
+				}
+			}
+		}
+	}
+
+	void FSignalFlowDashboardViewFactory::SortNodeOrder()
+	{
+		auto SortSignalNodes = [this](const TSharedPtr<ISignalFlowNode>& First, const TSharedPtr<ISignalFlowNode>& Second)
+		{
+			if (!First.IsValid() || !Second.IsValid())
+			{
+				return false;
+			}
+
+			// Sort by the depth in the signal node tree first
+			if (First->GetTreeDepth() != Second->GetTreeDepth())
+			{
+				return First->GetTreeDepth() < Second->GetTreeDepth();
+			}
+
+			// Rows are sorted by the order a node was processed - this groups nodes that output to the same input together where possible
+			// Inputs were sorted inside previous step FSignalFlowDashboardViewFactory::CreateFilteredNodeGraph
+			if (First->GetNodeOrderID() != Second->GetNodeOrderID())
+			{
+				return First->GetNodeOrderID() < Second->GetNodeOrderID();
+			}
+
+			// Fallback to timestamp
+			return First->GetTimestamp() < Second->GetTimestamp();
+		};
+
+		FilteredGraphNodes.Sort(SortSignalNodes);
+	}
+
+	void FSignalFlowDashboardViewFactory::CalcXPositions()
+	{
+		int32 GraphWidth = 0;
+
+		TreeDepthPair CurrentTreeDepthPair{ ESignalFlowEntryType::InvalidMax, TNumericLimits<int32>::Max() };
+
+		// Keep track of node indexes that we have processed
+		// These are tracked as rows of indexes
+		TArray<TArray<int32>> ProcessedNodeRows;
+		ProcessedNodeRows.Reserve(FilteredGraphNodes.Num());
+
+		int32 RowIndex = INDEX_NONE;
+
+		// At this point, the nodes should be sorted from lowest to highest tree depth
+		// We can iterate backwards (from audio device to source) to calc the positions
+		for (int32 RevIndex = FilteredGraphNodes.Num() - 1; RevIndex >= 0; --RevIndex)
+		{
+			TSharedPtr<ISignalFlowNode> NextNodeToProcess = FilteredGraphNodes[RevIndex];
+			if (!NextNodeToProcess.IsValid())
+			{
+				continue;
+			}
+
+			// If we have moved on to a new row, calculate the positions of the current structure
+			if (CurrentTreeDepthPair != NextNodeToProcess->GetTreeDepth())
+			{
+				CalcXPositionsForRows(ProcessedNodeRows, GraphWidth);
+
+				// Create the next row
+				RowIndex = RowIndex == INDEX_NONE ? 0 : RowIndex + 1;
+
+				TArray<int32> NextRow;
+				NextRow.Reserve(FilteredGraphNodes.Num());
+				ProcessedNodeRows.Add(NextRow);
+
+				CurrentTreeDepthPair = NextNodeToProcess->GetTreeDepth();
+			}
+
+			ensure(ProcessedNodeRows.IsValidIndex(RowIndex));
+			ProcessedNodeRows[RowIndex].Add(RevIndex);
+		}
+
+		// Do one final positioning process for the final row
+		CalcXPositionsForRows(ProcessedNodeRows, GraphWidth);
+	}
+
+	void FSignalFlowDashboardViewFactory::CalcXPositionsForRows(const TArray<TArray<int32>>& NodeIndexRows, int32& OutGraphWidth)
+	{
+		if (NodeIndexRows.IsEmpty())
+		{
+			return;
+		}
+
+		const int32 RowToProcessIndex = NodeIndexRows.Num() - 1;
+
+		const TArray<int32>& RowToProcess = NodeIndexRows[RowToProcessIndex];
+		const int32 RowToProcessNumNodes = RowToProcess.Num();
+
+		if (OutGraphWidth <= RowToProcessNumNodes)
+		{
+			// new largest row, position the nodes around X = 0.0
+			PositionRowAroundOrigin(RowToProcess);
+
+			// Propagate change down the hierachy - moving nodes relative to their inputs
+			for (int32 PreviousRowsIndex = RowToProcessIndex + 1; PreviousRowsIndex < NodeIndexRows.Num(); ++PreviousRowsIndex)
+			{
+				const TArray<int32> PrevRowOfIndexes = NodeIndexRows[PreviousRowsIndex];
+				TArray<TSharedPtr<ISignalFlowNode>> PrevPositionedNodes;
+				PrevPositionedNodes.Reserve(PrevRowOfIndexes.Num());
+
+				for (const int32 NodeIndexInGraph : PrevRowOfIndexes)
+				{
+					const TSharedPtr<ISignalFlowNode>& Node = FilteredGraphNodes[NodeIndexInGraph];
+					CalcXPositionForNode(Node, Node->FilteredInputs, PrevPositionedNodes, PrevRowOfIndexes.Num());
+
+					PrevPositionedNodes.Add(Node);
+				}
+			}
+
+			OutGraphWidth = RowToProcessNumNodes;
+		}
+		else
+		{
+			TArray<TSharedPtr<ISignalFlowNode>> PrevPositionedNodes;
+			PrevPositionedNodes.Reserve(RowToProcess.Num());
+
+			// Process just this row - moving nodes relative to their outputs
+			for (int32 NodeIndexInRow = 0; NodeIndexInRow < RowToProcess.Num(); ++NodeIndexInRow)
+			{
+				TSharedPtr<ISignalFlowNode>& Node = FilteredGraphNodes[RowToProcess[NodeIndexInRow]];
+				if (Node.IsValid())
+				{
+					CalcXPositionForNode(Node, Node->FilteredOutputs, PrevPositionedNodes, RowToProcess.Num());
+				}
+
+				PrevPositionedNodes.Add(Node);
+			}
+		}
+	}
+
+	void FSignalFlowDashboardViewFactory::PositionRowAroundOrigin(const TArray<int32>& NodeIndexRow)
+	{
+		// Space a row of nodes around zero (origin) with a X distance of 1.0
+		const int32 NumNodes = NodeIndexRow.Num();
+		const bool bIsEvenNumNodes = NumNodes % 2 == 0;
+		const float HalfNumNodes = static_cast<float>(NumNodes) / 2.0f;
+
+		for (int32 NodeIndexInRow = 0; NodeIndexInRow < NodeIndexRow.Num(); ++NodeIndexInRow)
+		{
+			const TSharedPtr<ISignalFlowNode>& Node = FilteredGraphNodes[NodeIndexRow[NodeIndexInRow]];
+			if (!Node.IsValid())
+			{
+				continue;
+			}
+
+			if (bIsEvenNumNodes)
+			{
+				// Even number of nodes - space the nodes around the origin with an empty space at the origin
+				Node->XPos = NumNodes - HalfNumNodes - NodeIndexInRow;
+				if (NodeIndexInRow >= HalfNumNodes)
+				{
+					Node->XPos--;
+				}
+			}
+			else
+			{
+				// Odd number of nodes - evenly space the nodes around the origin
+				Node->XPos = static_cast<float>(((NumNodes - NodeIndexInRow - 1) * 2) - (NumNodes - 1)) / 2.0f;
+			}
+		}
+	}
+
+	void FSignalFlowDashboardViewFactory::CalcXPositionForNode(const TSharedPtr<ISignalFlowNode>& Node, const TArray<FSignalFlowEntryKey>& NodesToPositionRelativeTo, const TArray<TSharedPtr<ISignalFlowNode>>& PrevPositionedNodes, const int32 NumNodesInRow)
+	{
+		if (!Node.IsValid() || NodesToPositionRelativeTo.IsEmpty())
+		{
+			return;
+		}
+
+		const int32 NumPositioningNodes = NodesToPositionRelativeTo.Num();
+		const bool bEvenNumPositioningNodes = NumPositioningNodes % 2 == 0;
+
+		if (bEvenNumPositioningNodes)
+		{
+			// Even num positioners: nodes are positioned in between the two middle positioning nodes
+			const int32 NodeToAlignWithIndex02 = NumPositioningNodes / 2;
+			const int32 NodeToAlignWithIndex01 = NodeToAlignWithIndex02 - 1;
+			
+			if (ensure(NodesToPositionRelativeTo.IsValidIndex(NodeToAlignWithIndex01) && NodesToPositionRelativeTo.IsValidIndex(NodeToAlignWithIndex02)))
+			{
+				TSharedPtr<ISignalFlowNode>* NodeToAlignWith01 = FilteredEntries.Find(NodesToPositionRelativeTo[NodeToAlignWithIndex01]);
+				TSharedPtr<ISignalFlowNode>* NodeToAlignWith02 = FilteredEntries.Find(NodesToPositionRelativeTo[NodeToAlignWithIndex02]);
+
+				if (NodeToAlignWith01 && NodeToAlignWith02 && NodeToAlignWith01->IsValid() && NodeToAlignWith02->IsValid())
+				{
+					Node->XPos = FMath::Floor(((*NodeToAlignWith01)->XPos + (*NodeToAlignWith02)->XPos) / 2.0f);
+					MakeSpaceForNodeInRow(Node, PrevPositionedNodes, NumNodesInRow);
+				}
+			}
+		}
+		else
+		{
+			// Odd num positioners: nodes are positioned in line with the middle positioning node
+			const int32 NodeToAlignWithIndex = NumPositioningNodes / 2;
+			if (ensure(NodesToPositionRelativeTo.IsValidIndex(NodeToAlignWithIndex)))
+			{
+				TSharedPtr<ISignalFlowNode>* NodeToAlignWith = FilteredEntries.Find(NodesToPositionRelativeTo[NodeToAlignWithIndex]);
+				if (NodeToAlignWith && NodeToAlignWith->IsValid())
+				{
+					Node->XPos = (*NodeToAlignWith)->XPos;
+					MakeSpaceForNodeInRow(Node, PrevPositionedNodes, NumNodesInRow);
+				}
+			}
+		}
+	}
+
+	void FSignalFlowDashboardViewFactory::MakeSpaceForNodeInRow(const TSharedPtr<ISignalFlowNode>& Node, const TArray<TSharedPtr<ISignalFlowNode>>& PrevPositionedNodes, const int32 NumNodesInRow)
+	{
+		// In cases where nodes share inputs/outputs, we can end up with nodes placed in the same XPos
+		// Check if we have hit this edge case and reposition the row accordingly
+		const TSharedPtr<ISignalFlowNode>* PrevNodeAtExistingPosition = PrevPositionedNodes.FindByPredicate([XPos = Node->XPos](const TSharedPtr<ISignalFlowNode>& AlreadyPositionedNode)
+		{
+			return AlreadyPositionedNode.IsValid() && AlreadyPositionedNode->XPos == XPos;
+		});
+
+		if (PrevNodeAtExistingPosition)
+		{
+			// there already is a node at this position - either shift all nodes to the left or put the new node on the right
+			const bool bShiftNodesLeft = static_cast<float>(PrevPositionedNodes.Num()) / static_cast<float>(NumNodesInRow) <= 0.5f;
+			if (bShiftNodesLeft)
+			{
+				for (const TSharedPtr<ISignalFlowNode>& PrevNode : PrevPositionedNodes)
+				{
+					PrevNode->XPos -= 1.0f;
+				}
+			}
+			else
+			{
+				Node->XPos = PrevPositionedNodes.Last()->XPos + 1.0f;
+			}
+		}
+	}
+} // namespace UE::Audio::Insights
+
+#undef LOCTEXT_NAMESPACE
diff --git a/Engine/Plugins/AudioInsights/Source/AudioInsights/Public/Views/SignalFlowDashboardViewFactory.h b/Engine/Plugins/AudioInsights/Source/AudioInsights/Public/Views/SignalFlowDashboardViewFactory.h
new file mode 100644
--- /dev/null
+++ b/Engine/Plugins/AudioInsights/Source/AudioInsights/Public/Views/SignalFlowDashboardViewFactory.h
@@ -0,0 +1,110 @@
+// Copyright Epic Games, Inc. All Rights Reserved.
+#pragma once
+
+#include "Messages/SignalFlowEntryKey.h"
+#include "Templates/SharedPointer.h"
+#include "Views/SSignalFlowGraph.h"
+#include "Views/TableDashboardViewFactory.h"
+#include "Widgets/Views/SListView.h"
+
+#define UE_API AUDIOINSIGHTS_API
+
+namespace UE::Audio::Insights
+{
+	class FSignalFlowTraceProvider;
+	class FSignalFlowEntryNode;
+	class FDummyConnectionNode;
+	class ISignalFlowNode;
+
+	struct FSignalFlowDashboardEntry;
+	struct TreeDepthPair;
+
+	class FSignalFlowDashboardViewFactory : public FTraceDashboardViewFactoryBase, public TSharedFromThis<FSignalFlowDashboardViewFactory>
+	{
+	public:
+		UE_API FSignalFlowDashboardViewFactory();
+		UE_API virtual ~FSignalFlowDashboardViewFactory();
+
+		UE_API virtual FName GetName() const override;
+		UE_API virtual FText GetDisplayName() const override;
+		UE_API virtual FSlateIcon GetIcon() const override;
+		UE_API virtual EDefaultDashboardTabStack GetDefaultTabStack() const override;
+		UE_API virtual TSharedRef<SWidget> MakeWidget(TSharedRef<SDockTab> OwnerTab, const FSpawnTabArgs& SpawnTabArgs) override;
+
+	private:
+		struct SRowWidget : public STableRow<TSharedPtr<FSignalFlowDashboardEntry>>
+		{
+			SLATE_BEGIN_ARGS(SRowWidget)
+				: _Style(nullptr)
+				{
+				}
+				SLATE_STYLE_ARGUMENT(FTableRowStyle, Style)
+			SLATE_END_ARGS()
+
+			void Construct(const FArguments& InArgs, const TSharedRef<STableViewBase>& InOwnerTable, TSharedPtr<FSignalFlowDashboardEntry> InData);
+
+			TSharedPtr<FSignalFlowDashboardEntry> RowEntry;
+		};
+
+		const FTableRowStyle* GetRowStyle() const;
+
+		TSharedRef<SWidget> CreateSelectionPanel();
+		TSharedRef<SWidget> CreateGraphSection();
+		TSharedRef<SListView<TSharedPtr<FSignalFlowDashboardEntry>>> CreateSelectableList(const FText& HeadingText, const TArray<TSharedPtr<FSignalFlowDashboardEntry>>& ListSource);
+
+		void OnSelectionChanged(TSharedPtr<FSignalFlowDashboardEntry> SelectedItem, ESelectInfo::Type SelectInfo);
+
+		virtual void ProcessEntries(EProcessReason Reason) override;
+		void BuildSelectionMenus(const TSortedMap<FSignalFlowEntryKey, TSharedPtr<FSignalFlowDashboardEntry>>& DeviceData, const EProcessReason Reason);
+
+		virtual void RefreshFilteredEntriesListView() override;
+		void OnRequestGraphRefresh();
+		void CreateFilteredNodeGraph(const FSignalFlowTraceProvider& Provider, TSharedPtr<FSignalFlowDashboardEntry> FilterFromEntry);
+		
+		void AddAllNodesToGraph();
+		void AddFilteredInputsRecursive(const TSortedMap<FSignalFlowEntryKey, TSharedPtr<FSignalFlowDashboardEntry>>& DeviceData, TSharedPtr<FSignalFlowDashboardEntry> Entry, TSharedPtr<FSignalFlowDashboardEntry> ChildEntry, int32 TreeDepth);
+		void AddFilteredOutputs(const TSortedMap<FSignalFlowEntryKey, TSharedPtr<FSignalFlowDashboardEntry>>& DeviceData, TSharedPtr<FSignalFlowDashboardEntry> Entry, int32 TreeDepth);
+		void AddFilteredOutputsRecursive(const TSortedMap<FSignalFlowEntryKey, TSharedPtr<FSignalFlowDashboardEntry>>& DeviceData, TSharedPtr<FSignalFlowDashboardEntry> Entry, TSharedPtr<FSignalFlowDashboardEntry> ParentEntry, int32 TreeDepth);
+
+		void CalculateFilteredGraphDepthStructure(TOptional<FSignalFlowEntryKey>& OutRootEntryKey, TArray<TreeDepthPair>& OutSortedTreeDepths) const;
+
+		void CreateDummyInputs(const TArray<TreeDepthPair>& TreeDepthStructure);
+		TSharedPtr<FDummyConnectionNode> CreateDummyConnectionNode(const TSharedPtr<FSignalFlowEntryNode>& InputNode, const TSharedPtr<ISignalFlowNode>& OutputNode, const TreeDepthPair& TreeDepthPair, const FSignalFlowEntryKey& ConnectionOutputNodeKey);
+		
+		void AssignNodeOrderRecursive(TSharedPtr<ISignalFlowNode> OutNode, const int32 Depth, int32& OutNodeProcessingOrder);
+		void SortNodeInputs(TSharedPtr<ISignalFlowNode> OutNode);
+		TreeDepthPair GetOriginTreeDepthForNode(const TSharedPtr<ISignalFlowNode>& Node) const;
+		void GetMinTreeDepthAndTimestamps(const TArray<FSignalFlowEntryKey>& NodeKeys, TreeDepthPair& OutMinTreeDepth, double& OutMinTimestamp) const;
+
+		void FixUnassignedNodeOrders();
+		void SortNodeOrder();
+
+		void CalcXPositions();
+		void CalcXPositionsForRows(const TArray<TArray<int32>>& NodeIndexRows, int32& OutGraphWidth);
+		void PositionRowAroundOrigin(const TArray<int32>& NodeIndexRow);
+		void CalcXPositionForNode(const TSharedPtr<ISignalFlowNode>& Node, const TArray<FSignalFlowEntryKey>& NodesToPositionRelativeTo, const TArray<TSharedPtr<ISignalFlowNode>>& PrevPositionedNodes, const int32 NumNodesInRow);
+		void MakeSpaceForNodeInRow(const TSharedPtr<ISignalFlowNode>& Node, const TArray<TSharedPtr<ISignalFlowNode>>& PrevPositionedNodes, const int32 NumNodesInRow);
+
+		TSharedPtr<FSignalFlowTraceProvider> SignalFlowProvider;
+		TSharedPtr<SWidget> DashboardWidget;
+
+		TSharedPtr<SListView<TSharedPtr<FSignalFlowDashboardEntry>>> ActiveSourcesListView;
+		TArray<TSharedPtr<FSignalFlowDashboardEntry>> ActiveSources;
+
+		TSharedPtr<SListView<TSharedPtr<FSignalFlowDashboardEntry>>> ActiveBusesListView;
+		TArray<TSharedPtr<FSignalFlowDashboardEntry>> ActiveBuses;
+
+		TSharedPtr<SListView<TSharedPtr<FSignalFlowDashboardEntry>>> ActiveSubmixesListView;
+		TArray<TSharedPtr<FSignalFlowDashboardEntry>> ActiveSubmixes;
+
+		TSortedMap<FSignalFlowEntryKey, TSharedPtr<ISignalFlowNode>> FilteredEntries;
+		TArray<TSharedPtr<ISignalFlowNode>> FilteredGraphNodes;
+		TSharedPtr<SSignalFlowGraph> GraphWidgetContents;
+
+		TSharedPtr<FSignalFlowDashboardEntry> SelectedItem;
+
+		bool bGraphRequiresRefresh = false;
+	};
+} // namespace UE::Audio::Insights
+
+#undef UE_API
