[Audio Insights] Remove unused MainSubmixKey from signal flow provider - this has been suceeded by the Audio Device entry as the root of the graph.

#Synced-CL 49327150

diff --git a/Engine/Plugins/AudioInsights/Source/AudioInsights/Private/Providers/SignalFlowTraceProvider.cpp b/Engine/Plugins/AudioInsights/Source/AudioInsights/Private/Providers/SignalFlowTraceProvider.cpp
--- a/Engine/Plugins/AudioInsights/Source/AudioInsights/Private/Providers/SignalFlowTraceProvider.cpp
+++ b/Engine/Plugins/AudioInsights/Source/AudioInsights/Private/Providers/SignalFlowTraceProvider.cpp
@@ -489,8 +489,6 @@ namespace UE::Audio::Insights
 
 			if (Msg.bIsMainSubmix)
 			{
-				MainSubmixKey = { ESignalFlowEntryType::Submix, Msg.SubmixId };
-
 				const FSignalFlowEntryKey OutputAudioDeviceKey = { ESignalFlowEntryType::AudioDevice, Msg.DeviceId };
 
 				EntryRef.Outputs.Add(OutputAudioDeviceKey);
diff --git a/Engine/Plugins/AudioInsights/Source/AudioInsights/Private/Providers/SignalFlowTraceProvider.h b/Engine/Plugins/AudioInsights/Source/AudioInsights/Private/Providers/SignalFlowTraceProvider.h
--- a/Engine/Plugins/AudioInsights/Source/AudioInsights/Private/Providers/SignalFlowTraceProvider.h
+++ b/Engine/Plugins/AudioInsights/Source/AudioInsights/Private/Providers/SignalFlowTraceProvider.h
@@ -39,7 +39,6 @@ namespace UE::Audio::Insights
 
 		FSignalFlowMessages TraceMessages;
 
-		TOptional<FSignalFlowEntryKey> MainSubmixKey;
 		bool bForceGraphRefresh = false;
 	};
 } // namespace UE::Audio::Insights
