[Audio Insights] Enable Signal Flow inside audio insights (editor only)

#Synced-CL 49327150

diff --git a/Engine/Plugins/AudioInsights/Source/AudioInsightsEditor/Private/AudioInsightsEditorModule.cpp b/Engine/Plugins/AudioInsights/Source/AudioInsightsEditor/Private/AudioInsightsEditorModule.cpp
--- a/Engine/Plugins/AudioInsights/Source/AudioInsightsEditor/Private/AudioInsightsEditorModule.cpp
+++ b/Engine/Plugins/AudioInsights/Source/AudioInsightsEditor/Private/AudioInsightsEditorModule.cpp
@@ -13,6 +13,7 @@
 #include "Views/AudioEventLogDashboardViewFactory.h"
 #include "Views/AudioMetersPanelDashboardViewFactory.h"
 #include "Views/LogDashboardViewFactory.h"
+#include "Views/SignalFlowDashboardViewFactory.h"
 #include "Views/SoundDashboardViewFactory.h"
 #include "Views/SoundPlotsDashboardViewFactory.h"
 #include "Views/SubmixDashboardViewFactory.h"
@@ -48,6 +49,7 @@ namespace UE::Audio::Insights
 			DashboardFactory->RegisterViewFactory(MakeShared<FSubmixDashboardViewFactory>());
 			DashboardFactory->RegisterViewFactory(MakeShared<FAudioBusDashboardViewFactory>());
 			DashboardFactory->RegisterViewFactory(MakeShared<FAudioMetersPanelDashboardViewFactory>());
+			DashboardFactory->RegisterViewFactory(MakeShared<FSignalFlowDashboardViewFactory>());
 			DashboardFactory->RegisterViewFactory(PlotsDashboard);
 			DashboardFactory->RegisterViewFactory(MakeShared<FAudioAnalyzerRackDashboardViewFactory>());
 
