[Audio Insights] Integrate scrubbing in the signal flow graph
- Send info messages from AudioMixerSourceManager updated to send all sends for an asset at once (rather than a message per send). This helps when looking up sends for a specific entry while scrubbing, and should reduce memory usage of send info trace messages too.
- Add the audio device id to FSignalFlowEntryKey - we need this to distinguish between assets with the same ids across multiple audio devices (such as submixes) when scrubbing through the cache
- Implement OnTimingViewTimeMarkerChanged inside SignalFlowTraceProvider
- To calculate the current state of alive entries on the dashboard, we only actually have to iterate through the cache from the last IsAlive ping for each entry. Use a window of 2 seconds to ensure we catch all pings (assuming they are sent at a one second rate).

#Synced-CL 49557410

diff --git a/Engine/Plugins/AudioInsights/Source/AudioInsights/Private/Messages/SignalFlowTraceMessages.cpp b/Engine/Plugins/AudioInsights/Source/AudioInsights/Private/Messages/SignalFlowTraceMessages.cpp
--- a/Engine/Plugins/AudioInsights/Source/AudioInsights/Private/Messages/SignalFlowTraceMessages.cpp
+++ b/Engine/Plugins/AudioInsights/Source/AudioInsights/Private/Messages/SignalFlowTraceMessages.cpp
@@ -24,15 +24,38 @@ namespace UE::Audio::Insights
 		const Trace::IAnalyzer::FEventData& EventData = InContext.EventData;
 
 		SenderID = EventData.GetValue<uint32>("SenderID");
-		ReceiverID = EventData.GetValue<uint32>("ReceiverID");
-		SendLevel = EventData.GetValue<float>("SendLevel");
+
+		const TArrayView<const uint32> ReceiverIDs = EventData.GetArrayView<uint32>("ReceiverIDs");
+		const TArrayView<const float> SendLevels = EventData.GetArrayView<float>("SendLevels");
+		if (ensure(ReceiverIDs.Num() == SendLevels.Num()))
+		{
+			const int32 NumSends = ReceiverIDs.Num();
+			SendInfos.Reserve(NumSends);
+			for (int32 SendIndex = 0; SendIndex < NumSends; ++SendIndex)
+			{
+				SendInfos.Add(FSendInfoPair{ ReceiverIDs[SendIndex], SendLevels[SendIndex] });
+			}
+		}
+	}
+
+	uint32 FSignalFlowSendMessageBase::GetSizeOf() const
+	{
+		return sizeof(FSignalFlowSendMessageBase) + SendInfos.GetAllocatedSize();;
 	}
 
 	FSubmixToSubmixSendMessage::FSubmixToSubmixSendMessage(const Trace::IAnalyzer::FOnEventContext& InContext)
-		: FSignalFlowSendMessageBase(InContext)
+		: FSignalFlowMessageBase(InContext)
 	{
 		const Trace::IAnalyzer::FEventData& EventData = InContext.EventData;
 
+		SenderID = EventData.GetValue<uint32>("SenderID");
+		ReceiverID = EventData.GetValue<uint32>("ReceiverID");
+		SendLevel = EventData.GetValue<float>("SendLevel");
 		bIsFinalStage = EventData.GetValue<bool>("IsFinalStage");
 	}
+
+	uint32 FSubmixToSubmixSendMessage::GetSizeOf() const
+	{
+		return sizeof(FSubmixToSubmixSendMessage);
+	}
 } // namespace UE::Audio::Insights
diff --git a/Engine/Plugins/AudioInsights/Source/AudioInsights/Private/Providers/SignalFlowTraceProvider.cpp b/Engine/Plugins/AudioInsights/Source/AudioInsights/Private/Providers/SignalFlowTraceProvider.cpp
--- a/Engine/Plugins/AudioInsights/Source/AudioInsights/Private/Providers/SignalFlowTraceProvider.cpp
+++ b/Engine/Plugins/AudioInsights/Source/AudioInsights/Private/Providers/SignalFlowTraceProvider.cpp
@@ -222,74 +222,61 @@ namespace UE::Audio::Insights
 	{
 		auto GetOwnerEntry = [this](const FSoundStopMessage& Msg)
 		{
-			return FindDeviceEntry(Msg.DeviceId, { ESignalFlowEntryType::OwnerObject, Msg.ActiveSoundPlayOrder });
+			return FindDeviceEntry(Msg.DeviceId, { Msg.DeviceId, ESignalFlowEntryType::OwnerObject, Msg.ActiveSoundPlayOrder });
 		};
 
 		auto GetOwnerEntryFromSoundWaveStop = [this](const FSoundWaveStopMessage& Msg)
 		{
-			return FindDeviceEntry(Msg.DeviceId, { ESignalFlowEntryType::OwnerObject, Msg.ActiveSoundPlayOrder });
-		};
-
-		auto CreateSoundwaveEntry = [this](const FSoundWaveStartMessage& Msg)
-		{
-			CreateAudioDeviceEntry(Msg, Msg.DeviceId);
-
-			return CreateEntry(Msg, Msg.DeviceId, ESignalFlowEntryType::SoundSource, Msg.WaveInstancePlayOrder);
+			return FindDeviceEntry(Msg.DeviceId, { Msg.DeviceId, ESignalFlowEntryType::OwnerObject, Msg.ActiveSoundPlayOrder });
 		};
 
 		auto GetSoundwaveEntry = [this](const ::Audio::FDeviceId DeviceId, const uint32 WaveInstancePlayOrder)
 		{
-			return FindDeviceEntry(DeviceId, { ESignalFlowEntryType::SoundSource, WaveInstancePlayOrder });
+			return FindDeviceEntry(DeviceId, { DeviceId, ESignalFlowEntryType::SoundSource, WaveInstancePlayOrder });
 		};
 
 		auto GetSoundwaveEntryFromMessage = [this](const FSoundWaveStopMessage& Msg)
 		{
-			return FindDeviceEntry(Msg.DeviceId, { ESignalFlowEntryType::SoundSource, Msg.WaveInstancePlayOrder });
+			return FindDeviceEntry(Msg.DeviceId, { Msg.DeviceId, ESignalFlowEntryType::SoundSource, Msg.WaveInstancePlayOrder });
 		};
 
-		auto GetSoundwaveEntryForSubmixSendMessage = [&GetSoundwaveEntry](const FSignalFlowSendMessageBase& Msg)
+		auto GetEntryFromSendMessage = [&GetSoundwaveEntry](const FSignalFlowSendMessageBase& Msg)
 		{
 			return GetSoundwaveEntry(Msg.DeviceId, Msg.SenderID);
 		};
 
-		auto CreateBusEntry = [this](const FAudioBusStartMessage& Msg)
-		{
-			CreateAudioDeviceEntry(Msg, Msg.DeviceId);
-
-			return CreateEntry(Msg, Msg.DeviceId, ESignalFlowEntryType::AudioBus, Msg.AudioBusId);
-		};
-
 		auto GetBusEntry = [this](const FAudioBusMessageBase& Msg)
 		{
-			return FindDeviceEntry(Msg.DeviceId, { ESignalFlowEntryType::AudioBus, Msg.AudioBusId });
-		};
-
-		auto CreateSubmixEntry = [this](const FSubmixMessageBase& Msg)
-		{
-			CreateAudioDeviceEntry(Msg, Msg.DeviceId);
-
-			return CreateEntry(Msg, Msg.DeviceId, ESignalFlowEntryType::Submix, Msg.SubmixId);
+			return FindDeviceEntry(Msg.DeviceId, { Msg.DeviceId, ESignalFlowEntryType::AudioBus, Msg.AudioBusId });
 		};
 
 		auto GetSubmixEntry = [this](const FSubmixMessageBase& Msg)
 		{
-			return FindDeviceEntry(Msg.DeviceId, { ESignalFlowEntryType::Submix, Msg.SubmixId });
+			return FindDeviceEntry(Msg.DeviceId, { Msg.DeviceId, ESignalFlowEntryType::Submix, Msg.SubmixId });
 		};
 
 		auto GetSubmixToSubmixEntry = [this](const FSubmixToSubmixSendMessage& Msg)
 		{
-			return FindDeviceEntry(Msg.DeviceId, { ESignalFlowEntryType::Submix, Msg.SenderID });
+			return FindDeviceEntry(Msg.DeviceId, { Msg.DeviceId, ESignalFlowEntryType::Submix, Msg.SenderID });
 		};
 
 		TSet<FSignalFlowEntryKey> RemovedKeys;
 
-		ProcessMessageQueue<FAudioBusStartMessage>(TraceMessages.BusStartMessages, CreateBusEntry,
+		ProcessMessageQueue<FAudioBusStartMessage>(TraceMessages.BusStartMessages, 
+		[this](const FAudioBusStartMessage& Msg) 
+		{ 
+			return CreateAudioBusEntry(Msg); 
+		},
 		[this](const FAudioBusStartMessage& Msg, TSharedPtr<FSignalFlowDashboardEntry>* OutEntry)
 		{
 			ProcessNewBusEntry(Msg, OutEntry);
 		});
 
-		ProcessMessageQueue<FAudioBusStartMessage>(TraceMessages.BusIsAlivePingMessages, CreateBusEntry,
+		ProcessMessageQueue<FAudioBusStartMessage>(TraceMessages.BusIsAlivePingMessages, 
+		[this](const FAudioBusStartMessage& Msg)
+		{
+			return CreateAudioBusEntry(Msg);
+		},
 		[this](const FAudioBusStartMessage& Msg, TSharedPtr<FSignalFlowDashboardEntry>* OutEntry)
 		{
 			ProcessNewBusEntry(Msg, OutEntry);
@@ -301,13 +288,21 @@ namespace UE::Audio::Insights
 			ProcessStopMessage(Msg.DeviceId, Msg.Timestamp, OutEntry, RemovedKeys);
 		});
 
-		ProcessMessageQueue<FSubmixLoadedMessage>(TraceMessages.SubmixLoadedMessages, CreateSubmixEntry,
+		ProcessMessageQueue<FSubmixLoadedMessage>(TraceMessages.SubmixLoadedMessages,
+		[this](const FSubmixLoadedMessage& Msg)
+		{
+			return CreateSubmixEntry(Msg);
+		},
 		[this](const FSubmixLoadedMessage& Msg, TSharedPtr<FSignalFlowDashboardEntry>* OutEntry)
 		{
 			ProcessNewSubmixEntry(Msg, OutEntry);
 		});
 
-		ProcessMessageQueue<FSubmixAlivePingMessage>(TraceMessages.SubmixAliveMessages, CreateSubmixEntry,
+		ProcessMessageQueue<FSubmixAlivePingMessage>(TraceMessages.SubmixAliveMessages,
+		[this](const FSubmixLoadedMessage& Msg)
+		{
+			return CreateSubmixEntry(Msg);
+		},
 		[this](const FSubmixAlivePingMessage& Msg, TSharedPtr<FSignalFlowDashboardEntry>* OutEntry)
 		{
 			ProcessNewSubmixEntry(Msg, OutEntry);
@@ -319,36 +314,50 @@ namespace UE::Audio::Insights
 			ProcessStopMessage(Msg.DeviceId, Msg.Timestamp, OutEntry, RemovedKeys);
 		});
 
-		ProcessMessageQueue<FSoundWaveStartMessage>(TraceMessages.SoundWaveStartMessages, CreateSoundwaveEntry,
+		ProcessMessageQueue<FSoundWaveStartMessage>(TraceMessages.SoundWaveStartMessages, 
+		[this](const FSoundWaveStartMessage& Msg)
+		{
+			return CreateSoundSourceEntry(Msg);
+		},
 		[this](const FSoundWaveStartMessage& Msg, TSharedPtr<FSignalFlowDashboardEntry>* OutEntry)
 		{
 			ProcessNewSoundwaveEntry(Msg, OutEntry);
 		});
 
-		ProcessMessageQueue<FSoundWaveIsAlivePingMessage>(TraceMessages.SoundWaveIsAlivePingMessages, CreateSoundwaveEntry,
+		ProcessMessageQueue<FSoundWaveIsAlivePingMessage>(TraceMessages.SoundWaveIsAlivePingMessages, 
+		[this](const FSoundWaveIsAlivePingMessage& Msg)
+		{
+			return CreateSoundSourceEntry(Msg);
+		},
 		[this](const FSoundWaveIsAlivePingMessage& Msg, TSharedPtr<FSignalFlowDashboardEntry>* OutEntry)
 		{
 			ProcessNewSoundwaveEntry(Msg, OutEntry);
 		});
 
-		ProcessMessageQueue<FSoundToBusSendMessage>(TraceMessages.SoundToBusSendMessages, GetSoundwaveEntryForSubmixSendMessage,
+		ProcessMessageQueue<FSoundToBusSendMessage>(TraceMessages.SoundToBusSendMessages, GetEntryFromSendMessage,
 		[this](const FSoundToBusSendMessage& Msg, TSharedPtr<FSignalFlowDashboardEntry>* OutEntry)
 		{
-			const FSignalFlowEntryKey OutputEntryKey{ ESignalFlowEntryType::AudioBus, Msg.ReceiverID };
-			ProcessOutputSendMessage(OutEntry, OutputEntryKey, Msg.SendLevel);
+			for (const FSignalFlowSendMessageBase::FSendInfoPair& BusSend : Msg.SendInfos)
+			{
+				const FSignalFlowEntryKey OutputEntryKey{ Msg.DeviceId, ESignalFlowEntryType::AudioBus, BusSend.ReceiverID };
+				ProcessOutputSendMessage(OutEntry, OutputEntryKey, BusSend.SendLevel);
+			}
 		});
 
-		ProcessMessageQueue<FSoundToSubmixSendMessage>(TraceMessages.SoundToSubmixSendMessages, GetSoundwaveEntryForSubmixSendMessage,
+		ProcessMessageQueue<FSoundToSubmixSendMessage>(TraceMessages.SoundToSubmixSendMessages, GetEntryFromSendMessage,
 		[this](const FSoundToSubmixSendMessage& Msg, TSharedPtr<FSignalFlowDashboardEntry>* OutEntry)
 		{
-			const FSignalFlowEntryKey OutputEntryKey{ ESignalFlowEntryType::Submix, Msg.ReceiverID };
-			ProcessOutputSendMessage(OutEntry, OutputEntryKey, Msg.SendLevel);
+			for (const FSignalFlowSendMessageBase::FSendInfoPair& SubmixSend : Msg.SendInfos)
+			{
+				const FSignalFlowEntryKey OutputEntryKey{ Msg.DeviceId, ESignalFlowEntryType::Submix, SubmixSend.ReceiverID };
+				ProcessOutputSendMessage(OutEntry, OutputEntryKey, SubmixSend.SendLevel);
+			}
 		});
 
 		ProcessMessageQueue<FSubmixToSubmixSendMessage>(TraceMessages.SubmixToSubmixSendMessages, GetSubmixToSubmixEntry,
 		[this](const FSubmixToSubmixSendMessage& Msg, TSharedPtr<FSignalFlowDashboardEntry>* OutEntry)
 		{
-			const FSignalFlowEntryKey OutputEntryKey{ Msg.bIsFinalStage ? ESignalFlowEntryType::AudioDevice : ESignalFlowEntryType::Submix, Msg.ReceiverID };
+			const FSignalFlowEntryKey OutputEntryKey{ Msg.DeviceId, Msg.bIsFinalStage ? ESignalFlowEntryType::AudioDevice : ESignalFlowEntryType::Submix, Msg.ReceiverID };
 			ProcessOutputSendMessage(OutEntry, OutputEntryKey, Msg.SendLevel);
 		});
 
@@ -380,13 +389,13 @@ namespace UE::Audio::Insights
 					RemovedKeys.Add((*OwnerObjectEntry)->GetSignalFlowEntryKey());
 				}
 
-				RemoveDeviceEntry(SoundStop.DeviceId, { ESignalFlowEntryType::OwnerObject, SoundStop.ActiveSoundPlayOrder });
+				RemoveDeviceEntry(SoundStop.DeviceId, { SoundStop.DeviceId, ESignalFlowEntryType::OwnerObject, SoundStop.ActiveSoundPlayOrder });
 			}
 
 			if (TSortedMap<FSignalFlowEntryKey, TSharedPtr<FSignalFlowDashboardEntry>>* DeviceData = DeviceDataMap.Find(SoundStop.DeviceId))
 			{
 				// Remove the audio device entry if it's the only remaining entry in the map
-				if (DeviceData->Num() == 1 && DeviceData->Contains({ ESignalFlowEntryType::AudioDevice, SoundStop.DeviceId }))
+				if (DeviceData->Num() == 1 && DeviceData->Contains({ SoundStop.DeviceId, ESignalFlowEntryType::AudioDevice, SoundStop.DeviceId }))
 				{
 					DeviceDataMap.Remove(SoundStop.DeviceId);
 				}
@@ -410,17 +419,231 @@ namespace UE::Audio::Insights
 
 	void FSignalFlowTraceProvider::OnTimingViewTimeMarkerChanged(double TimeMarker)
 	{
+		TSet<FSignalFlowEntryKey> ActiveEntries;
+		TSet<FSignalFlowEntryKey> RemovedKeys;
+
+		const FAudioInsightsCacheManager& CacheManager = IAudioInsightsModule::GetChecked().GetCacheManager();
+
+		// IsAlive pings are sent every second - therefore we should only need to check for alive nodes from the last second
+		// Double this to 2s to ensure we do catch all alive sounds
+		constexpr double MaxLookbackTime = 2.0;
+		const double RangeCheckStart = FMath::Max(TimeMarker - MaxLookbackTime, 0.0);
+
+		/*
+		 * Audio Buses
+		*/
+		CacheManager.IterateOverRange<FAudioBusStartMessage>(AudioBusMessageNames::Start, RangeCheckStart, TimeMarker, [this, &ActiveEntries](const FAudioBusStartMessage& Msg)
+		{
+			TSharedPtr<FSignalFlowDashboardEntry>* Entry = CreateAudioBusEntry(Msg, &ActiveEntries);
+			if (Entry && Entry->IsValid())
+			{
+				ProcessNewBusEntry(Msg, Entry);
+			}
+		});
+
+		CacheManager.IterateOverRange<FAudioBusStopMessage>(AudioBusMessageNames::Stop, RangeCheckStart, TimeMarker, [this, &ActiveEntries, &RemovedKeys](const FAudioBusStopMessage& Msg)
+		{
+			const FSignalFlowEntryKey EntryKey = { Msg.DeviceId, ESignalFlowEntryType::AudioBus, Msg.AudioBusId };
+			TSharedPtr<FSignalFlowDashboardEntry>* Entry = FindDeviceEntry(Msg.DeviceId, EntryKey);
+			if (Entry && Entry->IsValid() && (*Entry)->Timestamp < Msg.Timestamp)
+			{
+				ActiveEntries.Remove(EntryKey);
+				ProcessStopMessage(Msg.DeviceId, Msg.Timestamp, Entry, RemovedKeys);
+			}
+		});
+
+		/*
+		 * Submixes
+		*/
+		CacheManager.IterateOverRange<FSubmixLoadedMessage>(SubmixMessageNames::Loaded, RangeCheckStart, TimeMarker, [this, &ActiveEntries](const FSubmixLoadedMessage& Msg)
+		{
+			TSharedPtr<FSignalFlowDashboardEntry>* Entry = CreateSubmixEntry(Msg, &ActiveEntries);
+			if (Entry && Entry->IsValid())
+			{
+				ProcessNewSubmixEntry(Msg, Entry);
+			}
+		});
+
+		CacheManager.IterateOverRange<FSubmixAlivePingMessage>(SubmixMessageNames::IsAlivePing, RangeCheckStart, TimeMarker, [this, &ActiveEntries](const FSubmixAlivePingMessage& Msg)
+		{
+			TSharedPtr<FSignalFlowDashboardEntry>* Entry = CreateSubmixEntry(Msg, &ActiveEntries);
+			if (Entry && Entry->IsValid())
+			{
+				ProcessNewSubmixEntry(Msg, Entry);
+			}
+		});
+
+		CacheManager.IterateOverRange<FSubmixUnloadedMessage>(SubmixMessageNames::Unloaded, RangeCheckStart, TimeMarker, [this, &ActiveEntries, &RemovedKeys](const FSubmixUnloadedMessage& Msg)
+		{
+			const FSignalFlowEntryKey EntryKey = { Msg.DeviceId, ESignalFlowEntryType::Submix, Msg.SubmixId };
+			TSharedPtr<FSignalFlowDashboardEntry>* Entry = FindDeviceEntry(Msg.DeviceId, EntryKey);
+			if (Entry && Entry->IsValid() && (*Entry)->Timestamp < Msg.Timestamp)
+			{
+				ActiveEntries.Remove(EntryKey);
+				ProcessStopMessage(Msg.DeviceId, Msg.Timestamp, Entry, RemovedKeys);
+			}
+		});
+
+		/*
+		 * Sound Sources
+		*/
+		CacheManager.IterateOverRange<FSoundWaveStartMessage>(SoundMessageNames::SoundWaveStart, RangeCheckStart, TimeMarker, [this, &ActiveEntries](const FSoundWaveStartMessage& Msg)
+		{
+			TSharedPtr<FSignalFlowDashboardEntry>* Entry = CreateSoundSourceEntry(Msg, &ActiveEntries);
+			if (Entry && Entry->IsValid())
+			{
+				ProcessNewSoundwaveEntry(Msg, Entry, &ActiveEntries);
+			}
+		});
+
+		CacheManager.IterateOverRange<FSoundWaveIsAlivePingMessage>(SoundMessageNames::SoundWaveIsAlivePing, RangeCheckStart, TimeMarker, [this, &ActiveEntries](const FSoundWaveIsAlivePingMessage& Msg)
+		{
+			TSharedPtr<FSignalFlowDashboardEntry>* Entry = CreateSoundSourceEntry(Msg, &ActiveEntries);
+			if (Entry && Entry->IsValid())
+			{
+				ProcessNewSoundwaveEntry(Msg, Entry, &ActiveEntries);
+			}
+		});
+
+		/*
+		 * Sound source stop
+		*/
+		CacheManager.IterateOverRange<FSoundWaveStopMessage>(SoundMessageNames::SoundWaveStop, RangeCheckStart, TimeMarker, [this, &ActiveEntries, &RemovedKeys](const FSoundWaveStopMessage& Msg)
+		{
+			const FSignalFlowEntryKey EntryKey = { Msg.DeviceId, ESignalFlowEntryType::SoundSource, Msg.WaveInstancePlayOrder };
+			TSharedPtr<FSignalFlowDashboardEntry>* Entry = FindDeviceEntry(Msg.DeviceId, EntryKey);
+			if (Entry && Entry->IsValid() && (*Entry)->Timestamp < Msg.Timestamp)
+			{
+				TSharedPtr<FSignalFlowDashboardEntry>* OwnerObjectEntry = FindDeviceEntry(Msg.DeviceId, { Msg.DeviceId, ESignalFlowEntryType::OwnerObject, Msg.ActiveSoundPlayOrder });
+				if (OwnerObjectEntry && OwnerObjectEntry->IsValid())
+				{
+					(*OwnerObjectEntry)->Outputs.Remove(EntryKey);
+				}
+		
+				ActiveEntries.Remove(EntryKey);
+				ProcessStopMessage(Msg.DeviceId, Msg.Timestamp, Entry, RemovedKeys);
+			}
+		});
+		
+		CacheManager.IterateOverRange<FSoundStopMessage>(SoundMessageNames::SoundStop, RangeCheckStart, TimeMarker, [this, &ActiveEntries, &RemovedKeys](const FSoundStopMessage& Msg)
+		{
+			const FSignalFlowEntryKey OwnerObjectEntryKey = { Msg.DeviceId, ESignalFlowEntryType::OwnerObject, Msg.ActiveSoundPlayOrder };
+			TSharedPtr<FSignalFlowDashboardEntry>* OwnerObjectEntry = FindDeviceEntry(Msg.DeviceId, OwnerObjectEntryKey);
+			if (OwnerObjectEntry && OwnerObjectEntry->IsValid() && (*OwnerObjectEntry)->Timestamp < Msg.Timestamp)
+			{
+				RemovedKeys.Add(OwnerObjectEntryKey);
+				ActiveEntries.Remove(OwnerObjectEntryKey);
+
+				RemoveDeviceEntry(Msg.DeviceId, OwnerObjectEntryKey);
+			}
+		});
+
+		/*
+		 * Send levels/remove unreferenced entries
+		*/
+		for (auto DeviceItr = DeviceDataMap.CreateIterator(); DeviceItr; ++DeviceItr)
+		{
+			for (auto DataItr = DeviceItr.Value().CreateIterator(); DataItr; ++DataItr)
+			{
+				if (ActiveEntries.Contains(DataItr.Key()))
+				{
+					CollectSendInfoForTimestamp(DataItr.Key(), &DataItr.Value(), CacheManager, TimeMarker, RangeCheckStart);
+				}
+				else
+				{
+					RemovedKeys.Add(DataItr.Key());
+					DataItr.RemoveCurrent();
+					bForceGraphRefresh = true;
+				}
+			}
+
+			if (DeviceItr.Value().Num() == 1 && DeviceItr.Value().Contains({ DeviceItr.Key(), ESignalFlowEntryType::AudioDevice, DeviceItr.Key() }))
+			{
+				RemovedKeys.Add({ DeviceItr.Key(), ESignalFlowEntryType::AudioDevice, DeviceItr.Key() });
+				DeviceItr.RemoveCurrent();
+				bForceGraphRefresh = true;
+			}
+		}
+
+		if (!RemovedKeys.IsEmpty())
+		{
+			RemoveInputs(RemovedKeys);
+			bForceGraphRefresh = true;
+		}
+				
+		if (bForceGraphRefresh)
+		{
+			OnRequestGraphRefresh.Broadcast();
+			bForceGraphRefresh = false;
+		}
+
 		// Call parent method to update LastMessageId
 		FTraceProviderBase::OnTimingViewTimeMarkerChanged(TimeMarker);
+	}
+
+	void FSignalFlowTraceProvider::CollectSendInfoForTimestamp(const FSignalFlowEntryKey& EntryKey, TSharedPtr<FSignalFlowDashboardEntry>* Entry, const FAudioInsightsCacheManager& CacheManager, const double TimeMarker, const double RangeCheckStart)
+	{
+		if (Entry == nullptr || !Entry->IsValid())
+		{
+			return;
+		}
 
-		// TODO - Implement scrubbing
+		switch (EntryKey.EntryType)
+		{
+			case ESignalFlowEntryType::SoundSource:
+			{
+				const FSoundToBusSendMessage* SoundToBusMsg = CacheManager.FindClosestMessage<FSoundToBusSendMessage>(SignalFlowMessageNames::SoundWaveToBusSend, TimeMarker, GetTypeHash(EntryKey), RangeCheckStart);
+				if (SoundToBusMsg)
+				{
+					for (const FSignalFlowSendMessageBase::FSendInfoPair& SendInfo : SoundToBusMsg->SendInfos)
+					{
+						const FSignalFlowEntryKey OutputEntryKey{ SoundToBusMsg->DeviceId, ESignalFlowEntryType::AudioBus, SendInfo.ReceiverID };
+						ProcessOutputSendMessage(Entry, OutputEntryKey, SendInfo.SendLevel);
+					}
+				}
+
+				const FSoundToSubmixSendMessage* SoundToSubmixMsg = CacheManager.FindClosestMessage<FSoundToSubmixSendMessage>(SignalFlowMessageNames::SoundWaveToSubmixSend, TimeMarker, GetTypeHash(EntryKey), RangeCheckStart);
+				if (SoundToSubmixMsg)
+				{
+					for (const FSignalFlowSendMessageBase::FSendInfoPair& SendInfo : SoundToSubmixMsg->SendInfos)
+					{
+						const FSignalFlowEntryKey OutputEntryKey{ SoundToSubmixMsg->DeviceId, ESignalFlowEntryType::Submix, SendInfo.ReceiverID };
+						ProcessOutputSendMessage(Entry, OutputEntryKey, SendInfo.SendLevel);
+					}
+				}
+				break;
+			}
+				
+
+			case ESignalFlowEntryType::AudioBus:
+			{
+				// TODO
+				break;
+			}
+
+			case ESignalFlowEntryType::Submix:
+			{
+				const FSubmixToSubmixSendMessage* SubmixToSubmixMsg = CacheManager.FindClosestMessage<FSubmixToSubmixSendMessage>(SignalFlowMessageNames::SubmixToSubmixSend, TimeMarker, GetTypeHash(EntryKey), RangeCheckStart);
+				if (SubmixToSubmixMsg)
+				{
+					const FSignalFlowEntryKey OutputEntryKey{ SubmixToSubmixMsg->DeviceId, ESignalFlowEntryType::Submix, SubmixToSubmixMsg->ReceiverID };
+					ProcessOutputSendMessage(Entry, OutputEntryKey, SubmixToSubmixMsg->SendLevel);
+				}
+				break;
+			}
+		}
 	}
 
-	TSharedPtr<FSignalFlowDashboardEntry>* FSignalFlowTraceProvider::CreateEntry(const IAudioCachedMessage& Msg, const ::Audio::FDeviceId DeviceID, const ESignalFlowEntryType EntryType, const uint32 NodeID)
+	void FSignalFlowTraceProvider::OnTimeControlMethodReset()
+	{
+		Reset();
+	}
+
+	TSharedPtr<FSignalFlowDashboardEntry>* FSignalFlowTraceProvider::CreateEntry(const IAudioCachedMessage& Msg, const ::Audio::FDeviceId DeviceID, const ESignalFlowEntryType EntryType, const uint32 NodeID, TSet<FSignalFlowEntryKey>* OutActiveEntriesTracker /*= nullptr*/)
 	{
 		TSharedPtr<FSignalFlowDashboardEntry>* ToReturn = nullptr;
 
-		UpdateDeviceEntry(DeviceID, { EntryType, NodeID }, [&ToReturn, &Msg, EntryType, DeviceID, NodeID](TSharedPtr<FSignalFlowDashboardEntry>& Entry)
+		UpdateDeviceEntry(DeviceID, { DeviceID, EntryType, NodeID }, [&ToReturn, &Msg, EntryType, DeviceID, NodeID](TSharedPtr<FSignalFlowDashboardEntry>& Entry)
 		{
 			if (!Entry.IsValid())
 			{
@@ -435,12 +658,17 @@ namespace UE::Audio::Insights
 			ToReturn = &Entry;
 		});
 
+		if (OutActiveEntriesTracker && ToReturn && ToReturn->IsValid())
+		{
+			OutActiveEntriesTracker->Add((*ToReturn)->GetSignalFlowEntryKey());
+		}
+
 		return ToReturn;
 	}
 
-	TSharedPtr<FSignalFlowDashboardEntry>* FSignalFlowTraceProvider::CreateAudioDeviceEntry(const IAudioCachedMessage& Msg, const::Audio::FDeviceId DeviceID)
+	TSharedPtr<FSignalFlowDashboardEntry>* FSignalFlowTraceProvider::CreateAudioDeviceEntry(const IAudioCachedMessage& Msg, const::Audio::FDeviceId DeviceID, TSet<FSignalFlowEntryKey>* OutActiveEntriesTracker /*= nullptr*/)
 	{
-		TSharedPtr<FSignalFlowDashboardEntry>* AudioDeviceEntry = CreateEntry(Msg, DeviceID, ESignalFlowEntryType::AudioDevice, DeviceID);
+		TSharedPtr<FSignalFlowDashboardEntry>* AudioDeviceEntry = CreateEntry(Msg, DeviceID, ESignalFlowEntryType::AudioDevice, DeviceID, OutActiveEntriesTracker);
 		if (AudioDeviceEntry && AudioDeviceEntry->IsValid())
 		{
 			(*AudioDeviceEntry)->Name = FText::Format(LOCTEXT("SignalFlow_AudioDevice", "Audio Device {0}"), static_cast<uint32>(DeviceID)).ToString();
@@ -450,6 +678,34 @@ namespace UE::Audio::Insights
 		return AudioDeviceEntry;
 	}
 
+	TSharedPtr<FSignalFlowDashboardEntry>* FSignalFlowTraceProvider::CreateOwnerEntry(const FSoundWaveStartMessage& Msg, TSet<FSignalFlowEntryKey>* OutActiveEntriesTracker /*= nullptr*/)
+	{
+		CreateAudioDeviceEntry(Msg, Msg.DeviceId, OutActiveEntriesTracker);
+
+		return CreateEntry(Msg, Msg.DeviceId, ESignalFlowEntryType::OwnerObject, Msg.ActiveSoundPlayOrder, OutActiveEntriesTracker);
+	}
+
+	TSharedPtr<FSignalFlowDashboardEntry>* FSignalFlowTraceProvider::CreateSoundSourceEntry(const FSoundWaveStartMessage& Msg, TSet<FSignalFlowEntryKey>* OutActiveEntriesTracker /*= nullptr*/)
+	{
+		CreateAudioDeviceEntry(Msg, Msg.DeviceId, OutActiveEntriesTracker);
+
+		return CreateEntry(Msg, Msg.DeviceId, ESignalFlowEntryType::SoundSource, Msg.WaveInstancePlayOrder, OutActiveEntriesTracker);
+	}
+
+	TSharedPtr<FSignalFlowDashboardEntry>* FSignalFlowTraceProvider::CreateAudioBusEntry(const FAudioBusStartMessage& Msg, TSet<FSignalFlowEntryKey>* OutActiveEntriesTracker /*= nullptr*/)
+	{
+		CreateAudioDeviceEntry(Msg, Msg.DeviceId, OutActiveEntriesTracker);
+
+		return CreateEntry(Msg, Msg.DeviceId, ESignalFlowEntryType::AudioBus, Msg.AudioBusId, OutActiveEntriesTracker);
+	}
+
+	TSharedPtr<FSignalFlowDashboardEntry>* FSignalFlowTraceProvider::CreateSubmixEntry(const FSubmixMessageBase& Msg, TSet<FSignalFlowEntryKey>* OutActiveEntriesTracker /*= nullptr*/)
+	{
+		CreateAudioDeviceEntry(Msg, Msg.DeviceId, OutActiveEntriesTracker);
+
+		return CreateEntry(Msg, Msg.DeviceId, ESignalFlowEntryType::Submix, Msg.SubmixId, OutActiveEntriesTracker);
+	}
+
 	void FSignalFlowTraceProvider::LinkInput(const FSignalFlowDashboardEntry& Entry, const FSignalFlowEntryKey& OutputKey, const ::Audio::FDeviceId DeviceID)
 	{
 		if (const TSharedPtr<FSignalFlowDashboardEntry>* OutputEntry = FindDeviceEntry(DeviceID, OutputKey))
@@ -489,41 +745,32 @@ namespace UE::Audio::Insights
 
 			if (Msg.bIsMainSubmix)
 			{
-				const FSignalFlowEntryKey OutputAudioDeviceKey = { ESignalFlowEntryType::AudioDevice, Msg.DeviceId };
+				const FSignalFlowEntryKey OutputAudioDeviceKey = { Msg.DeviceId, ESignalFlowEntryType::AudioDevice, Msg.DeviceId };
 
-				EntryRef.Outputs.Add(OutputAudioDeviceKey);
+				EntryRef.Outputs.FindOrAdd(OutputAudioDeviceKey);
 				LinkInput(EntryRef, OutputAudioDeviceKey, Msg.DeviceId);
 			}
 		}
 	}
 
-	void FSignalFlowTraceProvider::ProcessNewSoundwaveEntry(const FSoundWaveStartMessage& Msg, TSharedPtr<FSignalFlowDashboardEntry>* OutEntry)
+	void FSignalFlowTraceProvider::ProcessNewSoundwaveEntry(const FSoundWaveStartMessage& Msg, TSharedPtr<FSignalFlowDashboardEntry>* OutEntry, TSet<FSignalFlowEntryKey>* OutActiveEntriesTracker /*= nullptr*/)
 	{
-		auto CreateOwnerEntry = [this](const FSoundWaveStartMessage& Msg)
-		{
-			CreateAudioDeviceEntry(Msg, Msg.DeviceId);
-
-			return CreateEntry(Msg, Msg.DeviceId, ESignalFlowEntryType::OwnerObject, Msg.ActiveSoundPlayOrder);
-		};
-
 		if (OutEntry && OutEntry->IsValid())
 		{
 			FSignalFlowDashboardEntry& EntryRef = *OutEntry->Get();
 			EntryRef.Name = Msg.Name;
 			EntryRef.IconName = FSignalFlowTraceProviderPrivate::GetSourceIconName(Msg.EntryType);
 
-			if (TSharedPtr<FSignalFlowDashboardEntry>* OwnerObjectEntry = CreateOwnerEntry(Msg))
+			TSharedPtr<FSignalFlowDashboardEntry>* OwnerObjectEntry = CreateOwnerEntry(Msg, OutActiveEntriesTracker);
+			if (OwnerObjectEntry && OwnerObjectEntry->IsValid())
 			{
-				if (OwnerObjectEntry->IsValid())
-				{
-					(*OwnerObjectEntry)->Name = Msg.ActorLabel;
-					(*OwnerObjectEntry)->EntryType = ESignalFlowEntryType::OwnerObject;
-					(*OwnerObjectEntry)->Name = Msg.ActorLabel;
-					(*OwnerObjectEntry)->IconName = FName(Msg.ActorIconName);
+				(*OwnerObjectEntry)->Name = Msg.ActorLabel;
+				(*OwnerObjectEntry)->EntryType = ESignalFlowEntryType::OwnerObject;
+				(*OwnerObjectEntry)->Name = Msg.ActorLabel;
+				(*OwnerObjectEntry)->IconName = FName(Msg.ActorIconName);
 
-					LinkInput((**OwnerObjectEntry), EntryRef.GetSignalFlowEntryKey(), Msg.DeviceId);
-					FSignalFlowOutputData& OutputInfo = (*OwnerObjectEntry)->Outputs.FindOrAdd(EntryRef.GetSignalFlowEntryKey());
-				}
+				LinkInput((**OwnerObjectEntry), EntryRef.GetSignalFlowEntryKey(), Msg.DeviceId);
+				(*OwnerObjectEntry)->Outputs.FindOrAdd(EntryRef.GetSignalFlowEntryKey());
 			}
 		}
 	}
diff --git a/Engine/Plugins/AudioInsights/Source/AudioInsights/Private/Providers/SignalFlowTraceProvider.h b/Engine/Plugins/AudioInsights/Source/AudioInsights/Private/Providers/SignalFlowTraceProvider.h
--- a/Engine/Plugins/AudioInsights/Source/AudioInsights/Private/Providers/SignalFlowTraceProvider.h
+++ b/Engine/Plugins/AudioInsights/Source/AudioInsights/Private/Providers/SignalFlowTraceProvider.h
@@ -23,13 +23,20 @@ namespace UE::Audio::Insights
 	private:
 		virtual bool ProcessMessages() override;
 		virtual void OnTimingViewTimeMarkerChanged(double TimeMarker) override;
+		void CollectSendInfoForTimestamp(const FSignalFlowEntryKey& EntryKey, TSharedPtr<FSignalFlowDashboardEntry>* Entry, const FAudioInsightsCacheManager& CacheManager, const double TimeMarker, const double RangeCheckStart);
 
-		TSharedPtr<FSignalFlowDashboardEntry>* CreateEntry(const IAudioCachedMessage& Msg, const ::Audio::FDeviceId DeviceID, const ESignalFlowEntryType EntryType, const uint32 NodeID);
-		TSharedPtr<FSignalFlowDashboardEntry>* CreateAudioDeviceEntry(const IAudioCachedMessage& Msg, const ::Audio::FDeviceId DeviceID);
+		virtual void OnTimeControlMethodReset() override;
+
+		TSharedPtr<FSignalFlowDashboardEntry>* CreateEntry(const IAudioCachedMessage& Msg, const ::Audio::FDeviceId DeviceID, const ESignalFlowEntryType EntryType, const uint32 NodeID, TSet<FSignalFlowEntryKey>* OutActiveEntriesTracker = nullptr);
+		TSharedPtr<FSignalFlowDashboardEntry>* CreateAudioDeviceEntry(const IAudioCachedMessage& Msg, const ::Audio::FDeviceId DeviceID, TSet<FSignalFlowEntryKey>* OutActiveEntriesTracker = nullptr);
+		TSharedPtr<FSignalFlowDashboardEntry>* CreateOwnerEntry(const FSoundWaveStartMessage& Msg, TSet<FSignalFlowEntryKey>* OutActiveEntriesTracker = nullptr);
+		TSharedPtr<FSignalFlowDashboardEntry>* CreateSoundSourceEntry(const FSoundWaveStartMessage& Msg, TSet<FSignalFlowEntryKey>* OutActiveEntriesTracker = nullptr);
+		TSharedPtr<FSignalFlowDashboardEntry>* CreateAudioBusEntry(const FAudioBusStartMessage& Msg, TSet<FSignalFlowEntryKey>* OutActiveEntriesTracker = nullptr);
+		TSharedPtr<FSignalFlowDashboardEntry>* CreateSubmixEntry(const FSubmixMessageBase& Msg, TSet<FSignalFlowEntryKey>* OutActiveEntriesTracker = nullptr);
 
 		void ProcessNewBusEntry(const FAudioBusStartMessage& Msg, TSharedPtr<FSignalFlowDashboardEntry>* OutEntry);
 		void ProcessNewSubmixEntry(const FSubmixLoadedMessage& Msg, TSharedPtr<FSignalFlowDashboardEntry>* OutEntry);
-		void ProcessNewSoundwaveEntry(const FSoundWaveStartMessage& Msg, TSharedPtr<FSignalFlowDashboardEntry>* OutEntry);
+		void ProcessNewSoundwaveEntry(const FSoundWaveStartMessage& Msg, TSharedPtr<FSignalFlowDashboardEntry>* OutEntry, TSet<FSignalFlowEntryKey>* OutActiveEntriesTracker = nullptr);
 		void ProcessStopMessage(const ::Audio::FDeviceId DeviceID, const double MsgTimestamp, TSharedPtr<FSignalFlowDashboardEntry>* OutEntry, TSet<FSignalFlowEntryKey>& OutRemovedKeys);
 
 		void ProcessOutputSendMessage(TSharedPtr<FSignalFlowDashboardEntry>* OutEntry, const FSignalFlowEntryKey& OutputEntryKey, const float SendLevel);
diff --git a/Engine/Plugins/AudioInsights/Source/AudioInsights/Private/Views/SignalFlowDashboardViewFactory.cpp b/Engine/Plugins/AudioInsights/Source/AudioInsights/Private/Views/SignalFlowDashboardViewFactory.cpp
--- a/Engine/Plugins/AudioInsights/Source/AudioInsights/Private/Views/SignalFlowDashboardViewFactory.cpp
+++ b/Engine/Plugins/AudioInsights/Source/AudioInsights/Private/Views/SignalFlowDashboardViewFactory.cpp
@@ -502,7 +502,7 @@ namespace UE::Audio::Insights
 			if (!DeviceData->IsEmpty())
 			{
 				const ::Audio::FDeviceId DeviceID = Provider->GetFilteredAudioDeviceID();
-				const TSharedPtr<FSignalFlowDashboardEntry>* DeviceEntry = DeviceData->Find({ ESignalFlowEntryType::AudioDevice, DeviceID });
+				const TSharedPtr<FSignalFlowDashboardEntry>* DeviceEntry = DeviceData->Find({ DeviceID, ESignalFlowEntryType::AudioDevice, DeviceID });
 
 				if (DeviceEntry && DeviceEntry->IsValid())
 				{
@@ -738,7 +738,7 @@ namespace UE::Audio::Insights
 
 		// Create a dummy node in between the input and output nodes
 		// We can use these dummy connection nodes to help order and position real nodes in the graph and reduce connections crossing
-		const FSignalFlowEntryKey DummyConnectionEntryKey = FSignalFlowEntryKey(TreeDepthPair.LayerID, HashCombine(GetTypeHash(InputKey), TreeDepthPair.TreeDepth));
+		const FSignalFlowEntryKey DummyConnectionEntryKey = FSignalFlowEntryKey(InputKey.DeviceID, TreeDepthPair.LayerID, HashCombine(GetTypeHash(InputKey), TreeDepthPair.TreeDepth));
 		TSharedPtr<FDummyConnectionNode> DummyConnectionNode = MakeShared<FDummyConnectionNode>(DummyConnectionEntryKey, TreeDepthPair.TreeDepth, TreeDepthPair.LayerID, InputNode->Entry->Timestamp);
 		
 		DummyConnectionNode->ConnectionOutputKey = ConnectionOutputNodeKey;
diff --git a/Engine/Plugins/AudioInsights/Source/AudioInsights/Public/Messages/SignalFlowEntryKey.h b/Engine/Plugins/AudioInsights/Source/AudioInsights/Public/Messages/SignalFlowEntryKey.h
--- a/Engine/Plugins/AudioInsights/Source/AudioInsights/Public/Messages/SignalFlowEntryKey.h
+++ b/Engine/Plugins/AudioInsights/Source/AudioInsights/Public/Messages/SignalFlowEntryKey.h
@@ -1,6 +1,7 @@
 // Copyright Epic Games, Inc. All Rights Reserved.
 #pragma once
 
+#include "AudioDefines.h"
 #include "CoreTypes.h"
 #include "HAL/Platform.h"
 #include "Templates/TypeHash.h"
@@ -25,28 +26,39 @@ namespace UE::Audio::Insights
 	struct FSignalFlowEntryKey
 	{
 		FSignalFlowEntryKey() = default;
-		FSignalFlowEntryKey(const ESignalFlowEntryType InEntryType, const uint32 InEntryID)
-			: EntryType(InEntryType)
+		FSignalFlowEntryKey(const ::Audio::FDeviceId InDeviceID, const ESignalFlowEntryType InEntryType, const uint32 InEntryID)
+			: DeviceID(InDeviceID)
+			, EntryType(InEntryType)
 			, EntryID(InEntryID)
 		{
 		}
 
 		friend uint32 GetTypeHash(const FSignalFlowEntryKey& Key)
 		{
-			return HashCombine(GetTypeHash(static_cast<int32>(Key.EntryType)), GetTypeHash(Key.EntryID));
+			return HashCombine(GetTypeHash(Key.DeviceID), GetTypeHash(static_cast<int32>(Key.EntryType)), GetTypeHash(Key.EntryID));
 		}
 
 		bool operator==(const FSignalFlowEntryKey& InOther) const
 		{
-			return EntryType == InOther.EntryType && EntryID == InOther.EntryID;
+			return DeviceID == InOther.DeviceID && EntryType == InOther.EntryType && EntryID == InOther.EntryID;
 		}
 
 		bool operator<(const FSignalFlowEntryKey& InOther) const
 		{
-			return EntryType == InOther.EntryType ? EntryID < InOther.EntryID 
-												  : static_cast<int32>(EntryType) < static_cast<int32>(InOther.EntryType);
+			if (DeviceID != InOther.DeviceID)
+			{
+				return DeviceID < InOther.DeviceID;
+			}
+
+			if (EntryType != InOther.EntryType)
+			{
+				return static_cast<int32>(EntryType) < static_cast<int32>(InOther.EntryType);
+			}
+
+			return EntryID < InOther.EntryID;
 		}
 
+		::Audio::FDeviceId DeviceID = INDEX_NONE;
 		ESignalFlowEntryType EntryType = ESignalFlowEntryType::SoundSource;
 		uint32 EntryID = INDEX_NONE;
 	};
diff --git a/Engine/Plugins/AudioInsights/Source/AudioInsights/Public/Messages/SignalFlowTraceMessages.h b/Engine/Plugins/AudioInsights/Source/AudioInsights/Public/Messages/SignalFlowTraceMessages.h
--- a/Engine/Plugins/AudioInsights/Source/AudioInsights/Public/Messages/SignalFlowTraceMessages.h
+++ b/Engine/Plugins/AudioInsights/Source/AudioInsights/Public/Messages/SignalFlowTraceMessages.h
@@ -32,11 +32,16 @@ namespace UE::Audio::Insights
 		FSignalFlowSendMessageBase() = default;
 		FSignalFlowSendMessageBase(const Trace::IAnalyzer::FOnEventContext& InContext);
 
-		virtual uint32 GetSizeOf() const override { return sizeof(FSignalFlowSendMessageBase); }
+		virtual uint32 GetSizeOf() const override;
+
+		struct FSendInfoPair
+		{
+			uint32 ReceiverID = INDEX_NONE;
+			float SendLevel = 0.0f;
+		};
 
 		uint32 SenderID = INDEX_NONE;
-		uint32 ReceiverID = INDEX_NONE;
-		float SendLevel = 0.0f;
+		TArray<FSendInfoPair> SendInfos;
 	};
 
 	struct FSoundToBusSendMessage : public FSignalFlowSendMessageBase
@@ -47,7 +52,7 @@ namespace UE::Audio::Insights
 		{
 		}
 
-		virtual uint64 GetID() const override { return GetTypeHash(FSignalFlowEntryKey{ ESignalFlowEntryType::AudioBus, SenderID }); }
+		virtual uint64 GetID() const override { return GetTypeHash(FSignalFlowEntryKey{ DeviceId, ESignalFlowEntryType::SoundSource, SenderID }); }
 		virtual const FName GetMessageName() const override { return SignalFlowMessageNames::SoundWaveToBusSend; }
 	};
 
@@ -59,18 +64,22 @@ namespace UE::Audio::Insights
 		{
 		}
 
-		virtual uint64 GetID() const override { return GetTypeHash(FSignalFlowEntryKey{ ESignalFlowEntryType::SoundSource, SenderID }); }
+		virtual uint64 GetID() const override { return GetTypeHash(FSignalFlowEntryKey{ DeviceId, ESignalFlowEntryType::SoundSource, SenderID }); }
 		virtual const FName GetMessageName() const override { return SignalFlowMessageNames::SoundWaveToSubmixSend; }
 	};
 
-	struct FSubmixToSubmixSendMessage : public FSignalFlowSendMessageBase
+	struct FSubmixToSubmixSendMessage : public FSignalFlowMessageBase
 	{
 		FSubmixToSubmixSendMessage() = default;
 		FSubmixToSubmixSendMessage(const Trace::IAnalyzer::FOnEventContext& InContext);
 
-		virtual uint64 GetID() const override { return GetTypeHash(FSignalFlowEntryKey{ ESignalFlowEntryType::Submix, SenderID }); }
+		virtual uint64 GetID() const override { return GetTypeHash(FSignalFlowEntryKey{ DeviceId, ESignalFlowEntryType::Submix, SenderID }); }
 		virtual const FName GetMessageName() const override { return SignalFlowMessageNames::SubmixToSubmixSend; }
+		virtual uint32 GetSizeOf() const override;
 
+		uint32 SenderID = INDEX_NONE;
+		uint32 ReceiverID = INDEX_NONE;
+		float SendLevel = 0.0f;
 		bool bIsFinalStage = false;
 	};
 
@@ -114,7 +123,7 @@ namespace UE::Audio::Insights
 		virtual UObject* GetObject() override { return FSoftObjectPath(Name).ResolveObject(); }
 		virtual bool IsValid() const override { return SignalFlowNodeID != static_cast<uint32>(INDEX_NONE); }
 
-		FSignalFlowEntryKey GetSignalFlowEntryKey() const { return { EntryType, SignalFlowNodeID }; }
+		FSignalFlowEntryKey GetSignalFlowEntryKey() const { return { DeviceId, EntryType, SignalFlowNodeID }; }
 
 		::Audio::FDeviceId DeviceId = INDEX_NONE;
 		ESignalFlowEntryType EntryType = ESignalFlowEntryType::SoundSource;
diff --git a/Engine/Source/Runtime/AudioMixer/Private/AudioMixerSourceManager.cpp b/Engine/Source/Runtime/AudioMixer/Private/AudioMixerSourceManager.cpp
--- a/Engine/Source/Runtime/AudioMixer/Private/AudioMixerSourceManager.cpp
+++ b/Engine/Source/Runtime/AudioMixer/Private/AudioMixerSourceManager.cpp
@@ -281,16 +281,16 @@ UE_TRACE_EVENT_BEGIN(Audio, SoundToSubmixSend)
 	UE_TRACE_EVENT_FIELD(uint32, DeviceId)
 	UE_TRACE_EVENT_FIELD(uint64, Timestamp)
 	UE_TRACE_EVENT_FIELD(uint32, SenderID)
-	UE_TRACE_EVENT_FIELD(uint32, ReceiverID)
-	UE_TRACE_EVENT_FIELD(float, SendLevel)
+	UE_TRACE_EVENT_FIELD(uint32[], ReceiverIDs)
+	UE_TRACE_EVENT_FIELD(float[], SendLevels)
 UE_TRACE_EVENT_END()
 
 UE_TRACE_EVENT_BEGIN(Audio, SoundToBusSend)
 	UE_TRACE_EVENT_FIELD(uint32, DeviceId)
 	UE_TRACE_EVENT_FIELD(uint64, Timestamp)
 	UE_TRACE_EVENT_FIELD(uint32, SenderID)
-	UE_TRACE_EVENT_FIELD(uint32, ReceiverID)
-	UE_TRACE_EVENT_FIELD(float, SendLevel)
+	UE_TRACE_EVENT_FIELD(uint32[], ReceiverIDs)
+	UE_TRACE_EVENT_FIELD(float[], SendLevels)
 UE_TRACE_EVENT_END()
 
 // Audio Bus messages
@@ -437,20 +437,6 @@ namespace Audio
 				<< AudioBusStop.Timestamp(FPlatformTime::Cycles64());
 		}
 	}
-
-	void TraceSourceToBusSendMessage(const uint32 InDeviceId, const uint32 InSourcePlayOrder, const uint32 InBusID, const float InSendLevel)
-	{
-		const bool bChannelEnabled = UE_TRACE_CHANNELEXPR_IS_ENABLED(AudioMixerChannel);
-		if (bChannelEnabled)
-		{
-			UE_TRACE_LOG(Audio, SoundToBusSend, AudioMixerChannel)
-				<< SoundToBusSend.DeviceId(InDeviceId)
-				<< SoundToBusSend.Timestamp(FPlatformTime::Cycles64())
-				<< SoundToBusSend.SenderID(InSourcePlayOrder)
-				<< SoundToBusSend.ReceiverID(InBusID)
-				<< SoundToBusSend.SendLevel(InSendLevel);
-		}
-	}
 #endif // UE_AUDIO_PROFILERTRACE_ENABLED
 
 	/*************************************************************************
@@ -620,6 +606,85 @@ namespace Audio
 			}
 		}
 	}
+
+	void FMixerSourceManager::TraceSourceToSubmixSendMessage(const FSourceInfo& SourceInfo) const
+	{
+		const bool bChannelEnabled = UE_TRACE_CHANNELEXPR_IS_ENABLED(AudioMixerChannel);
+
+		if (!bChannelEnabled || SourceInfo.SubmixSends.IsEmpty())
+		{
+			return;
+		}
+
+		TArray<uint32, TInlineAllocator<32>> SendReceiverIDs;
+		TArray<float, TInlineAllocator<32>> SendLevels;
+
+		SendReceiverIDs.Reserve(SourceInfo.SubmixSends.Num());
+		SendLevels.Reserve(SourceInfo.SubmixSends.Num());
+
+		for (const Audio::FMixerSourceSubmixSend& SubmixSend : SourceInfo.SubmixSends)
+		{
+			const FMixerSubmixPtr SubmixPtr = SubmixSend.Submix.Pin();
+			if (SubmixPtr.IsValid())
+			{
+				SendReceiverIDs.Add(GetTypeHash(SubmixPtr->GetPathName()));
+				SendLevels.Add(SubmixSend.SendLevel);
+			}
+		}
+
+		if (!SendReceiverIDs.IsEmpty() && ensure(SendReceiverIDs.Num() == SendLevels.Num()))
+		{
+			UE_TRACE_LOG(Audio, SoundToSubmixSend, AudioMixerChannel)
+				<< SoundToSubmixSend.DeviceId(MixerDevice->DeviceID)
+				<< SoundToSubmixSend.Timestamp(FPlatformTime::Cycles64())
+				<< SoundToSubmixSend.SenderID(SourceInfo.PlayOrder)
+				<< SoundToSubmixSend.ReceiverIDs(SendReceiverIDs.GetData(), SendReceiverIDs.Num())
+				<< SoundToSubmixSend.SendLevels(SendLevels.GetData(), SendLevels.Num());
+		}
+	}
+
+	void FMixerSourceManager::TraceSourceToBusSendMessage(const uint32 SourceId, const FSourceInfo& SourceInfo) const
+	{
+		const bool bChannelEnabled = UE_TRACE_CHANNELEXPR_IS_ENABLED(AudioMixerChannel);
+		if (!bChannelEnabled)
+		{
+			return;
+		}
+
+		TArray<uint32, TInlineAllocator<32>> BusReceiverIDs;
+		TArray<float, TInlineAllocator<32>> BusSendLevels;
+
+		for (int32 AudioBusSendType = 0; AudioBusSendType < (int32)EBusSendType::Count; ++AudioBusSendType)
+		{
+			for (uint32 AudioBusId : SourceInfo.AudioBusSends[AudioBusSendType])
+			{
+				const TSharedPtr<FMixerAudioBus>* AudioBusPtr = AudioBuses.Find(AudioBusId);
+				if (AudioBusPtr && AudioBusPtr->IsValid())
+				{
+					FAudioBusSend* SourceToBusSendInfo = (*AudioBusPtr)->AudioBusSends[AudioBusSendType].FindByPredicate([SourceId](const FAudioBusSend& BusSendInfo)
+					{
+						return BusSendInfo.SourceId == SourceId;
+					});
+
+					if (SourceToBusSendInfo)
+					{
+						BusReceiverIDs.Add(AudioBusId);
+						BusSendLevels.Add(SourceToBusSendInfo->SendLevel);
+					}
+				}
+			}
+		}
+
+		if (!BusReceiverIDs.IsEmpty() && ensure(BusReceiverIDs.Num() == BusSendLevels.Num()))
+		{
+			UE_TRACE_LOG(Audio, SoundToBusSend, AudioMixerChannel)
+				<< SoundToBusSend.DeviceId(MixerDevice->DeviceID)
+				<< SoundToBusSend.Timestamp(FPlatformTime::Cycles64())
+				<< SoundToBusSend.SenderID(SourceInfo.PlayOrder)
+				<< SoundToBusSend.ReceiverIDs(BusReceiverIDs.GetData(), BusReceiverIDs.Num())
+				<< SoundToBusSend.SendLevels(BusSendLevels.GetData(), BusSendLevels.Num());
+		}
+	}
 #endif // UE_AUDIO_PROFILERTRACE_ENABLED
 
 	void FMixerSourceManager::Init(const FSourceManagerInitParams& InitParams)
@@ -1438,8 +1503,6 @@ namespace Audio
 						AudioBusPtr->BusName = AudioBusSend.AudioBusName;
 						TraceAudioBusStartMessage(MixerDevice->DeviceID, AudioBusSend.AudioBusId, AudioBusSend.AudioBusName, AudioBusSend.BusChannels);
 					}
-
-					TraceSourceToBusSendMessage(MixerDevice->DeviceID, SourceInfo.PlayOrder, AudioBusSend.AudioBusId, AudioBusSend.SendLevel);
 #endif // UE_AUDIO_PROFILERTRACE_ENABLED
 					
 					// Store on this source, which buses its sending its audio to
@@ -2267,8 +2330,6 @@ namespace Audio
 				AudioBusPtr->BusName = InBusName;
 				TraceAudioBusStartMessage(MixerDevice->DeviceID, AudioBusId, InBusName, SourceInfo.NumInputChannels);
 			}
-
-			TraceSourceToBusSendMessage(MixerDevice->DeviceID, SourceInfo.PlayOrder, SourceInfo.AudioBusId, BusSendLevel);
 #endif // UE_AUDIO_PROFILERTRACE_ENABLED
 
 			// Check to see if we need to create new bus data. If we are not playing a bus with this id, then we
@@ -3337,18 +3398,8 @@ namespace Audio
 					<< SoundDistanceAttenuation.ActiveSoundPlayOrder(SourceInfo.ActiveSoundPlayOrder)
 					<< SoundDistanceAttenuation.DistanceAttenuation(SourceInfo.DistanceAttenuationSourceDestination);
 
-				for (const Audio::FMixerSourceSubmixSend& SubmixSend : SourceInfo.SubmixSends)
-				{
-					if (SubmixSend.Submix.Pin().IsValid())
-					{
-						UE_TRACE_LOG(Audio, SoundToSubmixSend, AudioMixerChannel)
-							<< SoundToSubmixSend.DeviceId(MixerDevice->DeviceID)
-							<< SoundToSubmixSend.Timestamp(FPlatformTime::Cycles64())
-							<< SoundToSubmixSend.SenderID(SourceInfo.PlayOrder)
-							<< SoundToSubmixSend.ReceiverID(GetTypeHash(SubmixSend.Submix.Pin()->GetPathName()))
-							<< SoundToSubmixSend.SendLevel(SubmixSend.SendLevel);
-					}
-				}
+				TraceSourceToSubmixSendMessage(SourceInfo);
+				TraceSourceToBusSendMessage(SourceId, SourceInfo);
 			}
 #endif // UE_AUDIO_PROFILERTRACE_ENABLED
 		}
diff --git a/Engine/Source/Runtime/AudioMixer/Private/AudioMixerSourceManager.h b/Engine/Source/Runtime/AudioMixer/Private/AudioMixerSourceManager.h
--- a/Engine/Source/Runtime/AudioMixer/Private/AudioMixerSourceManager.h
+++ b/Engine/Source/Runtime/AudioMixer/Private/AudioMixerSourceManager.h
@@ -385,25 +385,6 @@ namespace Audio
 
 		void AudioMixerThreadCommand(TFunction<void()>&& InFunction, const char* InDebugString = nullptr, bool bInDeferExecution = false);
 
-#if UE_AUDIO_PROFILERTRACE_ENABLED
-		void OnTraceStarted(FTraceAuxiliary::EConnectionType TraceType, const FString& TraceDestination);
-
-		void RegisterAudioBusEnvelopeFollowerCVars();
-		void UnregisterAudioBusEnvelopeFollowerCVars();
-
-		void OnStartAudioBusEnvelopeFollowerCVarChanged(IConsoleVariable* InStartAudioBusEnvelopeFollowerCVar);
-		void OnStopAudioBusEnvelopeFollowerCVarChanged(IConsoleVariable* InStopAudioBusEnvelopeFollowerCVar);
-
-		void StartAudioBusEnvelopeFollowerFromCVar(const uint32 InAudioBusId);
-		void StopAudioBusEnvelopeFollowerFromCVar(const uint32 InAudioBusId);
-
-		IConsoleVariable* StartAudioBusEnvelopeFollowerCVar = nullptr;
-		IConsoleVariable* StopAudioBusEnvelopeFollowerCVar  = nullptr;
-
-		double BusIsAlivePingTimeSec = 0.0;
-		static constexpr double BusIsAlivePingWaitTimeSec = 1.0;
-#endif // UE_AUDIO_PROFILERTRACE_ENABLED
-
 		// Critical section to ensure mutating effect chains is thread-safe
 		FCriticalSection EffectChainMutationCriticalSection;
 
@@ -639,6 +620,28 @@ namespace Audio
 
 		TSharedRef<FMixerAudioBus> FindOrAddAudioBus(const FAudioBusKey BusKey, const int32 InNumChannels, const bool bIsAutomatic, bool* bOutAlreadyExisted = nullptr);
 
+#if UE_AUDIO_PROFILERTRACE_ENABLED
+		void OnTraceStarted(FTraceAuxiliary::EConnectionType TraceType, const FString& TraceDestination);
+
+		void RegisterAudioBusEnvelopeFollowerCVars();
+		void UnregisterAudioBusEnvelopeFollowerCVars();
+
+		void OnStartAudioBusEnvelopeFollowerCVarChanged(IConsoleVariable* InStartAudioBusEnvelopeFollowerCVar);
+		void OnStopAudioBusEnvelopeFollowerCVarChanged(IConsoleVariable* InStopAudioBusEnvelopeFollowerCVar);
+
+		void StartAudioBusEnvelopeFollowerFromCVar(const uint32 InAudioBusId);
+		void StopAudioBusEnvelopeFollowerFromCVar(const uint32 InAudioBusId);
+
+		void TraceSourceToSubmixSendMessage(const FSourceInfo& SourceInfo) const;
+		void TraceSourceToBusSendMessage(const uint32 SourceId, const FSourceInfo& SourceInfo) const;
+
+		IConsoleVariable* StartAudioBusEnvelopeFollowerCVar = nullptr;
+		IConsoleVariable* StopAudioBusEnvelopeFollowerCVar = nullptr;
+
+		double BusIsAlivePingTimeSec = 0.0;
+		static constexpr double BusIsAlivePingWaitTimeSec = 1.0;
+#endif // UE_AUDIO_PROFILERTRACE_ENABLED
+
 		// Array of listener transforms
 		TArray<FTransform> ListenerTransforms;
 
