[Audio Insights] Enable Signal Flow dashboard in Standalone Unreal Insights

#Synced-CL 49557410

diff --git a/Engine/Plugins/AudioInsights/Source/AudioInsights/Private/AudioInsightsModule.cpp b/Engine/Plugins/AudioInsights/Source/AudioInsights/Private/AudioInsightsModule.cpp
--- a/Engine/Plugins/AudioInsights/Source/AudioInsights/Private/AudioInsightsModule.cpp
+++ b/Engine/Plugins/AudioInsights/Source/AudioInsights/Private/AudioInsightsModule.cpp
@@ -19,6 +19,7 @@
 #include "Views/AudioEventLogDashboardViewFactory.h"
 #include "Views/AudioMetersPanelDashboardViewFactory.h"
 #include "Views/OutputMeterDashboardViewFactory.h"
+#include "Views/SignalFlowDashboardViewFactory.h"
 #include "Views/SoundDashboardViewFactory.h"
 #include "Views/SoundPlotsDashboardViewFactory.h"
 #include "Views/SubmixDashboardViewFactory.h"
@@ -69,6 +70,7 @@ namespace UE::Audio::Insights
 			DashboardFactory->RegisterViewFactory(MakeShared<FAudioBusDashboardViewFactory>());
 			DashboardFactory->RegisterViewFactory(MakeShared<FAudioMetersPanelDashboardViewFactory>());
 			DashboardFactory->RegisterViewFactory(MakeShared<FOutputMeterDashboardViewFactory>(SubmixDashboard));
+			DashboardFactory->RegisterViewFactory(MakeShared<FSignalFlowDashboardViewFactory>());
 
 			PlotsDashboard->InitPlots(SoundsDashboard);
 #endif // !WITH_EDITOR
