[Audio Insights] Add search bar to left hand menus in the signal flow graph

#Synced-CL 49588626

diff --git a/Engine/Plugins/AudioInsights/Source/AudioInsights/Private/Views/SignalFlowDashboardViewFactory.cpp b/Engine/Plugins/AudioInsights/Source/AudioInsights/Private/Views/SignalFlowDashboardViewFactory.cpp
--- a/Engine/Plugins/AudioInsights/Source/AudioInsights/Private/Views/SignalFlowDashboardViewFactory.cpp
+++ b/Engine/Plugins/AudioInsights/Source/AudioInsights/Private/Views/SignalFlowDashboardViewFactory.cpp
@@ -137,6 +137,18 @@ namespace UE::Audio::Insights
 
 		return SNew(SVerticalBox)
 
+		+ SVerticalBox::Slot()
+		.AutoHeight()
+		.MaxHeight(FSlateStyle::Get().GetSearchBoxMaxHeight())
+		.Padding(0.0f, 4.0f, 0.0f, 6.0f)
+		[
+			SNew(SSearchBox)
+			.SelectAllTextWhenFocused(true)
+			.HintText(LOCTEXT("AudioDashboard_SignalFlow_SearchBoxHintText", "Search"))
+			.MinDesiredWidth(200.0f)
+			.OnTextChanged(this, &FSignalFlowDashboardViewFactory::SetSearchBoxFilterText)
+		]
+
 		+ SVerticalBox::Slot()
 		.FillHeight(SourcesListHeightProp)
 		[
@@ -272,6 +284,12 @@ namespace UE::Audio::Insights
 			);
 	}
 
+	void FSignalFlowDashboardViewFactory::SetSearchBoxFilterText(const FText& NewText)
+	{
+		SearchBoxFilterText = NewText;
+		UpdateFilterReason = EProcessReason::FilterUpdated;
+	}
+
 	void FSignalFlowDashboardViewFactory::OnSelectionChanged(TSharedPtr<FSignalFlowDashboardEntry> InSelectedItem, ESelectInfo::Type SelectInfo)
 	{
 		if (!ActiveSourcesListView.IsValid() || !ActiveBusesListView.IsValid() || !ActiveSubmixesListView.IsValid())
@@ -335,11 +353,6 @@ namespace UE::Audio::Insights
 
 	void FSignalFlowDashboardViewFactory::ProcessEntries(EProcessReason Reason)
 	{
-		if (Reason != EProcessReason::EntriesUpdated)
-		{
-			return;
-		}
-
 		const TSharedPtr<const FSignalFlowTraceProvider> Provider = FindProvider<const FSignalFlowTraceProvider>();
 		if (!Provider.IsValid())
 		{
@@ -358,6 +371,8 @@ namespace UE::Audio::Insights
 		ActiveBuses.Reset();
 		ActiveSubmixes.Reset();
 
+		const FString FilterString = SearchBoxFilterText.ToString();
+
 		// Build the left-hand menus of selectable nodes from this device's data
 		for (const auto& [Key, Entry] : DeviceData)
 		{
@@ -366,6 +381,11 @@ namespace UE::Audio::Insights
 				continue;
 			}
 
+			if (!SearchBoxFilterText.IsEmpty() && !Entry->GetDisplayName().ToString().Contains(FilterString))
+			{
+				continue;
+			}
+
 			switch (Entry->EntryType)
 			{
 				case ESignalFlowEntryType::SoundSource:
diff --git a/Engine/Plugins/AudioInsights/Source/AudioInsights/Public/Views/SignalFlowDashboardViewFactory.h b/Engine/Plugins/AudioInsights/Source/AudioInsights/Public/Views/SignalFlowDashboardViewFactory.h
--- a/Engine/Plugins/AudioInsights/Source/AudioInsights/Public/Views/SignalFlowDashboardViewFactory.h
+++ b/Engine/Plugins/AudioInsights/Source/AudioInsights/Public/Views/SignalFlowDashboardViewFactory.h
@@ -52,6 +52,7 @@ namespace UE::Audio::Insights
 		TSharedRef<SWidget> CreateGraphSection();
 		TSharedRef<SListView<TSharedPtr<FSignalFlowDashboardEntry>>> CreateSelectableList(const FText& HeadingText, const TArray<TSharedPtr<FSignalFlowDashboardEntry>>& ListSource);
 
+		void SetSearchBoxFilterText(const FText& NewText);
 		void OnSelectionChanged(TSharedPtr<FSignalFlowDashboardEntry> SelectedItem, ESelectInfo::Type SelectInfo);
 
 		virtual void ProcessEntries(EProcessReason Reason) override;
@@ -102,6 +103,7 @@ namespace UE::Audio::Insights
 		TSharedPtr<SSignalFlowGraph> GraphWidgetContents;
 
 		TSharedPtr<FSignalFlowDashboardEntry> SelectedItem;
+		FText SearchBoxFilterText;
 
 		bool bGraphRequiresRefresh = false;
 	};
