[Audio Insights] Signal Flow - Move send levels for nodes to connection wires

#Synced-CL 49593349

diff --git a/Engine/Plugins/AudioInsights/Source/AudioInsights/Private/Views/SSignalFlowGraph.cpp b/Engine/Plugins/AudioInsights/Source/AudioInsights/Private/Views/SSignalFlowGraph.cpp
--- a/Engine/Plugins/AudioInsights/Source/AudioInsights/Private/Views/SSignalFlowGraph.cpp
+++ b/Engine/Plugins/AudioInsights/Source/AudioInsights/Private/Views/SSignalFlowGraph.cpp
@@ -3,6 +3,7 @@
 #include "Views/SSignalFlowGraph.h"
 
 #include "AudioInsightsStyle.h"
+#include "Brushes/SlateRoundedBoxBrush.h"
 #include "Internationalization/Text.h"
 #include "Messages/SignalFlowTraceMessages.h"
 #include "Styling/StyleColors.h"
@@ -29,6 +30,27 @@ namespace UE::Audio::Insights
 
 		constexpr float NodePaddingY = 24.0f;
 		constexpr float NodeMaxHeight = 100.0f;
+
+		constexpr float RoundedBrushRadius = 8.0f;
+
+		constexpr float SplineRadius = 80.0f;
+		constexpr float SendLabelPropAlongLine = 0.25f; // How far alone the spline we want the send level to appear (range 0 -> 1)
+
+		const FSlateRoundedBoxBrush& GetRoundedBoxBrush()
+		{
+			static const FSlateRoundedBoxBrush RoundedBoxBrush(FStyleColors::AccentGray.GetSpecifiedColor(), RoundedBrushRadius);
+			return RoundedBoxBrush;
+		}
+
+		FVector2d GetSplineStartDirection(const EOrientation Orientation)
+		{
+			return Orientation == EOrientation::Orient_Horizontal ? FVector2d{ SplineRadius, 0.0f } : FVector2d{ 0.0f, SplineRadius };
+		}
+
+		FVector2d GetSplineEndDirection(const EOrientation Orientation)
+		{
+			return Orientation == EOrientation::Orient_Horizontal ? FVector2d{ SplineRadius, 0.0f } : FVector2d{ 0.0f, SplineRadius };
+		}
 	}
 
 	void SSignalFlowGraphNode::Construct(const SSignalFlowGraphNode::FArguments& InArgs)
@@ -47,51 +69,8 @@ namespace UE::Audio::Insights
 		{
 			return;
 		}
-		// Todo - move the send levels to the connection lines
-		TSharedPtr<SHorizontalBox> SendLevels = SNew(SHorizontalBox);
-
-		for (const auto& [Key, SendInfo] : InArgs._Entry->Outputs)
-		{
-			if (SendInfo.SendLevel == TNumericLimits<float>::Lowest())
-			{
-				continue;
-			}
-
-			if (SendLevels->NumSlots() == 0)
-			{
-				SendLevels->AddSlot()
-				[
-					SNullWidget::NullWidget
-				];
-			}
-
-			SendLevels->AddSlot()
-			.AutoWidth()
-			.Padding(12.0f, 0.0f)
-			.HAlign(HAlign_Center)
-			[
-				SNew(STextBlock)
-				.Text_Lambda([Key, Entry = InArgs._Entry]()
-				{ 
-					if (Entry.IsValid())
-					{
-						if (FSignalFlowOutputData* FoundSendInfo = Entry->Outputs.Find(Key))
-						{
-							return FText::AsNumber(FoundSendInfo->SendLevel);
-						}
-					}
-					return FText::GetEmpty();
-				})
-			];
-		}
 
-		if (SendLevels->NumSlots() > 0)
-		{
-			SendLevels->AddSlot()
-			[
-				SNullWidget::NullWidget
-			];
-		}
+		NodeEntryKey = Entry->GetSignalFlowEntryKey();
 
 		ChildSlot
 		[
@@ -150,13 +129,6 @@ namespace UE::Audio::Insights
 								.Clipping(EWidgetClipping::ClipToBounds)
 							]
 						]
-
-						// TODO - Move to lines
-						+ SVerticalBox::Slot()
-						.Padding(4.0f)
-						[
-							SendLevels.IsValid() && SendLevels->NumSlots() > 0 ? SendLevels.ToSharedRef() : SNullWidget::NullWidget
-						]
 					]
 				]
 			]
@@ -238,6 +210,97 @@ namespace UE::Audio::Insights
 		return CachedPositionInGraph;
 	}
 
+	void SSignalFlowGraphSendLabel::Construct(const SSignalFlowGraphSendLabel::FArguments& InArgs)
+	{
+		using namespace FSignalFlowGraphPrivate;
+
+		SendReceiverEntryKey = InArgs._SendReceiverEntryKey;
+		SenderPositionInGraph = InArgs._SenderPositionInGraph;
+		ReceiverPositionInGraph = InArgs._ReceiverPositionInGraph;
+		SendLevel = InArgs._SendLevel;
+		GraphOrientation = InArgs._GraphOrientation;
+
+		ChildSlot
+		[
+			SNew(SHorizontalBox)
+			.Visibility_Lambda([this]()
+			{
+				return SendLevel.IsBound() && SendLevel.Get().IsSet() ? EVisibility::Visible : EVisibility::Collapsed;
+			})
+
+			+ SHorizontalBox::Slot()
+			.FillContentWidth(1.0f)
+			.HAlign(HAlign_Center)
+			.VAlign(VAlign_Center)
+			[
+				SNew(SBorder)
+				.BorderImage(&GetRoundedBoxBrush())
+				.HAlign(HAlign_Fill)
+				.VAlign(VAlign_Fill)
+				.BorderBackgroundColor(FStyleColors::White.GetSpecifiedColor())
+				.Padding(1.0f)
+				.Content()
+				[
+					SNew(SBorder)
+					.BorderImage(&GetRoundedBoxBrush())
+					.HAlign(HAlign_Fill)
+					.VAlign(VAlign_Fill)
+					.BorderBackgroundColor(FStyleColors::Panel.GetSpecifiedColor())
+					.Content()
+					[
+						SNew(STextBlock)
+						.Text_Lambda([this]() 
+						{
+							if (!SendLevel.IsBound())
+							{
+								return FText::GetEmpty();
+							}
+
+							const TOptional<float> SendLevelOpt = SendLevel.Get();
+							if (!SendLevelOpt.IsSet())
+							{
+								return FText::GetEmpty();
+							}
+
+							return FText::AsNumber(SendLevelOpt.GetValue(), FSlateStyle::Get().GetSendLevelFormat());
+						})
+						.Justification(ETextJustify::Center)
+						.Margin(2.0f)
+					]
+				]
+			]
+		];
+	}
+
+	FVector2D SSignalFlowGraphSendLabel::ComputeDesiredSize(float InLayoutScaleMultiplier) const
+	{
+		using namespace FSignalFlowGraphPrivate;
+
+		const FVector2D Size = SCompoundWidget::ComputeDesiredSize(InLayoutScaleMultiplier);
+		return Size;
+	}
+
+	FVector2D SSignalFlowGraphSendLabel::UpdatePositionInGraph()
+	{
+		using namespace FSignalFlowGraphPrivate;
+
+		const FVector2D SenderNodePos = SenderPositionInGraph.Get();
+		const FVector2D ReceiverNodePos = ReceiverPositionInGraph.Get();
+
+		const EOrientation Orientation = GraphOrientation.Get();
+		const FVector2D PointAlongSpline = FMath::CubicInterp(SenderNodePos, GetSplineStartDirection(Orientation), ReceiverNodePos, GetSplineEndDirection(Orientation), SendLabelPropAlongLine);
+
+		const FVector2D LabelSize = GetDesiredSize();
+
+		CachedPositionInGraph = PointAlongSpline - (LabelSize / 2.0f);
+		return CachedPositionInGraph;
+	}
+
+	FVector2D SSignalFlowGraphSendLabel::GetCachedPositionInGraph() const
+	{
+		return CachedPositionInGraph;
+	}
+
 	void SSignalFlowGraph::Construct(const SSignalFlowGraph::FArguments& InArgs)
 	{
 		ListSource = InArgs._ListSource;
@@ -270,48 +333,11 @@ namespace UE::Audio::Insights
 		TArray<FSignalFlowEntryKey> RemainingNodeKeys;
 		NodeWidgets.GenerateKeyArray(RemainingNodeKeys);
 
-		DummyEntryToOutNode.Reset();
-
-		TreeDepthPair CurrentTreeDepth{ ESignalFlowEntryType::InvalidLow, TNumericLimits<int32>::Lowest() };
-
-		float MinXPos = TNumericLimits<float>::Max();
-		float MaxXPos = TNumericLimits<float>::Lowest();
-
-		// Gather stats on the width of the graph
-		for (const TSharedPtr<ISignalFlowNode>& Node : *ListSource)
-		{
-			const float XPos = Node->XPos;
-
-			MinXPos = FMath::Min(MinXPos, XPos);
-			MaxXPos = FMath::Max(MaxXPos, XPos);
-		}
-
-		// DepthTracker is used to assign a YPosition to each node
-		// Each time we cross a layer, add one to the depth tracker
-		int32 DepthTracker = 0;
+		DummyToRealConnectionKeys.Reset();
 
-		for (const TSharedPtr<ISignalFlowNode>& Node : *ListSource)
-		{
-			if (CurrentTreeDepth != Node->GetTreeDepth())
-			{
-				CurrentTreeDepth = Node->GetTreeDepth();
-				DepthTracker++;
-			}
-
-			CreateGraphNode(Node, MinXPos, MaxXPos, DepthTracker, RemainingNodeKeys);
-		}
-
-		// Clean up and remove any old nodes that we did not detect in the new graph
-		for (const FSignalFlowEntryKey& Key : RemainingNodeKeys)
-		{
-			TSharedPtr<SSignalFlowGraphNode> RemovedNode = nullptr;
-			NodeWidgets.RemoveAndCopyValue(Key, RemovedNode);
-			
-			if (RemovedNode.IsValid())
-			{
-				Canvas->RemoveSlot(RemovedNode.ToSharedRef());
-			}
-		}
+		CreateGraphNodes(RemainingNodeKeys);
+		CreateSendLabels();
+		CleanUpOldWidgets(RemainingNodeKeys);
 	}
 
 	int32 SSignalFlowGraph::OnPaint(const FPaintArgs& Args, const FGeometry& AllottedGeometry, const FSlateRect& MyCullingRect, FSlateWindowElementList& OutDrawElements, int32 LayerId, const FWidgetStyle& InWidgetStyle, bool bParentEnabled) const
@@ -323,6 +349,8 @@ namespace UE::Audio::Insights
 
 	int32 SSignalFlowGraph::OnPaintConnections(const FGeometry& AllottedGeometry, FSlateWindowElementList& OutDrawElements, int32 LayerId) const
 	{
+		using namespace FSignalFlowGraphPrivate;
+
 		if (!ListSource)
 		{
 			return LayerId;
@@ -346,28 +374,7 @@ namespace UE::Audio::Insights
 
 			for (const FSignalFlowEntryKey& NodeOutputKey : Node->FilteredOutputs)
 			{
-				TSharedPtr<SSignalFlowGraphNode> OutputNode = nullptr;
-
-				if (const TSharedPtr<SSignalFlowGraphNode>* FoundOutputNode = NodeWidgets.Find(NodeOutputKey))
-				{
-					if (FoundOutputNode->IsValid())
-					{
-						OutputNode = *FoundOutputNode;
-					}
-				}
-				else
-				{
-					const FSignalFlowEntryKey* FoundDummyOutputKey = DummyEntryToOutNode.Find(NodeOutputKey);
-					if (FoundDummyOutputKey)
-					{
-						const TSharedPtr<SSignalFlowGraphNode>* FoundDummyOutputNode = NodeWidgets.Find(*FoundDummyOutputKey);
-
-						if (FoundDummyOutputNode && FoundDummyOutputNode->IsValid())
-						{
-							OutputNode = *FoundDummyOutputNode;
-						}
-					}
-				}
+				const TSharedPtr<SSignalFlowGraphNode> OutputNode = FindRealNode(NodeOutputKey, ENodeConnectionType::Output);
 
 				if (OutputNode.IsValid())
 				{
@@ -381,18 +388,14 @@ namespace UE::Audio::Insights
 						continue;
 					}
 
-					constexpr float Radius = 120.0f;
-					const FVector2d StartDirection = Orientation == EOrientation::Orient_Horizontal ? FVector2d{ Radius, 0.0f } : FVector2d{ 0.0f, Radius };
-					const FVector2d EndDirection = Orientation == EOrientation::Orient_Horizontal ? FVector2d{ Radius, 0.0f } : FVector2d{ 0.0f, Radius };
-
 					FSlateDrawElement::MakeSpline(
 						OutDrawElements,
 						++LayerId,
 						AllottedGeometry.ToPaintGeometry(),
 						ConnectionStartPos,
-						StartDirection,
+						GetSplineStartDirection(Orientation),
 						ConnectionEndPos,
-						EndDirection,
+						GetSplineEndDirection(Orientation),
 						2.0f,
 						ESlateDrawEffect::NoPixelSnapping,
 						FStyleColors::AccentGray.GetSpecifiedColor()
@@ -444,6 +447,108 @@ namespace UE::Audio::Insights
 		return FVector2D{ GraphWidth, GraphHeight };
 	}
 
+	void SSignalFlowGraph::CreateGraphNodes(TArray<FSignalFlowEntryKey>& OutRemainingNodeKeys)
+	{
+		TreeDepthPair CurrentTreeDepth{ ESignalFlowEntryType::InvalidLow, TNumericLimits<int32>::Lowest() };
+
+		float MinXPos = TNumericLimits<float>::Max();
+		float MaxXPos = TNumericLimits<float>::Lowest();
+
+		// Gather stats on the width of the graph
+		for (const TSharedPtr<ISignalFlowNode>& Node : *ListSource)
+		{
+			const float XPos = Node->XPos;
+
+			MinXPos = FMath::Min(MinXPos, XPos);
+			MaxXPos = FMath::Max(MaxXPos, XPos);
+		}
+
+		// DepthTracker is used to assign a YPosition to each node
+		// Each time we cross a layer, add one to the depth tracker
+		int32 DepthTracker = 0;
+
+		// Create the nodes first
+		for (const TSharedPtr<ISignalFlowNode>& Node : *ListSource)
+		{
+			if (CurrentTreeDepth != Node->GetTreeDepth())
+			{
+				CurrentTreeDepth = Node->GetTreeDepth();
+				DepthTracker++;
+			}
+
+			CreateGraphNode(Node, MinXPos, MaxXPos, DepthTracker, OutRemainingNodeKeys);
+		}
+	}
+
+	void SSignalFlowGraph::CreateSendLabels()
+	{
+		// Run through the list source again to create send label widgets between nodes
+		for (const TSharedPtr<ISignalFlowNode>& Node : *ListSource)
+		{
+			if (!Node.IsValid() || !Node->IsRealNode())
+			{
+				continue;
+			}
+
+			for (const FSignalFlowEntryKey& OutputKey : Node->FilteredOutputs)
+			{
+				const TOptional<FSignalFlowEntryKey> RealNodeOutputKey = FindRealNodeKey(OutputKey, ENodeConnectionType::Output);
+				if (RealNodeOutputKey.IsSet())
+				{
+					CreateSendLabelWidget(Node->GetEntryKey(), RealNodeOutputKey.GetValue());
+				}
+			}
+		}
+	}
+
+	void SSignalFlowGraph::CleanUpOldWidgets(const TArray<FSignalFlowEntryKey>& NodeKeysToRemove)
+	{
+		// Clean up and remove any old nodes that we did not detect in the new graph
+		for (const FSignalFlowEntryKey& Key : NodeKeysToRemove)
+		{
+			TSharedPtr<SSignalFlowGraphNode>* RemovedNode = NodeWidgets.Find(Key);;
+			if (!RemovedNode || !RemovedNode->IsValid())
+			{
+				NodeWidgets.Remove(Key);
+				continue;
+			}
+
+			Canvas->RemoveSlot(RemovedNode->ToSharedRef());
+
+			TArray<TSharedPtr<SSignalFlowGraphSendLabel>> SendLabels;
+			SendLevelWidgets.MultiFind(Key, SendLabels);
+			for (const TSharedPtr<SSignalFlowGraphSendLabel>& SendWidget : SendLabels)
+			{
+				Canvas->RemoveSlot(SendWidget.ToSharedRef());
+			}
+
+			SendLevelWidgets.Remove(Key);
+
+			if ((*RemovedNode)->GetEntry().IsValid())
+			{
+				for (const FSignalFlowEntryKey& InputKey : (*RemovedNode)->GetEntry()->Inputs)
+				{
+					const TOptional<FSignalFlowEntryKey> RealInputKey = FindRealNodeKey(InputKey, ENodeConnectionType::Input);
+					if (!RealInputKey.IsSet())
+					{
+						continue;
+					}
+
+					const TSharedPtr<SSignalFlowGraphSendLabel> InputSendLabel = FindSendLabel(RealInputKey.GetValue(), Key);
+					if (!InputSendLabel.IsValid())
+					{
+						continue;
+					}
+
+					SendLevelWidgets.RemoveSingle(RealInputKey.GetValue(), InputSendLabel);
+					Canvas->RemoveSlot(InputSendLabel.ToSharedRef());
+				}
+			}
+
+			NodeWidgets.Remove(Key);
+		}
+	}
+
 	void SSignalFlowGraph::CreateGraphNode(const TSharedPtr<ISignalFlowNode>& Node, const float MinXPos, const float MaxXPos, const int32 DepthTracker, TArray<FSignalFlowEntryKey>& OutRemainingNodeKeys)
 	{
 		using namespace FSignalFlowGraphPrivate;
@@ -459,24 +564,55 @@ namespace UE::Audio::Insights
 			const TSharedPtr<FDummyConnectionNode> FoundDummyCast = StaticCastSharedPtr<FDummyConnectionNode>(Node);
 			if (FoundDummyCast.IsValid())
 			{
-				DummyEntryToOutNode.Add(FoundDummyCast->GetEntryKey(), FoundDummyCast->ConnectionOutputKey);
+				DummyToRealConnectionKeys.Add(FoundDummyCast->GetEntryKey(), { FoundDummyCast->ConnectionInputKey, FoundDummyCast->ConnectionOutputKey });
 				return;
 			}
 		}
 		
-
 		const FSignalFlowEntryKey EntryKey = Node->GetEntryKey();
 
 		// Avoid creating this node again if it already exists
 		if (TSharedPtr<SSignalFlowGraphNode>* NodeWidget = NodeWidgets.Find(EntryKey))
 		{
-			if (NodeWidget->IsValid())
+			if (ensure(NodeWidget->IsValid()))
 			{
 				// Update the X and Y positions of the node
 				(*NodeWidget)->SetXPos(Node->XPos - MinXPos);
 				(*NodeWidget)->SetYPos(static_cast<float>(DepthTracker));
 
 				OutRemainingNodeKeys.RemoveSwap(EntryKey);
+
+				TArray<TSharedPtr<SSignalFlowGraphSendLabel>> SendLabels;
+				SendLevelWidgets.MultiFind(EntryKey, SendLabels);
+
+				TSet<FSignalFlowEntryKey> ActiveOutputKeys;
+				for (const FSignalFlowEntryKey& NodeOutputKey : Node->FilteredOutputs)
+				{
+					const TOptional<FSignalFlowEntryKey> RealNodeOutputKeyOpt = FindRealNodeKey(NodeOutputKey, ENodeConnectionType::Output);
+					if (RealNodeOutputKeyOpt.IsSet())
+					{
+						const FSignalFlowEntryKey RealNodeOutputKey = RealNodeOutputKeyOpt.GetValue();
+						ActiveOutputKeys.Add(RealNodeOutputKey);
+
+						const TSharedPtr<SSignalFlowGraphSendLabel>* FoundExistingSendLabel = FindSendLabelInList(SendLabels, RealNodeOutputKey);
+						
+						if (FoundExistingSendLabel == nullptr)
+						{
+							CreateSendLabelWidget(EntryKey, RealNodeOutputKey);
+						}
+					}
+				}
+
+				// Ensure any removed connections between active nodes are cleaned up
+				for (const TSharedPtr<SSignalFlowGraphSendLabel>& ExistingSendLabel : SendLabels)
+				{
+					if (!ActiveOutputKeys.Contains(ExistingSendLabel->GetSendReceiverEntryKey()))
+					{
+						SendLevelWidgets.RemoveSingle(EntryKey, ExistingSendLabel);
+						Canvas->RemoveSlot(ExistingSendLabel.ToSharedRef());
+					}
+				}
+
 				return;
 			}
 			else
@@ -564,6 +700,154 @@ namespace UE::Audio::Insights
 				NodeWidget.ToSharedRef()
 			];
 	}
+
+	void SSignalFlowGraph::CreateSendLabelWidget(const FSignalFlowEntryKey& SenderEntryKey, const FSignalFlowEntryKey& ReceiverEntryKey)
+	{
+		using namespace FSignalFlowGraphPrivate;
+
+		if (FindSendLabel(SenderEntryKey, ReceiverEntryKey).IsValid())
+		{
+			return;
+		}
+
+		TSharedPtr<SSignalFlowGraphSendLabel> SendLabelWidget = SNew(SSignalFlowGraphSendLabel)
+			.SendReceiverEntryKey(ReceiverEntryKey)
+			.SenderPositionInGraph_Lambda([this, SenderEntryKey]() -> FVector2D
+			{
+				const TSharedPtr<SSignalFlowGraphNode>* SenderNode = NodeWidgets.Find(SenderEntryKey);
+				if (SenderNode && SenderNode->IsValid())
+				{
+					return (*SenderNode)->GetLocalOutputPos();
+				}
+
+				return FVector2D::ZeroVector;
+			})
+			.ReceiverPositionInGraph_Lambda([this, ReceiverEntryKey]() -> FVector2D
+			{
+				const TSharedPtr<SSignalFlowGraphNode>* ReceiverNode = NodeWidgets.Find(ReceiverEntryKey);
+				if (ReceiverNode && ReceiverNode->IsValid())
+				{
+					return (*ReceiverNode)->GetLocalInputPos();
+				}
+
+				return FVector2D::ZeroVector;
+			})
+			.SendLevel_Lambda([this, SenderEntryKey, ReceiverEntryKey]() -> TOptional<float>
+			{
+				const TSharedPtr<SSignalFlowGraphNode>* SenderNode = NodeWidgets.Find(SenderEntryKey);
+				if (SenderNode && SenderNode->IsValid() && (*SenderNode)->GetEntry().IsValid())
+				{
+					const TSharedPtr<FSignalFlowDashboardEntry>& SenderEntry = (*SenderNode)->GetEntry();
+					const FSignalFlowOutputData* SendInfo = SenderEntry->Outputs.Find(ReceiverEntryKey);
+					if (SendInfo && SendInfo->SendLevel.IsSet())
+					{
+						return SendInfo->SendLevel;
+					}
+				}
+
+				return NullOpt;
+			})
+			.GraphOrientation_Lambda([this]()
+			{
+				return Orientation;
+			});
+
+		SendLevelWidgets.Add(SenderEntryKey, SendLabelWidget);
+
+		Canvas->AddSlot()
+
+		.Position_Lambda([this, SenderEntryKey, ReceiverEntryKey]() -> FVector2D
+		{
+			const TSharedPtr<SSignalFlowGraphSendLabel> FoundExistingSendLabel = FindSendLabel(SenderEntryKey, ReceiverEntryKey);
+
+			if (FoundExistingSendLabel.IsValid())
+			{
+				return FoundExistingSendLabel->UpdatePositionInGraph();
+			}
+
+			return FVector2D::ZeroVector;
+		})
+
+		.Size_Lambda([this, SenderEntryKey, ReceiverEntryKey]() -> FVector2D
+		{
+			const TSharedPtr<SSignalFlowGraphSendLabel> FoundExistingSendLabel = FindSendLabel(SenderEntryKey, ReceiverEntryKey);
+
+			if (FoundExistingSendLabel.IsValid())
+			{
+				return FoundExistingSendLabel->GetDesiredSize();
+			}
+
+			return FVector2D::ZeroVector;
+		})
+
+		.HAlign(HAlign_Left)
+		.VAlign(VAlign_Top)
+		[
+			SendLabelWidget.ToSharedRef()
+		];
+	}
+
+	TSharedPtr<SSignalFlowGraphNode> SSignalFlowGraph::FindRealNode(const FSignalFlowEntryKey& EntryKey, const ENodeConnectionType ConnectionType) const
+	{
+		const TSharedPtr<SSignalFlowGraphNode>* FoundNode = NodeWidgets.Find(EntryKey);
+		if (FoundNode && FoundNode->IsValid())
+		{
+			return *FoundNode;
+		}
+		else
+		{
+			// If this is a dummy node, try and find the real input/output node we are connected to instead
+			const ConnectedNodeKeys* FoundRealConnectionKeys = DummyToRealConnectionKeys.Find(EntryKey);
+			if (FoundRealConnectionKeys)
+			{
+				const FSignalFlowEntryKey& LookupKey = ConnectionType == ENodeConnectionType::Input ? FoundRealConnectionKeys->RealInputKey 
+																										: FoundRealConnectionKeys->RealOutputKey;
+
+				const TSharedPtr<SSignalFlowGraphNode>* FoundConnectedNode = NodeWidgets.Find(LookupKey);
+
+				if (FoundConnectedNode && FoundConnectedNode->IsValid())
+				{
+					return *FoundConnectedNode;
+				}
+			}
+		}
+
+		return nullptr;
+	}
+
+	TOptional<FSignalFlowEntryKey> SSignalFlowGraph::FindRealNodeKey(const FSignalFlowEntryKey& EntryKey, const ENodeConnectionType ConnectionType) const
+	{
+		const TSharedPtr<SSignalFlowGraphNode> RealNode = FindRealNode(EntryKey, ConnectionType);
+		if (RealNode.IsValid() && RealNode->GetEntry().IsValid())
+		{
+			return RealNode->GetNodeEntryKey();
+		}
+
+		return NullOpt;
+	}
+
+	TSharedPtr<SSignalFlowGraphSendLabel> SSignalFlowGraph::FindSendLabel(const FSignalFlowEntryKey& SenderKey, const FSignalFlowEntryKey& ReceiverKey) const
+	{
+		TArray<TSharedPtr<SSignalFlowGraphSendLabel>> SendLabels;
+		SendLevelWidgets.MultiFind(SenderKey, SendLabels);
+
+		const TSharedPtr<SSignalFlowGraphSendLabel>* FoundSendLabel = FindSendLabelInList(SendLabels, ReceiverKey);
+
+		if (FoundSendLabel && FoundSendLabel->IsValid())
+		{
+			return *FoundSendLabel;
+		}
+
+		return nullptr;
+	}
+
+	const TSharedPtr<SSignalFlowGraphSendLabel>* SSignalFlowGraph::FindSendLabelInList(const TArray<TSharedPtr<SSignalFlowGraphSendLabel>>& SendLabelWidgets, const FSignalFlowEntryKey& ReceiverKey) const
+	{
+		return SendLabelWidgets.FindByPredicate([&ReceiverKey](const TSharedPtr<SSignalFlowGraphSendLabel>& SendLabel)
+		{
+			return SendLabel.IsValid() && SendLabel->GetSendReceiverEntryKey() == ReceiverKey;
+		});
+	}
 } // namespace UE::Audio::Insights
 
 #undef LOCTEXT_NAMESPACE
diff --git a/Engine/Plugins/AudioInsights/Source/AudioInsights/Public/AudioInsightsStyle.h b/Engine/Plugins/AudioInsights/Source/AudioInsights/Public/AudioInsightsStyle.h
--- a/Engine/Plugins/AudioInsights/Source/AudioInsights/Public/AudioInsightsStyle.h
+++ b/Engine/Plugins/AudioInsights/Source/AudioInsights/Public/AudioInsightsStyle.h
@@ -111,6 +111,15 @@ namespace UE::Audio::Insights
 			return &FloatFormat;
 		};
 
+		const FNumberFormattingOptions* GetSendLevelFormat() const
+		{
+			static FNumberFormattingOptions FloatFormat;
+			FloatFormat.MinimumIntegralDigits = 1;
+			FloatFormat.MinimumFractionalDigits = 1;
+			FloatFormat.MaximumFractionalDigits = 2;
+			return &FloatFormat;
+		};
+
 		FText FormatSecondsAsTime(float InTimeSec) const
 		{
 			return FText::Format(LOCTEXT("TimeInSecondsFormat", "{0}s"), FText::AsNumber(InTimeSec, GetTimeFormat()));
diff --git a/Engine/Plugins/AudioInsights/Source/AudioInsights/Public/Messages/SignalFlowTraceMessages.h b/Engine/Plugins/AudioInsights/Source/AudioInsights/Public/Messages/SignalFlowTraceMessages.h
--- a/Engine/Plugins/AudioInsights/Source/AudioInsights/Public/Messages/SignalFlowTraceMessages.h
+++ b/Engine/Plugins/AudioInsights/Source/AudioInsights/Public/Messages/SignalFlowTraceMessages.h
@@ -105,7 +105,7 @@ namespace UE::Audio::Insights
 
 	struct FSignalFlowOutputData
 	{
-		float SendLevel = TNumericLimits<float>::Lowest();
+		TOptional<float> SendLevel;
 	};
 
 	struct FSignalFlowDashboardEntry : public IObjectDashboardEntry
diff --git a/Engine/Plugins/AudioInsights/Source/AudioInsights/Public/Views/SSignalFlowGraph.h b/Engine/Plugins/AudioInsights/Source/AudioInsights/Public/Views/SSignalFlowGraph.h
--- a/Engine/Plugins/AudioInsights/Source/AudioInsights/Public/Views/SSignalFlowGraph.h
+++ b/Engine/Plugins/AudioInsights/Source/AudioInsights/Public/Views/SSignalFlowGraph.h
@@ -53,6 +53,9 @@ namespace UE::Audio::Insights
 		UE_API FVector2D UpdatePositionInGraph();
 		UE_API FVector2D GetCachedPositionInGraph() const;
 
+		const FSignalFlowEntryKey& GetNodeEntryKey() const { return NodeEntryKey; }
+		const TSharedPtr<FSignalFlowDashboardEntry> GetEntry() const { return Entry; }
+
 		float GetXPos() const { return XPos.Get(); }
 		float GetYPos() const { return YPos.Get(); }
 
@@ -69,6 +72,45 @@ namespace UE::Audio::Insights
 
 		FOnSignalFlowNodeSelected OnNodeSelected;
 		FVector2d CachedPositionInGraph = FVector2d::Zero();
+		FSignalFlowEntryKey NodeEntryKey;
+	};
+
+	class SSignalFlowGraphSendLabel : public SCompoundWidget
+	{
+	public:
+		SLATE_BEGIN_ARGS(SSignalFlowGraphSendLabel)
+			: _SendReceiverEntryKey(),
+			_SenderPositionInGraph(FVector2d::ZeroVector),
+			_ReceiverPositionInGraph(FVector2d::ZeroVector),
+			_SendLevel(),
+			_GraphOrientation(EOrientation::Orient_Vertical)
+		{
+		}
+
+		SLATE_ARGUMENT(FSignalFlowEntryKey, SendReceiverEntryKey)
+		SLATE_ATTRIBUTE(FVector2d, SenderPositionInGraph)
+		SLATE_ATTRIBUTE(FVector2d, ReceiverPositionInGraph)
+		SLATE_ATTRIBUTE(TOptional<float>, SendLevel)
+		SLATE_ATTRIBUTE(EOrientation, GraphOrientation)
+
+		SLATE_END_ARGS()
+
+		UE_API void Construct(const FArguments& InArgs);
+		UE_API virtual FVector2D ComputeDesiredSize(float InLayoutScaleMultiplier) const override;
+		
+		UE_API FVector2D UpdatePositionInGraph();
+		UE_API FVector2D GetCachedPositionInGraph() const;
+
+		const FSignalFlowEntryKey& GetSendReceiverEntryKey() const { return SendReceiverEntryKey; }
+
+	private:
+		FSignalFlowEntryKey SendReceiverEntryKey;
+		TAttribute<FVector2d> SenderPositionInGraph;
+		TAttribute<FVector2d> ReceiverPositionInGraph;
+		TAttribute<TOptional<float>> SendLevel;
+		TAttribute<EOrientation> GraphOrientation;
+
+		FVector2d CachedPositionInGraph = FVector2d::Zero();
 	};
 
 	class SSignalFlowGraph : public SCompoundWidget
@@ -98,12 +140,38 @@ namespace UE::Audio::Insights
 
 	private:
 		int32 OnPaintConnections(const FGeometry& AllottedGeometry, FSlateWindowElementList& OutDrawElements, int32 LayerId) const;
+
+		void CreateGraphNodes(TArray<FSignalFlowEntryKey>& OutRemainingNodeKeys);
+		void CreateSendLabels();
+		void CleanUpOldWidgets(const TArray<FSignalFlowEntryKey>& NodeKeysToRemove);
+
 		void CreateGraphNode(const TSharedPtr<ISignalFlowNode>& Node, const float MinXPos, const float MaxXPos, const int32 DepthTracker, TArray<FSignalFlowEntryKey>& OutRemainingNodeKeys);
+		void CreateSendLabelWidget(const FSignalFlowEntryKey& SenderEntryKey, const FSignalFlowEntryKey& ReceiverEntryKey);
+
+		enum class ENodeConnectionType : uint8
+		{
+			Input,
+			Output
+		};
+
+		TSharedPtr<SSignalFlowGraphNode> FindRealNode(const FSignalFlowEntryKey& EntryKey, const ENodeConnectionType ConnectionType) const;
+		TOptional<FSignalFlowEntryKey> FindRealNodeKey(const FSignalFlowEntryKey& EntryKey, const ENodeConnectionType ConnectionType) const;
+
+		TSharedPtr<SSignalFlowGraphSendLabel> FindSendLabel(const FSignalFlowEntryKey& SenderKey, const FSignalFlowEntryKey& ReceiverKey) const;
+		const TSharedPtr<SSignalFlowGraphSendLabel>* FindSendLabelInList(const TArray<TSharedPtr<SSignalFlowGraphSendLabel>>& SendLabelWidgets, const FSignalFlowEntryKey& ReceiverKey) const;
 
 		TSharedPtr<SCanvas> Canvas;
 
 		TMap<FSignalFlowEntryKey, TSharedPtr<SSignalFlowGraphNode>> NodeWidgets;
-		TMap<FSignalFlowEntryKey, FSignalFlowEntryKey> DummyEntryToOutNode;
+
+		struct ConnectedNodeKeys
+		{
+			FSignalFlowEntryKey RealInputKey;
+			FSignalFlowEntryKey RealOutputKey;
+		};
+
+		TMap<FSignalFlowEntryKey, ConnectedNodeKeys> DummyToRealConnectionKeys;
+		TMultiMap<FSignalFlowEntryKey, TSharedPtr<SSignalFlowGraphSendLabel>> SendLevelWidgets;
 
 		TArray<TSharedPtr<ISignalFlowNode>>* ListSource;
 		TAttribute<TSharedPtr<FSignalFlowDashboardEntry>> SelectedItem;
