[Audio Insights] Initial work on signal flow toolbar part 01: Moved helper functions for creating labeled and toggle button content from the sounds dashboard to the style class

#Synced-CL 49593349

diff --git a/Engine/Plugins/AudioInsights/Source/AudioInsights/Private/AudioInsightsStyle.cpp b/Engine/Plugins/AudioInsights/Source/AudioInsights/Private/AudioInsightsStyle.cpp
--- a/Engine/Plugins/AudioInsights/Source/AudioInsights/Private/AudioInsightsStyle.cpp
+++ b/Engine/Plugins/AudioInsights/Source/AudioInsights/Private/AudioInsightsStyle.cpp
@@ -2,6 +2,10 @@
 
 #include "AudioInsightsStyle.h"
 
+#include "Widgets/Images/SImage.h"
+#include "Widgets/SBoxPanel.h"
+#include "Widgets/Text/STextBlock.h"
+
 namespace UE::Audio::Insights
 {
 	FSlateStyle& FSlateStyle::Get()
@@ -9,4 +13,91 @@ namespace UE::Audio::Insights
 		static FSlateStyle InsightsStyle;
 		return InsightsStyle;
 	}
+
+	TSharedRef<SBox> FSlateStyle::CreateButtonContentWidget(const FName& IconName /*= FName()*/, const FText& Label /*= FText::GetEmpty()*/, const FName& TextStyle /*= TEXT("ButtonText")*/, const float HeightOverride /*= 16.0f*/) const
+	{
+		TSharedRef<SHorizontalBox> ButtonContainerWidget = SNew(SHorizontalBox);
+
+		// Button icon (optional)
+		if (!IconName.IsNone())
+		{
+			ButtonContainerWidget->AddSlot()
+			.AutoWidth()
+			.VAlign(VAlign_Center)
+			[
+				SNew(SImage)
+				.ColorAndOpacity(FSlateColor::UseForeground())
+				.Image(FSlateStyle::Get().GetBrush(IconName))
+			];
+		}
+
+		// Button text (optional)
+		if (!Label.IsEmpty())
+		{
+			const float LeftPadding = IconName.IsNone() ? 0.0f : 4.0f;
+
+			ButtonContainerWidget->AddSlot()
+				.VAlign(VAlign_Center)
+				.Padding(LeftPadding, 0.0f, 0.0f, 0.0f)
+				.AutoWidth()
+				[
+					SNew(STextBlock)
+					.TextStyle(&FSlateStyle::Get().GetWidgetStyle<FTextBlockStyle>(TextStyle))
+					.Justification(ETextJustify::Center)
+					.Text(Label)
+				];
+		}
+
+		return SNew(SBox)
+		.HeightOverride(HeightOverride)
+		[
+			ButtonContainerWidget
+		];
+	};
+
+	TSharedRef<SBox> FSlateStyle::CreateToggleButtonContent(const FName& IconName, const FName& TextStyle, TFunction<FSlateColor()> ColorFunction, TFunction<FText()> TextFunction, const float HeightOverride /*= 16.0f*/) const
+	{
+		TSharedRef<SHorizontalBox> ButtonContainerWidget = SNew(SHorizontalBox);
+
+		if (!IconName.IsNone())
+		{
+			ButtonContainerWidget->AddSlot()
+			.AutoWidth()
+			.VAlign(VAlign_Center)
+			[
+				SNew(SImage)
+				.ColorAndOpacity_Lambda([ColorFunction]()
+				{ 
+					return ColorFunction();
+				})
+				.Image(FSlateStyle::Get().GetBrush(IconName))
+			];
+		}
+
+		const float LeftPadding = IconName.IsNone() ? 0.0f : 4.0f;
+
+		ButtonContainerWidget->AddSlot()
+		.VAlign(VAlign_Center)
+		.Padding(LeftPadding, 0.0f, 0.0f, 0.0f)
+		.AutoWidth()
+		[
+			SNew(STextBlock)
+			.TextStyle(&FSlateStyle::Get().GetWidgetStyle<FTextBlockStyle>(TextStyle))
+			.Justification(ETextJustify::Center)
+			.Text_Lambda([TextFunction]()
+			{ 
+				return TextFunction();
+			})
+			.ColorAndOpacity_Lambda([ColorFunction]()
+			{
+				return ColorFunction();
+			})
+		];
+
+		return SNew(SBox)
+		.HeightOverride(HeightOverride)
+		[
+			ButtonContainerWidget
+		];
+	};
 }
diff --git a/Engine/Plugins/AudioInsights/Source/AudioInsights/Private/Views/SoundDashboardViewFactory.cpp b/Engine/Plugins/AudioInsights/Source/AudioInsights/Private/Views/SoundDashboardViewFactory.cpp
--- a/Engine/Plugins/AudioInsights/Source/AudioInsights/Private/Views/SoundDashboardViewFactory.cpp
+++ b/Engine/Plugins/AudioInsights/Source/AudioInsights/Private/Views/SoundDashboardViewFactory.cpp
@@ -136,87 +136,6 @@ namespace UE::Audio::Insights
 			}
 		};
 
-		TSharedRef<SBox> CreateButtonContentWidget(const FName& InIconName = FName(), const FText& InLabel = FText::GetEmpty(), const FName& InTextStyle = TEXT("ButtonText"))
-		{
-			TSharedRef<SHorizontalBox> ButtonContainerWidget = SNew(SHorizontalBox);
-
-			// Button icon (optional)
-			if (!InIconName.IsNone())
-			{
-				ButtonContainerWidget->AddSlot()
-					.AutoWidth()
-					.VAlign(VAlign_Center)
-					[
-						SNew(SImage)
-							.ColorAndOpacity(FSlateColor::UseForeground())
-							.Image(FSlateStyle::Get().GetBrush(InIconName))
-					];
-			}
-
-			// Button text (optional)
-			if (!InLabel.IsEmpty())
-			{
-				const float LeftPadding = InIconName.IsNone() ? 0.0f : 4.0f;
-
-				ButtonContainerWidget->AddSlot()
-					.VAlign(VAlign_Center)
-					.Padding(LeftPadding, 0.0f, 0.0f, 0.0f)
-					.AutoWidth()
-					[
-						SNew(STextBlock)
-							.TextStyle(&FSlateStyle::Get().GetWidgetStyle<FTextBlockStyle>(InTextStyle))
-							.Justification(ETextJustify::Center)
-							.Text(InLabel)
-					];
-			}
-
-			return SNew(SBox)
-				.HeightOverride(16.0f)
-				[
-					ButtonContainerWidget
-				];
-		};
-
-		TSharedRef<SBox> CreateMuteSoloSelectedButtonContent(const FName& InIconName,
-			const FName& InTextStyle,
-			TFunction<FSlateColor()> InColorFunction,
-			TFunction<FText()> InTextFunction)
-		{
-			TSharedRef<SHorizontalBox> ButtonContainerWidget = SNew(SHorizontalBox);
-
-			if (!InIconName.IsNone())
-			{
-				ButtonContainerWidget->AddSlot()
-					.AutoWidth()
-					.VAlign(VAlign_Center)
-					[
-						SNew(SImage)
-							.ColorAndOpacity_Lambda([InColorFunction]() { return InColorFunction(); })
-							.Image(FSlateStyle::Get().GetBrush(InIconName))
-					];
-			}
-
-			const float LeftPadding = InIconName.IsNone() ? 0.0f : 4.0f;
-
-			ButtonContainerWidget->AddSlot()
-				.VAlign(VAlign_Center)
-				.Padding(LeftPadding, 0.0f, 0.0f, 0.0f)
-				.AutoWidth()
-				[
-					SNew(STextBlock)
-						.TextStyle(&FSlateStyle::Get().GetWidgetStyle<FTextBlockStyle>(InTextStyle))
-						.Justification(ETextJustify::Center)
-						.Text_Lambda([InTextFunction]() { return InTextFunction(); })
-						.ColorAndOpacity_Lambda([InColorFunction]() { return InColorFunction(); })
-				];
-
-			return SNew(SBox)
-				.HeightOverride(16.0f)
-				[
-					ButtonContainerWidget
-				];
-		};
-
 		void RefreshInitExpandStateRecursive(FSoundDashboardEntry& SoundEntry)
 		{
 			SoundEntry.bHasSetInitExpansion = false;
@@ -907,7 +826,7 @@ namespace UE::Audio::Insights
 					return FReply::Handled();
 				})
 				[
-					CreateMuteSoloSelectedButtonContent("AudioInsights.Icon.SoundDashboard.Mute", "SmallButtonText", MuteColorFunc, MuteTextFunc)
+					FSlateStyle::Get().CreateToggleButtonContent("AudioInsights.Icon.SoundDashboard.Mute", "SmallButtonText", MuteColorFunc, MuteTextFunc)
 				]
 			]
 			// Solo button
@@ -936,7 +855,7 @@ namespace UE::Audio::Insights
 					return FReply::Handled();
 				})
 				[
-					CreateMuteSoloSelectedButtonContent("AudioInsights.Icon.SoundDashboard.Solo", "SmallButtonText", SoloColorFunc, SoloTextFunc)
+					FSlateStyle::Get().CreateToggleButtonContent("AudioInsights.Icon.SoundDashboard.Solo", "SmallButtonText", SoloColorFunc, SoloTextFunc)
 				]
 			]
 			// Clear Mutes/Solos button
@@ -956,7 +875,7 @@ namespace UE::Audio::Insights
 					return FReply::Handled();
 				})
 				[
-					CreateButtonContentWidget("AudioInsights.Icon.SoundDashboard.Reset", LOCTEXT("SoundsDashboard_ClearMutesAndSolosButtonText", "Clear All Mutes/Solos"), "SmallButtonText")
+					FSlateStyle::Get().FSlateStyle::Get().CreateButtonContentWidget("AudioInsights.Icon.SoundDashboard.Reset", LOCTEXT("SoundsDashboard_ClearMutesAndSolosButtonText", "Clear All Mutes/Solos"), "SmallButtonText")
 				]
 			];
 	}
@@ -3016,7 +2935,7 @@ namespace UE::Audio::Insights
 					return FReply::Handled();
 				})
 				[
-					CreateButtonContentWidget("AudioInsights.Icon.SoundDashboard.Plots", LOCTEXT("SoundDashboard_PlotsButtonText", "Plot Selected"), "SmallButtonText")
+					FSlateStyle::Get().CreateButtonContentWidget("AudioInsights.Icon.SoundDashboard.Plots", LOCTEXT("SoundDashboard_PlotsButtonText", "Plot Selected"), "SmallButtonText")
 				]
 			]
 			// Clear All Plots button
@@ -3037,7 +2956,7 @@ namespace UE::Audio::Insights
 					return FReply::Handled();
 				})
 				[
-					CreateButtonContentWidget("AudioInsights.Icon.SoundDashboard.Reset", LOCTEXT("SoundsDashboard_ClearAllPlotsButtonText", "Clear All Plots"), "SmallButtonText")
+					FSlateStyle::Get().CreateButtonContentWidget("AudioInsights.Icon.SoundDashboard.Reset", LOCTEXT("SoundsDashboard_ClearAllPlotsButtonText", "Clear All Plots"), "SmallButtonText")
 				]
 			];
 	}
diff --git a/Engine/Plugins/AudioInsights/Source/AudioInsights/Public/AudioInsightsStyle.h b/Engine/Plugins/AudioInsights/Source/AudioInsights/Public/AudioInsightsStyle.h
--- a/Engine/Plugins/AudioInsights/Source/AudioInsights/Public/AudioInsightsStyle.h
+++ b/Engine/Plugins/AudioInsights/Source/AudioInsights/Public/AudioInsightsStyle.h
@@ -9,6 +9,7 @@
 #include "Styling/SlateTypes.h"
 #include "Styling/StyleColors.h"
 #include "Textures/SlateIcon.h"
+#include "Widgets/Layout/SBox.h"
 
 #define LOCTEXT_NAMESPACE "AudioInsights"
 
@@ -165,6 +166,9 @@ namespace UE::Audio::Insights
 			return *DefaultBrush;
 		}
 
+		TSharedRef<SBox> CreateButtonContentWidget(const FName& IconName = FName(), const FText& Label = FText::GetEmpty(), const FName& TextStyle = TEXT("ButtonText"), const float HeightOverride = 16.0f) const;
+		TSharedRef<SBox> CreateToggleButtonContent(const FName& IconName, const FName& TextStyle, TFunction<FSlateColor()> ColorFunction, TFunction<FText()> TextFunction, const float HeightOverride = 16.0f) const;
+
 		FSlateStyle()
 			: FSlateStyleSet(GetStyleName())
 		{
