[Audio Insights] Initial work on signal flow toolbar part 02: Created toggle button for Horizontal Flow
* Added icons for upcoming buttons to Audio Insights Style
* Changed hotkey for toggle horizontal flow from F to Ctrl + H
* Moved keyboard focus for listening out for hotkeys from graph widget to whole signal flow dashboard

#Synced-CL 49593349

diff --git a/Engine/Plugins/AudioInsights/Source/AudioInsights/Private/Views/SSignalFlowGraph.cpp b/Engine/Plugins/AudioInsights/Source/AudioInsights/Private/Views/SSignalFlowGraph.cpp
--- a/Engine/Plugins/AudioInsights/Source/AudioInsights/Private/Views/SSignalFlowGraph.cpp
+++ b/Engine/Plugins/AudioInsights/Source/AudioInsights/Private/Views/SSignalFlowGraph.cpp
@@ -340,6 +340,13 @@ namespace UE::Audio::Insights
 		CleanUpOldWidgets(RemainingNodeKeys);
 	}
 
+	void SSignalFlowGraph::ToggleOrientation()
+	{
+		Orientation = (Orientation == EOrientation::Orient_Vertical) ? EOrientation::Orient_Horizontal : EOrientation::Orient_Vertical;
+
+		RefreshGraph();
+	}
+
 	int32 SSignalFlowGraph::OnPaint(const FPaintArgs& Args, const FGeometry& AllottedGeometry, const FSlateRect& MyCullingRect, FSlateWindowElementList& OutDrawElements, int32 LayerId, const FWidgetStyle& InWidgetStyle, bool bParentEnabled) const
 	{
 		LayerId = OnPaintConnections(AllottedGeometry, OutDrawElements, LayerId);
@@ -407,18 +414,6 @@ namespace UE::Audio::Insights
 		return LayerId;
 	}
 
-	FReply SSignalFlowGraph::OnKeyDown(const FGeometry& MyGeometry, const FKeyEvent& InKeyEvent)
-	{
-		if (InKeyEvent.GetKey() == EKeys::F)
-		{
-			Orientation = (Orientation == EOrientation::Orient_Vertical) ? EOrientation::Orient_Horizontal : EOrientation::Orient_Vertical;
-			RefreshGraph();
-			return FReply::Handled();
-		}
-
-		return SCompoundWidget::OnKeyDown(MyGeometry, InKeyEvent);
-	}
-
 	FVector2D SSignalFlowGraph::ComputeDesiredSize(float InLayoutScaleMultiplier) const
 	{
 		using namespace FSignalFlowGraphPrivate;
diff --git a/Engine/Plugins/AudioInsights/Source/AudioInsights/Private/Views/SignalFlowDashboardViewFactory.cpp b/Engine/Plugins/AudioInsights/Source/AudioInsights/Private/Views/SignalFlowDashboardViewFactory.cpp
--- a/Engine/Plugins/AudioInsights/Source/AudioInsights/Private/Views/SignalFlowDashboardViewFactory.cpp
+++ b/Engine/Plugins/AudioInsights/Source/AudioInsights/Private/Views/SignalFlowDashboardViewFactory.cpp
@@ -25,10 +25,27 @@ namespace UE::Audio::Insights
 {
 	namespace FSignalFlowDashboardViewFactoryPrivate
 	{
+		const FName ButtonStyleName = "SimpleButton";
+		const FName ButtonTextStyleName = "SmallButtonText";
+
 		constexpr float GraphWidthProp = 0.8f;
 		constexpr float SourcesListHeightProp = 0.5f;
 		constexpr float BusListHeightProp = 0.2f;
 		constexpr float SubmixesListHeightProp = 0.3f;
+
+		constexpr float ToolbarButtonHeight = 20.0f;
+		constexpr float ToolbarButtonPaddingX = 8.0f;
+		constexpr float ToolbarButtonPaddingY = 8.0f;
+	}
+
+	FReply SSignalFlowDashboard::OnKeyDown(const FGeometry& MyGeometry, const FKeyEvent& InKeyEvent)
+	{
+		if (OnKeyDownHandler.IsBound())
+		{
+			return OnKeyDownHandler.Execute(MyGeometry, InKeyEvent);
+		}
+
+		return FReply::Unhandled();
 	}
 
 	FSignalFlowDashboardViewFactory::FSignalFlowDashboardViewFactory()
@@ -82,24 +99,53 @@ namespace UE::Audio::Insights
 
 		if (!DashboardWidget.IsValid())
 		{
-			SAssignNew(DashboardWidget, SHorizontalBox)
+			SAssignNew(DashboardWidget, SSignalFlowDashboard)
 
-			+ SHorizontalBox::Slot()
-			.FillWidth(1.0f - GraphWidthProp)
+			+ SSignalFlowDashboard::Slot()
+			.AutoHeight()
 			[
-				CreateSelectionPanel()
+				CreateToolbar()
 			]
-			
-			+ SHorizontalBox::Slot()
-			.FillWidth(GraphWidthProp)
+
+			+ SSignalFlowDashboard::Slot()
+			.FillHeight(1.0f)
 			[
-				CreateGraphSection()
+				SNew(SHorizontalBox)
+
+				+ SHorizontalBox::Slot()
+				.FillWidth(1.0f - GraphWidthProp)
+				[
+					CreateSelectionPanel()
+				]
+				
+				+ SHorizontalBox::Slot()
+				.FillWidth(GraphWidthProp)
+				[
+					CreateGraphSection()
+				]
 			];
+
+			DashboardWidget->SetOnKeyDownHandler(FOnKeyDown::CreateSP(this, &FSignalFlowDashboardViewFactory::OnKeyDown));
 		}
 
 		return DashboardWidget->AsShared();
 	}
 
+	FReply FSignalFlowDashboardViewFactory::OnKeyDown(const FGeometry& MyGeometry, const FKeyEvent& InKeyEvent)
+	{
+		if (GraphWidgetContents.IsValid())
+		{
+			if (InKeyEvent.GetKey() == EKeys::H && InKeyEvent.IsControlDown())
+			{
+				GraphWidgetContents->ToggleOrientation();
+
+				return FReply::Handled();
+			}
+		}
+
+		return FReply::Unhandled();
+	}
+
 	void FSignalFlowDashboardViewFactory::SRowWidget::Construct(const FArguments& InArgs, const TSharedRef<STableViewBase>& InOwnerTable, TSharedPtr<FSignalFlowDashboardEntry> Entry)
 	{
 		if (!Entry.IsValid())
@@ -127,6 +173,58 @@ namespace UE::Audio::Insights
 		return &FSlateStyle::Get().GetWidgetStyle<FTableRowStyle>("TreeDashboard.TableViewRow");
 	}
 
+	TSharedRef<SWidget> FSignalFlowDashboardViewFactory::CreateToolbar()
+	{
+		using namespace FSignalFlowDashboardViewFactoryPrivate;
+
+		return SNew(SHorizontalBox)
+
+		+ SHorizontalBox::Slot()
+		.AutoWidth()
+		[
+			CreateToolbarToggleControls()
+		];
+	}
+
+	TSharedRef<SWidget> FSignalFlowDashboardViewFactory::CreateToolbarToggleControls()
+	{
+		using namespace FSignalFlowDashboardViewFactoryPrivate;
+
+		const FSlateStyle& Style = FSlateStyle::Get();
+
+		TFunction<FSlateColor()> HorizontalFlowColorFunc = [this]()
+		{ 
+			return GraphWidgetContents.IsValid() && GraphWidgetContents->GetGraphOrientation() == Orient_Horizontal ? FSlateColor(FColor::Green) : FSlateColor::UseForeground();
+		};
+
+		return SNew(SHorizontalBox)
+
+		+ SHorizontalBox::Slot()
+		.AutoWidth()
+		.Padding(ToolbarButtonPaddingX, ToolbarButtonPaddingY)
+		[
+			SNew(SButton)
+			.ButtonStyle(&FAppStyle::Get().GetWidgetStyle<FButtonStyle>(ButtonStyleName))
+			.ToolTipText(LOCTEXT("AudioDashboard_SignalFlow_HorizontalFlowTooltip", "Toggles between horizontally/vertically orienting the Signal Flow graph. (Ctrl + h)"))
+			.OnClicked_Lambda([this]()
+			{
+				if (GraphWidgetContents.IsValid())
+				{
+					GraphWidgetContents->ToggleOrientation();
+				}
+
+				return FReply::Handled();
+			})
+			[
+				Style.CreateToggleButtonContent("AudioInsights.Icon.SignalFlow.HorizontalFlow",
+												ButtonTextStyleName,
+												HorizontalFlowColorFunc,
+												[]() { return LOCTEXT("AudioDashboard_SignalFlow_HorizontalFlow", "Horizontal Flow"); },
+												ToolbarButtonHeight)
+			]
+		];
+	}
+
 	TSharedRef<SWidget> FSignalFlowDashboardViewFactory::CreateSelectionPanel()
 	{
 		using namespace FSignalFlowDashboardViewFactoryPrivate;
diff --git a/Engine/Plugins/AudioInsights/Source/AudioInsights/Public/AudioInsightsStyle.h b/Engine/Plugins/AudioInsights/Source/AudioInsights/Public/AudioInsightsStyle.h
--- a/Engine/Plugins/AudioInsights/Source/AudioInsights/Public/AudioInsightsStyle.h
+++ b/Engine/Plugins/AudioInsights/Source/AudioInsights/Public/AudioInsightsStyle.h
@@ -248,6 +248,12 @@ namespace UE::Audio::Insights
 
 			/* Event Log */
 			Set("AudioInsights.Icon.EventLog.MarkedForDelete",		  new IMAGE_BRUSH_SVG(TEXT("Icons/EventLog/cache_marked_for_delete"), Icon16));
+
+			/* Signal Flow */
+			Set("AudioInsights.Icon.SignalFlow.HorizontalFlow",		  new IMAGE_BRUSH_SVG(TEXT("Icons/SignalFlow/horizontal_flow"),		 Icon20));
+			Set("AudioInsights.Icon.SignalFlow.NodeDetails",		  new IMAGE_BRUSH_SVG(TEXT("Icons/SignalFlow/node_details"),		 Icon20));
+			Set("AudioInsights.Icon.SignalFlow.CenterView",			  new IMAGE_BRUSH_SVG(TEXT("Icons/SignalFlow/center_view"),			 Icon20));
+			Set("AudioInsights.Icon.SignalFlow.Settings",			  new IMAGE_BRUSH_SVG(TEXT("Icons/SignalFlow/settings"),			 Icon20));
 			
 			// Category colors
 			Set("SoundDashboard.MetaSoundColor", FLinearColor(0.008f, 0.76f, 0.078f));
diff --git a/Engine/Plugins/AudioInsights/Source/AudioInsights/Public/Views/SSignalFlowGraph.h b/Engine/Plugins/AudioInsights/Source/AudioInsights/Public/Views/SSignalFlowGraph.h
--- a/Engine/Plugins/AudioInsights/Source/AudioInsights/Public/Views/SSignalFlowGraph.h
+++ b/Engine/Plugins/AudioInsights/Source/AudioInsights/Public/Views/SSignalFlowGraph.h
@@ -130,10 +130,10 @@ namespace UE::Audio::Insights
 		UE_API void Construct(const FArguments& InArgs);
 		UE_API void RefreshGraph();
 
-		UE_API virtual int32 OnPaint(const FPaintArgs& Args, const FGeometry& AllottedGeometry, const FSlateRect& MyCullingRect, FSlateWindowElementList& OutDrawElements, int32 LayerId, const FWidgetStyle& InWidgetStyle, bool bParentEnabled) const override;
+		UE_API void ToggleOrientation();
+		UE_API EOrientation GetGraphOrientation() const { return Orientation; }
 
-		UE_API virtual FReply OnKeyDown(const FGeometry& MyGeometry, const FKeyEvent& InKeyEvent) override;
-		UE_API virtual bool SupportsKeyboardFocus() const override { return true; }
+		UE_API virtual int32 OnPaint(const FPaintArgs& Args, const FGeometry& AllottedGeometry, const FSlateRect& MyCullingRect, FSlateWindowElementList& OutDrawElements, int32 LayerId, const FWidgetStyle& InWidgetStyle, bool bParentEnabled) const override;
 
 	protected:
 		UE_API virtual FVector2D ComputeDesiredSize(float InLayoutScaleMultiplier) const override;
diff --git a/Engine/Plugins/AudioInsights/Source/AudioInsights/Public/Views/SignalFlowDashboardViewFactory.h b/Engine/Plugins/AudioInsights/Source/AudioInsights/Public/Views/SignalFlowDashboardViewFactory.h
--- a/Engine/Plugins/AudioInsights/Source/AudioInsights/Public/Views/SignalFlowDashboardViewFactory.h
+++ b/Engine/Plugins/AudioInsights/Source/AudioInsights/Public/Views/SignalFlowDashboardViewFactory.h
@@ -19,6 +19,18 @@ namespace UE::Audio::Insights
 	struct FSignalFlowDashboardEntry;
 	struct TreeDepthPair;
 
+	class SSignalFlowDashboard : public SVerticalBox
+	{
+	public:
+		UE_API virtual FReply OnKeyDown(const FGeometry& MyGeometry, const FKeyEvent& InKeyEvent) override;
+		virtual bool SupportsKeyboardFocus() const override { return true; }
+
+		UE_API void SetOnKeyDownHandler(FOnKeyDown Handler) { OnKeyDownHandler = Handler; }
+
+	private:
+		FOnKeyDown OnKeyDownHandler;
+	};
+
 	class FSignalFlowDashboardViewFactory : public FTraceDashboardViewFactoryBase, public TSharedFromThis<FSignalFlowDashboardViewFactory>
 	{
 	public:
@@ -31,6 +43,8 @@ namespace UE::Audio::Insights
 		UE_API virtual EDefaultDashboardTabStack GetDefaultTabStack() const override;
 		UE_API virtual TSharedRef<SWidget> MakeWidget(TSharedRef<SDockTab> OwnerTab, const FSpawnTabArgs& SpawnTabArgs) override;
 
+		UE_API FReply OnKeyDown(const FGeometry& MyGeometry, const FKeyEvent& InKeyEvent);
+
 	private:
 		struct SRowWidget : public STableRow<TSharedPtr<FSignalFlowDashboardEntry>>
 		{
@@ -48,6 +62,9 @@ namespace UE::Audio::Insights
 
 		const FTableRowStyle* GetRowStyle() const;
 
+		TSharedRef<SWidget> CreateToolbar();
+		TSharedRef<SWidget> CreateToolbarToggleControls();
+
 		TSharedRef<SWidget> CreateSelectionPanel();
 		TSharedRef<SWidget> CreateGraphSection();
 		TSharedRef<SListView<TSharedPtr<FSignalFlowDashboardEntry>>> CreateSelectableList(const FText& HeadingText, const TArray<TSharedPtr<FSignalFlowDashboardEntry>>& ListSource);
@@ -87,7 +104,7 @@ namespace UE::Audio::Insights
 		void MakeSpaceForNodeInRow(const TSharedPtr<ISignalFlowNode>& Node, const TArray<TSharedPtr<ISignalFlowNode>>& PrevPositionedNodes, const int32 NumNodesInRow);
 
 		TSharedPtr<FSignalFlowTraceProvider> SignalFlowProvider;
-		TSharedPtr<SWidget> DashboardWidget;
+		TSharedPtr<SSignalFlowDashboard> DashboardWidget;
 
 		TSharedPtr<SListView<TSharedPtr<FSignalFlowDashboardEntry>>> ActiveSourcesListView;
 		TArray<TSharedPtr<FSignalFlowDashboardEntry>> ActiveSources;
