[Audio Insights] Add button to center view on selected node in signal flow
* If no node is selected, we center on the Audio Device node

#Synced-CL 49617685

diff --git a/Engine/Plugins/AudioInsights/Source/AudioInsights/Private/Views/SSignalFlowGraph.cpp b/Engine/Plugins/AudioInsights/Source/AudioInsights/Private/Views/SSignalFlowGraph.cpp
--- a/Engine/Plugins/AudioInsights/Source/AudioInsights/Private/Views/SSignalFlowGraph.cpp
+++ b/Engine/Plugins/AudioInsights/Source/AudioInsights/Private/Views/SSignalFlowGraph.cpp
@@ -347,6 +347,17 @@ namespace UE::Audio::Insights
 		RefreshGraph();
 	}
 
+	TSharedPtr<SSignalFlowGraphNode> SSignalFlowGraph::FindNodeWidget(const FSignalFlowEntryKey& NodeEntryKey) const
+	{
+		const TSharedPtr<SSignalFlowGraphNode>* NodeWidget = NodeWidgets.Find(NodeEntryKey);
+		if (NodeWidget && NodeWidget->IsValid())
+		{
+			return *NodeWidget;
+		}
+
+		return nullptr;
+	}
+
 	int32 SSignalFlowGraph::OnPaint(const FPaintArgs& Args, const FGeometry& AllottedGeometry, const FSlateRect& MyCullingRect, FSlateWindowElementList& OutDrawElements, int32 LayerId, const FWidgetStyle& InWidgetStyle, bool bParentEnabled) const
 	{
 		LayerId = OnPaintConnections(AllottedGeometry, OutDrawElements, LayerId);
diff --git a/Engine/Plugins/AudioInsights/Source/AudioInsights/Private/Views/SignalFlowDashboardViewFactory.cpp b/Engine/Plugins/AudioInsights/Source/AudioInsights/Private/Views/SignalFlowDashboardViewFactory.cpp
--- a/Engine/Plugins/AudioInsights/Source/AudioInsights/Private/Views/SignalFlowDashboardViewFactory.cpp
+++ b/Engine/Plugins/AudioInsights/Source/AudioInsights/Private/Views/SignalFlowDashboardViewFactory.cpp
@@ -10,7 +10,7 @@
 #include "Styling/StyleColors.h"
 #include "Views/SignalFlowNodes.h"
 #include "Widgets/Colors/SColorBlock.h"
-#include "Widgets/Layout/SScrollBox.h"
+#include "Widgets/Layout/SSeparator.h"
 
 #if WITH_EDITOR
 #include "Editor.h"
@@ -141,6 +141,10 @@ namespace UE::Audio::Insights
 
 				return FReply::Handled();
 			}
+			else if (InKeyEvent.GetKey() == EKeys::F && InKeyEvent.IsControlDown())
+			{
+				return ScrollGraphViewToSelectedNode();
+			}
 		}
 
 		return FReply::Unhandled();
@@ -183,6 +187,20 @@ namespace UE::Audio::Insights
 		.AutoWidth()
 		[
 			CreateToolbarToggleControls()
+		]
+			
+		+ SHorizontalBox::Slot()
+		.AutoWidth()
+		.Padding(ToolbarButtonPaddingX, 0.0f)
+		[
+			SNew(SSeparator)
+			.Orientation(EOrientation::Orient_Vertical)
+		]
+
+		+ SHorizontalBox::Slot()
+		.AutoWidth()
+		[
+			CreateToolbarButtonActions()
 		];
 	}
 
@@ -224,6 +242,85 @@ namespace UE::Audio::Insights
 			]
 		];
 	}
+	
+	TSharedRef<SWidget> FSignalFlowDashboardViewFactory::CreateToolbarButtonActions()
+	{
+		using namespace FSignalFlowDashboardViewFactoryPrivate;
+
+		const FSlateStyle& Style = FSlateStyle::Get();
+
+		// Button actions
+		return SNew(SHorizontalBox)
+
+		+ SHorizontalBox::Slot()
+		.AutoWidth()
+		.Padding(ToolbarButtonPaddingX, ToolbarButtonPaddingY)
+		[
+			SNew(SButton)
+			.ButtonStyle(&FAppStyle::Get().GetWidgetStyle<FButtonStyle>(ButtonStyleName))
+			.ToolTipText(LOCTEXT("AudioDashboard_SignalFlow_CenterViewTooltip", "Focus the graph view on the selected node. (Ctrl + f)"))
+			.OnClicked_Lambda([this]()
+			{
+				return ScrollGraphViewToSelectedNode();
+			})
+			[
+				Style.CreateButtonContentWidget("AudioInsights.Icon.SignalFlow.CenterView", FText::GetEmpty(), ButtonTextStyleName)
+			]
+		];
+	}
+
+	FReply FSignalFlowDashboardViewFactory::ScrollGraphViewToSelectedNode()
+	{
+		if (GraphWidgetContents.IsValid())
+		{
+			const TSharedPtr<FSignalFlowDashboardEntry> FocusedEntry = SelectedItem.IsValid() ? SelectedItem : GetFilteredAudioDeviceEntry();
+			if (FocusedEntry.IsValid())
+			{
+				const TSharedPtr<SSignalFlowGraphNode> FocusedNodeWidget = GraphWidgetContents->FindNodeWidget(FocusedEntry->GetSignalFlowEntryKey());
+				if (FocusedNodeWidget.IsValid())
+				{
+					if (GraphHorizontalScrollBox.IsValid())
+					{
+						GraphHorizontalScrollBox->ScrollDescendantIntoView(FocusedNodeWidget, true, EDescendantScrollDestination::Center);
+					}
+
+					if (GraphVerticalScrollBox.IsValid())
+					{
+						GraphVerticalScrollBox->ScrollDescendantIntoView(FocusedNodeWidget, true, EDescendantScrollDestination::Center);
+					}
+
+					return FReply::Handled();
+				}
+			}
+		}
+
+		return FReply::Unhandled();
+	}
+
+	TSharedPtr<FSignalFlowDashboardEntry> FSignalFlowDashboardViewFactory::GetFilteredAudioDeviceEntry() const
+	{
+		const TSharedPtr<const FSignalFlowTraceProvider> Provider = FindProvider<const FSignalFlowTraceProvider>();
+		if (!Provider.IsValid())
+		{
+			return nullptr;
+		}
+
+		const FSignalFlowTraceProvider::FDeviceData* DeviceData = Provider->FindFilteredDeviceData();
+		if (DeviceData == nullptr || DeviceData->IsEmpty())
+		{
+			return nullptr;
+		}
+
+		const ::Audio::FDeviceId DeviceID = Provider->GetFilteredAudioDeviceID();
+		const TSharedPtr<FSignalFlowDashboardEntry>* DeviceEntry = DeviceData->Find({ DeviceID, ESignalFlowEntryType::AudioDevice, DeviceID });
+
+		if (DeviceEntry == nullptr || !DeviceEntry->IsValid())
+		{
+			return nullptr;
+		}
+
+		return *DeviceEntry;
+	}
 
 	TSharedRef<SWidget> FSignalFlowDashboardViewFactory::CreateSelectionPanel()
 	{
@@ -268,7 +365,7 @@ namespace UE::Audio::Insights
 
 	TSharedRef<SWidget> FSignalFlowDashboardViewFactory::CreateGraphSection()
 	{
-		return SNew(SScrollBox)
+		return SAssignNew(GraphVerticalScrollBox, SScrollBox)
 		.Orientation(EOrientation::Orient_Vertical)
 
 		+ SScrollBox::Slot()
@@ -276,7 +373,7 @@ namespace UE::Audio::Insights
 		.HAlign(HAlign_Fill)
 		.VAlign(VAlign_Fill)
 		[
-			SNew(SScrollBox)
+			SAssignNew(GraphHorizontalScrollBox, SScrollBox)
 			.Orientation(EOrientation::Orient_Horizontal)
 
 			+ SScrollBox::Slot()
@@ -447,6 +544,11 @@ namespace UE::Audio::Insights
 		}
 
 		OnRequestGraphRefresh();
+
+		if (SelectedItem.IsValid())
+		{
+			bAutoFocusAfterGraphRefresh = true;
+		}
 	}
 
 	void FSignalFlowDashboardViewFactory::ProcessEntries(EProcessReason Reason)
@@ -522,6 +624,12 @@ namespace UE::Audio::Insights
 		{
 			GraphWidgetContents->RefreshGraph();
 			bGraphRequiresRefresh = false;
+
+			if (bAutoFocusAfterGraphRefresh)
+			{
+				ScrollGraphViewToSelectedNode();
+				bAutoFocusAfterGraphRefresh = false;
+			}
 		}
 	}
 
@@ -615,18 +723,10 @@ namespace UE::Audio::Insights
 			return;
 		}
 
-		if (const FSignalFlowTraceProvider::FDeviceData* DeviceData = Provider->FindFilteredDeviceData())
+		const TSharedPtr<FSignalFlowDashboardEntry> DeviceEntry = GetFilteredAudioDeviceEntry();
+		if (DeviceEntry.IsValid())
 		{
-			if (!DeviceData->IsEmpty())
-			{
-				const ::Audio::FDeviceId DeviceID = Provider->GetFilteredAudioDeviceID();
-				const TSharedPtr<FSignalFlowDashboardEntry>* DeviceEntry = DeviceData->Find({ DeviceID, ESignalFlowEntryType::AudioDevice, DeviceID });
-
-				if (DeviceEntry && DeviceEntry->IsValid())
-				{
-					CreateFilteredNodeGraph(*Provider, *DeviceEntry);
-				}
-			}
+			CreateFilteredNodeGraph(*Provider, DeviceEntry);
 		}
 	}
 
diff --git a/Engine/Plugins/AudioInsights/Source/AudioInsights/Public/Views/SSignalFlowGraph.h b/Engine/Plugins/AudioInsights/Source/AudioInsights/Public/Views/SSignalFlowGraph.h
--- a/Engine/Plugins/AudioInsights/Source/AudioInsights/Public/Views/SSignalFlowGraph.h
+++ b/Engine/Plugins/AudioInsights/Source/AudioInsights/Public/Views/SSignalFlowGraph.h
@@ -132,6 +132,7 @@ namespace UE::Audio::Insights
 
 		UE_API void ToggleOrientation();
 		UE_API EOrientation GetGraphOrientation() const { return Orientation; }
+		UE_API TSharedPtr<SSignalFlowGraphNode> FindNodeWidget(const FSignalFlowEntryKey& NodeEntryKey) const;
 
 		UE_API virtual int32 OnPaint(const FPaintArgs& Args, const FGeometry& AllottedGeometry, const FSlateRect& MyCullingRect, FSlateWindowElementList& OutDrawElements, int32 LayerId, const FWidgetStyle& InWidgetStyle, bool bParentEnabled) const override;
 
diff --git a/Engine/Plugins/AudioInsights/Source/AudioInsights/Public/Views/SignalFlowDashboardViewFactory.h b/Engine/Plugins/AudioInsights/Source/AudioInsights/Public/Views/SignalFlowDashboardViewFactory.h
--- a/Engine/Plugins/AudioInsights/Source/AudioInsights/Public/Views/SignalFlowDashboardViewFactory.h
+++ b/Engine/Plugins/AudioInsights/Source/AudioInsights/Public/Views/SignalFlowDashboardViewFactory.h
@@ -5,6 +5,7 @@
 #include "Templates/SharedPointer.h"
 #include "Views/SSignalFlowGraph.h"
 #include "Views/TableDashboardViewFactory.h"
+#include "Widgets/Layout/SScrollBox.h"
 #include "Widgets/Views/SListView.h"
 
 #define UE_API AUDIOINSIGHTS_API
@@ -64,6 +65,10 @@ namespace UE::Audio::Insights
 
 		TSharedRef<SWidget> CreateToolbar();
 		TSharedRef<SWidget> CreateToolbarToggleControls();
+		TSharedRef<SWidget> CreateToolbarButtonActions();
+
+		FReply ScrollGraphViewToSelectedNode();
+		TSharedPtr<FSignalFlowDashboardEntry> GetFilteredAudioDeviceEntry() const;
 
 		TSharedRef<SWidget> CreateSelectionPanel();
 		TSharedRef<SWidget> CreateGraphSection();
@@ -119,10 +124,14 @@ namespace UE::Audio::Insights
 		TArray<TSharedPtr<ISignalFlowNode>> FilteredGraphNodes;
 		TSharedPtr<SSignalFlowGraph> GraphWidgetContents;
 
+		TSharedPtr<SScrollBox> GraphHorizontalScrollBox;
+		TSharedPtr<SScrollBox> GraphVerticalScrollBox;
+
 		TSharedPtr<FSignalFlowDashboardEntry> SelectedItem;
 		FText SearchBoxFilterText;
 
 		bool bGraphRequiresRefresh = false;
+		bool bAutoFocusAfterGraphRefresh = false;
 	};
 } // namespace UE::Audio::Insights
 
