[Audio Insights] Visual pass on signal flow nodes
* Pass on colors used within the graph
* Pass on node border style
* Improvements to control over spacing between nodes in both horizontal and vertical layouts

#Synced-CL 49648429

diff --git a/Engine/Plugins/AudioInsights/Source/AudioInsights/Private/Views/SSignalFlowGraph.cpp b/Engine/Plugins/AudioInsights/Source/AudioInsights/Private/Views/SSignalFlowGraph.cpp
--- a/Engine/Plugins/AudioInsights/Source/AudioInsights/Private/Views/SSignalFlowGraph.cpp
+++ b/Engine/Plugins/AudioInsights/Source/AudioInsights/Private/Views/SSignalFlowGraph.cpp
@@ -16,30 +16,68 @@
 
 namespace UE::Audio::Insights
 {
-	namespace FSignalFlowGraphPrivate
+	namespace SSignalFlowGraphStyle
 	{
-		float GetNodeMaxWidth()
+		constexpr float NodeMaxWidth = 200.0f;
+		constexpr float NodeMaxHeight = 50.0f;
+
+		constexpr float NodeRoundedBrushRadius = 12.0f;
+		constexpr float SendLabelRoundedBrushRadius = 8.0f;
+
+		constexpr float SplineRadius = 80.0f;
+		constexpr float SendLabelPropAlongLine = 0.25f; // How far alone the spline we want the send level to appear (range 0 -> 1)
+
+		constexpr float SelectedNodeBorderThickness = 3.0f;
+		constexpr float NodeBorderThickness = 2.0f;
+		constexpr float SendLabelBorderThickness = 1.0f;
+
+		constexpr float NodeBorderPadding = 12.0f;
+		constexpr float SendLevelTextMargin = 2.0f;
+
+		const FSlateRoundedBoxBrush& GetNodeRoundedBoxBrush()
 		{
-			return 160.0f;
+			static const FSlateRoundedBoxBrush RoundedBoxBrush(FLinearColor::White, NodeRoundedBrushRadius);
+			return RoundedBoxBrush;
 		}
 
-		float GetNodePaddingX(const EOrientation Orientation)
+		const FSlateRoundedBoxBrush& GetSendLabelRoundedBoxBrush()
 		{
-			return Orientation == EOrientation::Orient_Vertical ? 24.0f : 48.0f;
+			static const FSlateRoundedBoxBrush RoundedBoxBrush(FLinearColor::White, SendLabelRoundedBrushRadius);
+			return RoundedBoxBrush;
 		}
 
-		constexpr float NodePaddingY = 24.0f;
-		constexpr float NodeMaxHeight = 100.0f;
+		float GetNodeGraphSlotWidth(const EOrientation GraphOrientation)
+		{
+			const float XPaddingBetweenNodes = GraphOrientation == EOrientation::Orient_Horizontal ? 48.0f : 8.0f;
+			return NodeMaxWidth + XPaddingBetweenNodes;
+		}
 
-		constexpr float RoundedBrushRadius = 8.0f;
+		float GetNodeGraphSlotHeight(const EOrientation GraphOrientation)
+		{
+			const float YPaddingBetweenNodes = GraphOrientation == EOrientation::Orient_Horizontal ? 8.0f : 48.0f;
+			return NodeMaxHeight + YPaddingBetweenNodes;
+		}
 
-		constexpr float SplineRadius = 80.0f;
-		constexpr float SendLabelPropAlongLine = 0.25f; // How far alone the spline we want the send level to appear (range 0 -> 1)
+		FSlateColor GetNodeAccentColor(const ESignalFlowEntryType NodeType)
+		{
+			switch (NodeType)
+			{
+				case ESignalFlowEntryType::SoundSource:
+					return FStyleColors::AccentBlue;
+				case ESignalFlowEntryType::AudioBus:
+					return FStyleColors::AccentYellow;
+				case ESignalFlowEntryType::Submix:
+					return FStyleColors::AccentGreen;
+				case ESignalFlowEntryType::OwnerObject:
+				case ESignalFlowEntryType::AudioDevice:
+				default:
+					return FStyleColors::Foreground;
+			}
+		}
 
-		const FSlateRoundedBoxBrush& GetRoundedBoxBrush()
+		FSlateColor GetConnectionColor()
 		{
-			static const FSlateRoundedBoxBrush RoundedBoxBrush(FStyleColors::AccentGray.GetSpecifiedColor(), RoundedBrushRadius);
-			return RoundedBoxBrush;
+			return FStyleColors::White25;
 		}
 
 		FVector2d GetSplineStartDirection(const EOrientation Orientation)
@@ -55,7 +93,7 @@ namespace UE::Audio::Insights
 
 	void SSignalFlowGraphNode::Construct(const SSignalFlowGraphNode::FArguments& InArgs)
 	{
-		using namespace FSignalFlowGraphPrivate;
+		using namespace SSignalFlowGraphStyle;
 
 		Entry = InArgs._Entry;
 		PositionInGraph = InArgs._PositionInGraph;
@@ -71,52 +109,54 @@ namespace UE::Audio::Insights
 		}
 
 		NodeEntryKey = Entry->GetSignalFlowEntryKey();
+		const FSlateColor NodeAccentColor = GetNodeAccentColor(NodeEntryKey.EntryType);
 
 		ChildSlot
 		[
 			SNew(SVerticalBox)
 
-			+ SVerticalBox::Slot()
-			.FillContentHeight(1.0f)
-			[
-				SNullWidget::NullWidget
-			]
-
 			+ SVerticalBox::Slot()
 			.AutoHeight()
 			.HAlign(HAlign_Center)
 			[
 				SNew(SBorder)
-				.ColorAndOpacity_Lambda([this]()
+				.BorderImage(&GetNodeRoundedBoxBrush())
+				.HAlign(HAlign_Fill)
+				.VAlign(VAlign_Fill)
+				.BorderBackgroundColor_Lambda([this]()
 				{
-					return IsSelected.IsBound() && IsSelected.Get() ? FLinearColor::Green : FLinearColor::White;
+					return IsSelected.IsBound() && IsSelected.Get() ? FStyleColors::AccentOrange.GetSpecifiedColor() : FLinearColor::Transparent;
 				})
-				.Content()
+				.Padding(SelectedNodeBorderThickness)
 				[
-					SNew(SOverlay)
-
-					+ SOverlay::Slot()
-					[
-						SNew(SColorBlock)
-						.Color(FStyleColors::Panel.GetSpecifiedColor())
-					]
-
-					+ SOverlay::Slot()
+					SNew(SBorder)
+					.BorderImage(&GetNodeRoundedBoxBrush())
+					.HAlign(HAlign_Fill)
+					.VAlign(VAlign_Fill)
+					.BorderBackgroundColor(NodeAccentColor)
+					.Padding(NodeBorderThickness)
+					.Content()
 					[
-						SNew(SVerticalBox)
-
-						+ SVerticalBox::Slot()
-						.Padding(4.0f)
-						.AutoHeight()
+						SNew(SBorder)
+						.BorderImage(&GetNodeRoundedBoxBrush())
+						.HAlign(HAlign_Fill)
+						.VAlign(VAlign_Fill)
+						.BorderBackgroundColor_Lambda([this]()
+						{
+							return (IsSelected.IsBound() && IsSelected.Get()) ? FStyleColors::Header : FStyleColors::Panel;
+						})
+						.Padding(NodeBorderPadding)
+						.Content()
 						[
 							SNew(SHorizontalBox)
 
 							+ SHorizontalBox::Slot()
 							.AutoWidth()
-							.Padding(4.0f, 0.0f)
+							.Padding(0.0f, 0.0f, 4.0f, 0.0f)
 							[
 								SNew(SImage)
 								.Image(FSlateStyle::Get().CreateIcon(InArgs._Entry->IconName).GetIcon())
+								.ColorAndOpacity(NodeAccentColor)
 							]
 
 							+ SHorizontalBox::Slot()
@@ -132,12 +172,6 @@ namespace UE::Audio::Insights
 					]
 				]
 			]
-
-			+ SVerticalBox::Slot()
-			.FillContentHeight(1.0f)
-			[
-				SNullWidget::NullWidget
-			]
 		];
 	}
 
@@ -160,42 +194,46 @@ namespace UE::Audio::Insights
 
 	FVector2D SSignalFlowGraphNode::ComputeDesiredSize(float InLayoutScaleMultiplier) const
 	{
-		using namespace FSignalFlowGraphPrivate;
+		using namespace SSignalFlowGraphStyle;
 
 		const FVector2D Size = SCompoundWidget::ComputeDesiredSize(InLayoutScaleMultiplier);
-		const FVector2D ClampedSize = { FMath::Min(Size.X, GetNodeMaxWidth() - GetNodePaddingX(GraphOrientation.Get())), Size.Y};
+		const FVector2D ClampedSize = { FMath::Min(Size.X, NodeMaxWidth), FMath::Min(Size.Y, NodeMaxHeight) };
 		return ClampedSize;
 	}
 
 	FVector2D SSignalFlowGraphNode::GetLocalInputPos() const
 	{
+		using namespace SSignalFlowGraphStyle;
+
 		const FVector2D Size = GetDesiredSize();
 
 		if (GraphOrientation.IsBound() && GraphOrientation.Get() == EOrientation::Orient_Horizontal)
 		{
 			// Input pos = left center of node
-			return { CachedPositionInGraph.X, CachedPositionInGraph.Y + (Size.Y / 2.0f) };
+			return { CachedPositionInGraph.X + SelectedNodeBorderThickness, CachedPositionInGraph.Y + (Size.Y / 2.0f) };
 		}
 		else
 		{
 			// Input pos = top center of node
-			return { CachedPositionInGraph.X + (Size.X / 2.0f), CachedPositionInGraph.Y };
+			return { CachedPositionInGraph.X + (Size.X / 2.0f), CachedPositionInGraph.Y + SelectedNodeBorderThickness };
 		}
 	}
 	
 	FVector2D SSignalFlowGraphNode::GetLocalOutputPos() const
 	{
+		using namespace SSignalFlowGraphStyle;
+
 		const FVector2D Size = GetDesiredSize();
 
 		if (GraphOrientation.IsBound() && GraphOrientation.Get() == EOrientation::Orient_Horizontal)
 		{
 			// Out pos = right center of node
-			return { CachedPositionInGraph.X + Size.X, CachedPositionInGraph.Y + (Size.Y / 2.0f) };
+			return { CachedPositionInGraph.X + Size.X - SelectedNodeBorderThickness, CachedPositionInGraph.Y + (Size.Y / 2.0f) };
 		}
 		else
 		{
 			// Output pos = bottom center of node
-			return { CachedPositionInGraph.X + (Size.X / 2.0f), CachedPositionInGraph.Y + Size.Y };
+			return { CachedPositionInGraph.X + (Size.X / 2.0f), CachedPositionInGraph.Y + Size.Y - SelectedNodeBorderThickness };
 		}
 	}
 
@@ -212,7 +250,7 @@ namespace UE::Audio::Insights
 
 	void SSignalFlowGraphSendLabel::Construct(const SSignalFlowGraphSendLabel::FArguments& InArgs)
 	{
-		using namespace FSignalFlowGraphPrivate;
+		using namespace SSignalFlowGraphStyle;
 
 		SendReceiverEntryKey = InArgs._SendReceiverEntryKey;
 		SenderPositionInGraph = InArgs._SenderPositionInGraph;
@@ -234,15 +272,15 @@ namespace UE::Audio::Insights
 			.VAlign(VAlign_Center)
 			[
 				SNew(SBorder)
-				.BorderImage(&GetRoundedBoxBrush())
+				.BorderImage(&GetSendLabelRoundedBoxBrush())
 				.HAlign(HAlign_Fill)
 				.VAlign(VAlign_Fill)
-				.BorderBackgroundColor(FStyleColors::White.GetSpecifiedColor())
-				.Padding(1.0f)
+				.BorderBackgroundColor(FStyleColors::Foreground.GetSpecifiedColor())
+				.Padding(SendLabelBorderThickness)
 				.Content()
 				[
 					SNew(SBorder)
-					.BorderImage(&GetRoundedBoxBrush())
+					.BorderImage(&GetSendLabelRoundedBoxBrush())
 					.HAlign(HAlign_Fill)
 					.VAlign(VAlign_Fill)
 					.BorderBackgroundColor(FStyleColors::Panel.GetSpecifiedColor())
@@ -265,7 +303,7 @@ namespace UE::Audio::Insights
 							return FText::AsNumber(SendLevelOpt.GetValue(), FSlateStyle::Get().GetSendLevelFormat());
 						})
 						.Justification(ETextJustify::Center)
-						.Margin(2.0f)
+						.Margin(SendLevelTextMargin)
 					]
 				]
 			]
@@ -274,7 +312,7 @@ namespace UE::Audio::Insights
 
 	FVector2D SSignalFlowGraphSendLabel::ComputeDesiredSize(float InLayoutScaleMultiplier) const
 	{
-		using namespace FSignalFlowGraphPrivate;
+		using namespace SSignalFlowGraphStyle;
 
 		const FVector2D Size = SCompoundWidget::ComputeDesiredSize(InLayoutScaleMultiplier);
 		return Size;
@@ -282,7 +320,7 @@ namespace UE::Audio::Insights
 
 	FVector2D SSignalFlowGraphSendLabel::UpdatePositionInGraph()
 	{
-		using namespace FSignalFlowGraphPrivate;
+		using namespace SSignalFlowGraphStyle;
 
 		const FVector2D SenderNodePos = SenderPositionInGraph.Get();
 		const FVector2D ReceiverNodePos = ReceiverPositionInGraph.Get();
@@ -321,7 +359,7 @@ namespace UE::Audio::Insights
 
 	void SSignalFlowGraph::RefreshGraph()
 	{
-		using namespace FSignalFlowGraphPrivate;
+		using namespace SSignalFlowGraphStyle;
 
 		if (ListSource == nullptr)
 		{
@@ -367,7 +405,7 @@ namespace UE::Audio::Insights
 
 	int32 SSignalFlowGraph::OnPaintConnections(const FGeometry& AllottedGeometry, FSlateWindowElementList& OutDrawElements, int32 LayerId) const
 	{
-		using namespace FSignalFlowGraphPrivate;
+		using namespace SSignalFlowGraphStyle;
 
 		if (!ListSource)
 		{
@@ -416,7 +454,7 @@ namespace UE::Audio::Insights
 						GetSplineEndDirection(Orientation),
 						2.0f,
 						ESlateDrawEffect::NoPixelSnapping,
-						FStyleColors::AccentGray.GetSpecifiedColor()
+						GetConnectionColor().GetSpecifiedColor()
 					);
 				}
 			}
@@ -427,10 +465,14 @@ namespace UE::Audio::Insights
 
 	FVector2D SSignalFlowGraph::ComputeDesiredSize(float InLayoutScaleMultiplier) const
 	{
-		using namespace FSignalFlowGraphPrivate;
+		using namespace SSignalFlowGraphStyle;
 
 		float GraphWidth = 0.0f;
 		float GraphHeight = 0.0f;
+
+		const float NodeGraphSlotWidth = GetNodeGraphSlotWidth(GetGraphOrientation());
+		const float NodeGraphSlotHeight = GetNodeGraphSlotHeight(GetGraphOrientation());
+
 		for (const auto& [Key, NodeWidget] : NodeWidgets)
 		{
 			if (!NodeWidget.IsValid())
@@ -440,13 +482,13 @@ namespace UE::Audio::Insights
 
 			if (Orientation == Orient_Vertical)
 			{
-				GraphWidth = FMath::Max(GraphWidth, (NodeWidget->GetXPos() + 1) * GetNodeMaxWidth());
-				GraphHeight = FMath::Max(GraphHeight, (NodeWidget->GetYPos() + 1) * NodeMaxHeight);
+				GraphWidth = FMath::Max(GraphWidth, (NodeWidget->GetXPos() + 1) * NodeGraphSlotWidth);
+				GraphHeight = FMath::Max(GraphHeight, (NodeWidget->GetYPos() + 1) * NodeGraphSlotHeight);
 			}
 			else
 			{
-				GraphWidth = FMath::Max(GraphWidth, (NodeWidget->GetYPos() + 1) * GetNodeMaxWidth());
-				GraphHeight = FMath::Max(GraphHeight, (NodeWidget->GetXPos() + 1) * NodeMaxHeight);
+				GraphWidth = FMath::Max(GraphWidth, (NodeWidget->GetYPos() + 1) * NodeGraphSlotWidth);
+				GraphHeight = FMath::Max(GraphHeight, (NodeWidget->GetXPos() + 1) * NodeGraphSlotHeight);
 			}
 		}
 
@@ -557,7 +599,7 @@ namespace UE::Audio::Insights
 
 	void SSignalFlowGraph::CreateGraphNode(const TSharedPtr<ISignalFlowNode>& Node, const float MinXPos, const float MaxXPos, const int32 DepthTracker, TArray<FSignalFlowEntryKey>& OutRemainingNodeKeys)
 	{
-		using namespace FSignalFlowGraphPrivate;
+		using namespace SSignalFlowGraphStyle;
 
 		if (!Node.IsValid())
 		{
@@ -639,21 +681,24 @@ namespace UE::Audio::Insights
 				if (Node && Node->IsValid())
 				{
 					const FVector2D NodeSize = (*Node)->GetDesiredSize();
-					const float XOffset = (GetNodeMaxWidth() - NodeSize.X) / 2.0f;
-					const float YOffset = (NodeMaxHeight - NodeSize.Y) / 2.0f;
+					const float NodeGraphSlotWidth = GetNodeGraphSlotWidth(GetGraphOrientation());
+					const float NodeGraphSlotHeight = GetNodeGraphSlotHeight(GetGraphOrientation());
+
+					const float XOffset = (NodeGraphSlotWidth - NodeSize.X) / 2.0f;
+					const float YOffset = (NodeGraphSlotHeight - NodeSize.Y) / 2.0f;
 
 					if (Orientation == EOrientation::Orient_Vertical)
 					{
-						const float XPos = ((*Node)->GetXPos() * GetNodeMaxWidth()) + XOffset;
-						const float YPos = (*Node)->GetYPos() * NodeMaxHeight + YOffset;
+						const float XPos = ((*Node)->GetXPos() * NodeGraphSlotWidth) + XOffset;
+						const float YPos = (*Node)->GetYPos() * NodeGraphSlotHeight + YOffset;
 
 						const FVector2D NodePosition = { XPos, YPos };
 						return NodePosition;
 					}
 					else
 					{
-						const float XPos = ((*Node)->GetYPos() * GetNodeMaxWidth()) + XOffset;
-						const float YPos = (*Node)->GetXPos() * NodeMaxHeight + YOffset;
+						const float XPos = ((*Node)->GetYPos() * NodeGraphSlotWidth) + XOffset;
+						const float YPos = (*Node)->GetXPos() * NodeGraphSlotHeight + YOffset;
 
 						const FVector2D NodePosition = { XPos, YPos };
 						return NodePosition;
@@ -709,7 +754,7 @@ namespace UE::Audio::Insights
 
 	void SSignalFlowGraph::CreateSendLabelWidget(const FSignalFlowEntryKey& SenderEntryKey, const FSignalFlowEntryKey& ReceiverEntryKey)
 	{
-		using namespace FSignalFlowGraphPrivate;
+		using namespace SSignalFlowGraphStyle;
 
 		if (FindSendLabel(SenderEntryKey, ReceiverEntryKey).IsValid())
 		{
diff --git a/Engine/Plugins/AudioInsights/Source/AudioInsights/Public/Views/SSignalFlowGraph.h b/Engine/Plugins/AudioInsights/Source/AudioInsights/Public/Views/SSignalFlowGraph.h
--- a/Engine/Plugins/AudioInsights/Source/AudioInsights/Public/Views/SSignalFlowGraph.h
+++ b/Engine/Plugins/AudioInsights/Source/AudioInsights/Public/Views/SSignalFlowGraph.h
@@ -15,6 +15,11 @@ namespace UE::Audio::Insights
 
 	DECLARE_DELEGATE_OneParam(FOnSignalFlowNodeSelected, TSharedPtr<FSignalFlowDashboardEntry>);
 
+	namespace SSignalFlowGraphStyle
+	{
+		UE_API FSlateColor GetNodeAccentColor(const ESignalFlowEntryType NodeType);
+	}
+
 	class SSignalFlowGraphNode : public SCompoundWidget
 	{
 	public:
