[Audio Insights] Visual pass on signal flow left hand list menus
* Pass on colors and icons
* Pass on font style

#Synced-CL 49648429

diff --git a/Engine/Plugins/AudioInsights/Source/AudioInsights/Private/Views/SignalFlowDashboardViewFactory.cpp b/Engine/Plugins/AudioInsights/Source/AudioInsights/Private/Views/SignalFlowDashboardViewFactory.cpp
--- a/Engine/Plugins/AudioInsights/Source/AudioInsights/Private/Views/SignalFlowDashboardViewFactory.cpp
+++ b/Engine/Plugins/AudioInsights/Source/AudioInsights/Private/Views/SignalFlowDashboardViewFactory.cpp
@@ -163,10 +163,26 @@ namespace UE::Audio::Insights
 		(
 			STableRow<TSharedPtr<FSignalFlowDashboardEntry>>::FArguments()
 			.Style(InArgs._Style)
+			.Padding(4.0f)
 			.Content()
 			[
-				SNew(STextBlock)
+				SNew(SHorizontalBox)
+
+				+ SHorizontalBox::Slot()
+				.AutoWidth()
+				.Padding(0.0f, 0.0f, 4.0f, 0.0f)
+				[
+					SNew(SImage)
+					.Image(FSlateStyle::Get().CreateIcon(RowEntry->IconName).GetIcon())
+					.ColorAndOpacity(SSignalFlowGraphStyle::GetNodeAccentColor(RowEntry->EntryType))
+				]
+
+				+ SHorizontalBox::Slot()
+				.AutoWidth()
+				[
+					SNew(STextBlock)
 					.Text(RowEntry->GetDisplayName())
+				]
 			]
 			, InOwnerTable
 		);
@@ -458,9 +474,16 @@ namespace UE::Audio::Insights
 
 		static const FLazyName HeadingColumnID("Heading");
 		const SHeaderRow::FColumn::FArguments NameColumnArgs = SHeaderRow::Column(HeadingColumnID)
-																		   .DefaultLabel(HeadingText)
-																		   .HAlignCell(HAlign_Left)
-																		   .FillWidth(1.0f);
+			.DefaultLabel(HeadingText)
+			.HAlignCell(HAlign_Left)
+			.FillWidth(1.0f)
+			.HeaderContentPadding(FMargin(4.0f, 4.0f))
+			.HeaderContent()
+			[
+				SNew(STextBlock)
+				.Text(HeadingText)
+				.Font(FCoreStyle::Get().GetFontStyle("BoldFont"))
+			];
 
 		HeaderRowWidget->AddColumn(NameColumnArgs);
 
