[AudioInsights] - AudioModulation: Support modulation icons

Partial fix for:
https://jira.it.epicgames.com/browse/UE-356383

Passes bughawk

Preflight:
https://horde.devtools.epicgames.com/job/695ffde24d3a619274ceded5#rail_timeline

#Synced-CL 49759006

diff --git a/Engine/Plugins/AudioInsights/Source/AudioInsights/Public/AudioInsightsStyle.h b/Engine/Plugins/AudioInsights/Source/AudioInsights/Public/AudioInsightsStyle.h
--- a/Engine/Plugins/AudioInsights/Source/AudioInsights/Public/AudioInsightsStyle.h
+++ b/Engine/Plugins/AudioInsights/Source/AudioInsights/Public/AudioInsightsStyle.h
@@ -254,6 +254,14 @@ namespace UE::Audio::Insights
 			Set("AudioInsights.Icon.SignalFlow.NodeDetails",		  new IMAGE_BRUSH_SVG(TEXT("Icons/SignalFlow/node_details"),		 Icon20));
 			Set("AudioInsights.Icon.SignalFlow.CenterView",			  new IMAGE_BRUSH_SVG(TEXT("Icons/SignalFlow/center_view"),			 Icon20));
 			Set("AudioInsights.Icon.SignalFlow.Settings",			  new IMAGE_BRUSH_SVG(TEXT("Icons/SignalFlow/settings"),			 Icon20));
+
+			/* Modulation */
+			// Icons
+			Set("AudioInsights.Icon.Modulation.ControlBus", new IMAGE_BRUSH(TEXT("Icons/Modulation/SoundControlBus"), Icon16));
+			Set("AudioInsights.Icon.Modulation.ControlBusMix", new IMAGE_BRUSH(TEXT("Icons/Modulation/SoundControlBusMix"), Icon16));
+			Set("AudioInsights.Icon.Modulation.Generator", new IMAGE_BRUSH(TEXT("Icons/Modulation/SoundModulationGenerator"), Icon16));
+			Set("AudioInsights.Icon.Modulation.Parameter", new IMAGE_BRUSH(TEXT("Icons/Modulation/SoundModulationParameter"), Icon16));
+			Set("AudioInsights.Icon.Modulation.ParameterPatch", new IMAGE_BRUSH(TEXT("Icons/Modulation/SoundModulationPatch"), Icon16));
 			
 			// Category colors
 			Set("SoundDashboard.MetaSoundColor", FLinearColor(0.008f, 0.76f, 0.078f));
