[Audio Insights] FEditorDashboardFactory, trace audio channels checkbox: added missing SaveConfig call after modifying bOnlyTraceAudioChannels var from AudioInsightsEditorSettings. The value gets saved properly in its corresponding settings .ini file.

#Synced-CL 49838930

diff --git a/Engine/Plugins/AudioInsights/Source/AudioInsightsEditor/Private/AudioInsightsEditorDashboardFactory.cpp b/Engine/Plugins/AudioInsights/Source/AudioInsightsEditor/Private/AudioInsightsEditorDashboardFactory.cpp
--- a/Engine/Plugins/AudioInsights/Source/AudioInsightsEditor/Private/AudioInsightsEditorDashboardFactory.cpp
+++ b/Engine/Plugins/AudioInsights/Source/AudioInsightsEditor/Private/AudioInsightsEditorDashboardFactory.cpp
@@ -744,6 +744,7 @@ namespace UE::Audio::Insights
 						if (AudioInsightsEditorSettings)
 						{
 							AudioInsightsEditorSettings->bOnlyTraceAudioChannels = NewState == ECheckBoxState::Checked;
+							AudioInsightsEditorSettings->SaveConfig();
 
 							IAudioInsightsTraceModule& TraceModule = FAudioInsightsEditorModule::GetChecked().GetTraceModule();
 							TraceModule.OnOnlyTraceAudioChannelsStateChanged(AudioInsightsEditorSettings->bOnlyTraceAudioChannels);
