[AudioInsights] - AudioInsightsCacheManager: Support FindClosestMessageByPredicate

Swarm review: 49887278

#Synced-CL 49919523

diff --git a/Engine/Plugins/AudioInsights/Source/AudioInsights/Public/Cache/AudioCachedMessageChunk.h b/Engine/Plugins/AudioInsights/Source/AudioInsights/Public/Cache/AudioCachedMessageChunk.h
--- a/Engine/Plugins/AudioInsights/Source/AudioInsights/Public/Cache/AudioCachedMessageChunk.h
+++ b/Engine/Plugins/AudioInsights/Source/AudioInsights/Public/Cache/AudioCachedMessageChunk.h
@@ -35,11 +35,12 @@ namespace UE::Audio::Insights
 		UE_API double GetChunkTimeRangeStart() const;
 		UE_API double GetChunkTimeRangeEnd() const;
 
-		// Finds the closest message to a specific timestamp
+		// Finds the closest message to a specific timestamp and where predicate returns true.
 		// Will iterate backwards from param Timestamp until the first message of type MessageID is found
-		// Param ID can be overriden to filter the results further
-		template<TIsCacheableMessage T>
-		T* FindClosestMessage(const FName& MessageID, const double Timestamp, const uint64 ID = INDEX_NONE) const
+		// @param Predicate: A callable that returns a BooleanTestable value.
+		//					 E.g. [](const T& InMessage){ return InMessage.GetID() == 10; }
+		template<TIsCacheableMessage T, typename Predicate>
+		const T* FindClosestMessageByPredicate(const FName& MessageID, const double Timestamp, Predicate Pred) const
 		{
 			const TSharedPtr<TCachePagedArray>* MessageCache = Messages.Find(MessageID);
 			if (MessageCache == nullptr || !MessageCache->IsValid() || (*MessageCache)->Num() == 0u)
@@ -54,18 +55,34 @@ namespace UE::Audio::Insights
 				return Msg.IsValid() ? Msg->Timestamp : INVALID_TIMESTAMP;
 			});
 
-			// Iterate backwards from TimeMarker until we find the matching ID
+			// Iterate backwards from TimeMarker until we find a valid message and where the predicate returns true
 			for (auto MessageItr = MessagesArray.GetIteratorFromItem(ClosestMessageToTimeStampIndex); MessageItr; --MessageItr)
 			{
-				if (ID == INDEX_NONE || (*MessageItr)->GetID() == ID)
+				const T* Message = static_cast<const T*>((*MessageItr).Get());
+
+				if (Message && Pred(*Message))
 				{
-					return static_cast<T*>((*MessageItr).Get());
+					return Message;
 				}
 			}
 
 			return nullptr;
 		}
 
+		// Finds the closest message to a specific timestamp
+		// Will iterate backwards from param Timestamp until the first message of type MessageID is found
+		// Param ID can be overriden to filter the results further
+		template<TIsCacheableMessage T>
+		T* FindClosestMessage(const FName& MessageID, const double Timestamp, const uint64 ID = INDEX_NONE) const
+		{
+			return FindClosestMessageByPredicate<T>(MessageID,
+				Timestamp,
+				[ID](const T& InMessage)
+				{
+					return (ID == INDEX_NONE) || (InMessage.GetID() == ID);
+				});
+		}
+
 	private:
 		static constexpr double INVALID_TIMESTAMP = -1.0;
 
diff --git a/Engine/Plugins/AudioInsights/Source/AudioInsights/Public/Cache/AudioInsightsCacheManager.h b/Engine/Plugins/AudioInsights/Source/AudioInsights/Public/Cache/AudioInsightsCacheManager.h
--- a/Engine/Plugins/AudioInsights/Source/AudioInsights/Public/Cache/AudioInsightsCacheManager.h
+++ b/Engine/Plugins/AudioInsights/Source/AudioInsights/Public/Cache/AudioInsightsCacheManager.h
@@ -96,14 +96,14 @@ namespace UE::Audio::Insights
 		// e.g. NumChunksFromStart = 0u will return the chunk with the earliest messages in the cache
 		UE_API const FAudioCachedMessageChunk* const GetChunk(const uint32 NumChunksFromStart) const;
 
-		// Find the closest message to a timestamp in the cache for a specific message type
+		// Find the closest message to a timestamp in the cache for a specific message type and where predicate returns true.
 		// Will look backwards from the chosen timestamp until a message is found
-		// @param ID: When overloaded, the search algorithm will only return messages that have a matching ID.
-		//			  If left at INDEX_NONE, the search algorithm will return the closest message irrespective of ID
+		// @param Predicate: A callable that returns a BooleanTestable value.
+		//					 E.g. [](const T& InMessage){ return InMessage.GetID() == 10; }
 		// @param BackOutTimestamp: When iterating backwards, if a chunk is detected with an end timestamp less than
 		//							the backout timestamp, the search is canceled and this returns nullptr
-		template<TIsCacheableMessage T>
-		const T* FindClosestMessage(const FName& MessageID, const double Timestamp, const uint64 ID = INDEX_NONE, const double BackOutTimestamp = -1.0) const
+		template<TIsCacheableMessage T, typename Predicate>
+		const T* FindClosestMessageByPredicate(const FName& MessageID, const double Timestamp, Predicate Pred, const double BackOutTimestamp = -1.0) const
 		{
 			uint32 ReadIndex = INDEX_NONE;
 			for (uint32 Index = 0u; Index < Cache.Capacity(); ++Index)
@@ -132,7 +132,7 @@ namespace UE::Audio::Insights
 					break;
 				}
 
-				if (const T* FoundMessage = Chunk.FindClosestMessage<T>(MessageID, Timestamp, ID))
+				if (const T* FoundMessage = Chunk.FindClosestMessageByPredicate<T>(MessageID, Timestamp, Pred))
 				{
 					return FoundMessage;
 				}
@@ -155,6 +155,24 @@ namespace UE::Audio::Insights
 			return nullptr;
 		}
 
+		// Find the closest message to a timestamp in the cache for a specific message type
+		// Will look backwards from the chosen timestamp until a message is found
+		// @param ID: When overloaded, the search algorithm will only return messages that have a matching ID.
+		//			  If left at INDEX_NONE, the search algorithm will return the closest message irrespective of ID
+		// @param BackOutTimestamp: When iterating backwards, if a chunk is detected with an end timestamp less than
+		//							the backout timestamp, the search is canceled and this returns nullptr
+		template<TIsCacheableMessage T>
+		const T* FindClosestMessage(const FName& MessageID, const double Timestamp, const uint64 ID = INDEX_NONE, const double BackOutTimestamp = -1.0) const
+		{
+			return FindClosestMessageByPredicate<T>(MessageID,
+				Timestamp,
+				[ID](const T& InMessage)
+				{
+					return (ID == INDEX_NONE) || (InMessage.GetID() == ID);
+				},
+				BackOutTimestamp);
+		}
+
 		// Iterates the cache over a chosen timestamp range for a specific message ID
 		// Will stop iterating if the end of the cache is reached
 		// @param ID: When overloaded, the search algorithm will only return messages that have a matching ID.
