[Audio Insights] Improve positioning of nodes in the signal flow graph
* Don't reprocess branches inside AssignNodeOrderRecursive that have already been processed (optimization)
* Tweaks to input sorting - move nodes that have no inputs to one side. Use the minimum timestamp of the full branch (instead of the immediate inputs) to sort by parent timestamp. Only count transient entries (i.e. sound sources and owner objects) when finding the minimum timestamp of a branch.
* Nodes that were left unassigned by the node ordering alg have their order set to int MAX. This happens when a node has no connected output, and moves these nodes to one side of the graph.
* Fixed incorrect for loop inside CalcXPositionsForRows - we need to iterate backwards over previously positioned rows to assign new XPositions
* Fixed some nodes drawing over each other inside MakeSpaceForNodeInRow - we need to shift any node that is to the left or at the position of the new node.

#Synced-CL 49947352

diff --git a/Engine/Plugins/AudioInsights/Source/AudioInsights/Private/Views/SignalFlowDashboardViewFactory.cpp b/Engine/Plugins/AudioInsights/Source/AudioInsights/Private/Views/SignalFlowDashboardViewFactory.cpp
--- a/Engine/Plugins/AudioInsights/Source/AudioInsights/Private/Views/SignalFlowDashboardViewFactory.cpp
+++ b/Engine/Plugins/AudioInsights/Source/AudioInsights/Private/Views/SignalFlowDashboardViewFactory.cpp
@@ -1006,13 +1006,15 @@ namespace UE::Audio::Insights
 		// This helps keep nodes that are linked to the same input together in the graph
 		// We sort the inputs before proceeding forward to keep a good order in the graph
 		
-		// Only assign a new processing order if the node has not been reached yet
-		if (!OutNode->NodeOrderIsValid())
+		// Do not reprocess nodes that have already been processed
+		if (OutNode->NodeOrderIsValid())
 		{
-			OutNode->SetNodeOrderID(OutNodeProcessingOrder);
-			OutNodeProcessingOrder++;
+			return;
 		}
 
+		OutNode->SetNodeOrderID(OutNodeProcessingOrder);
+		OutNodeProcessingOrder++;
+
 		SortNodeInputs(OutNode);
 
 		for (const FSignalFlowEntryKey InputKey : OutNode->FilteredInputs)
@@ -1049,30 +1051,33 @@ namespace UE::Audio::Insights
 					return FirstNodeTreeDepth < SecondNodeTreeDepth;
 				}
 
-				// Next, try to keep nodes that have smaller gaps in tree depth on one side
-				TreeDepthPair FirstMinTreeDepth = { ESignalFlowEntryType::InvalidMax, TNumericLimits<int32>::Max() };
-				TreeDepthPair SecondMinTreeDepth = { ESignalFlowEntryType::InvalidMax, TNumericLimits<int32>::Max() };
-
-				double FirstInputTimestamp = TNumericLimits<double>::Max();
-				double SecondInputTimestamp = TNumericLimits<double>::Max();
+				// Move nodes that have no inputs to one side of the graph
+				const int32 FirstNumInputs = (*FirstNode)->FilteredInputs.Num();
+				const int32 SecondNumInputs = (*SecondNode)->FilteredInputs.Num();
 
-				GetMinTreeDepthAndTimestamps((*FirstNode)->FilteredInputs, FirstMinTreeDepth, FirstInputTimestamp);
-				GetMinTreeDepthAndTimestamps((*SecondNode)->FilteredInputs, SecondMinTreeDepth, SecondInputTimestamp);
+				const bool bNumInputsMatch = FirstNumInputs == SecondNumInputs;
+				const bool bEitherNodeHasZeroInputs = (FirstNumInputs == 0 || SecondNumInputs == 0);
 
-				// Sort by the minimum tree depth of the parent nodes
-				if (FirstMinTreeDepth != SecondMinTreeDepth)
+				if (bEitherNodeHasZeroInputs && !bNumInputsMatch)
 				{
-					return FirstMinTreeDepth < SecondMinTreeDepth;
+					return FirstNumInputs > SecondNumInputs;
 				}
 
+				// Next, try to order nodes by the minimum timestamp of their parent sound source/owner object nodes (if any exist)
+				double FirstInputTimestamp = TNumericLimits<double>::Max();
+				double SecondInputTimestamp = TNumericLimits<double>::Max();
+
+				GetMinTimestampRecursive((*FirstNode)->FilteredInputs, FirstInputTimestamp);
+				GetMinTimestampRecursive((*SecondNode)->FilteredInputs, SecondInputTimestamp);
+
 				// Sort by timestamp of the parent nodes
-				if (FirstInputTimestamp != SecondInputTimestamp)
+				if (!FMath::IsNearlyEqual(FirstInputTimestamp, SecondInputTimestamp))
 				{
 					return FirstInputTimestamp < SecondInputTimestamp;
 				}
 
 				// Try timestamp of this node
-				if ((*FirstNode)->GetTimestamp() != (*SecondNode)->GetTimestamp())
+				if (!FMath::IsNearlyEqual((*FirstNode)->GetTimestamp(), (*SecondNode)->GetTimestamp()))
 				{
 					return (*FirstNode)->GetTimestamp() < (*SecondNode)->GetTimestamp();
 				}
@@ -1104,7 +1109,7 @@ namespace UE::Audio::Insights
 		return Node->GetTreeDepth();
 	}
 
-	void FSignalFlowDashboardViewFactory::GetMinTreeDepthAndTimestamps(const TArray<FSignalFlowEntryKey>& NodeKeys, TreeDepthPair& OutMinTreeDepth, double& OutMinTimestamp) const
+	void FSignalFlowDashboardViewFactory::GetMinTimestampRecursive(const TArray<FSignalFlowEntryKey>& NodeKeys, double& OutMinTimestamp) const
 	{
 		for (const FSignalFlowEntryKey& NodeKey : NodeKeys)
 		{
@@ -1112,8 +1117,15 @@ namespace UE::Audio::Insights
 
 			if (Node && Node->IsValid())
 			{
-				OutMinTreeDepth = FMath::Min(GetOriginTreeDepthForNode(*Node), OutMinTreeDepth);
-				OutMinTimestamp = FMath::Min((*Node)->GetTimestamp(), OutMinTimestamp);
+				const ESignalFlowEntryType NodeEntryType = (*Node)->GetEntryKey().EntryType;
+
+				// Only count entries that are transient when finding the minimum timestamp of a branch
+				if (NodeEntryType == ESignalFlowEntryType::SoundSource || NodeEntryType == ESignalFlowEntryType::OwnerObject)
+				{
+					OutMinTimestamp = FMath::Min((*Node)->GetTimestamp(), OutMinTimestamp);
+				}
+				
+				GetMinTimestampRecursive((*Node)->FilteredInputs, OutMinTimestamp);
 			}
 		}
 	}
@@ -1124,27 +1136,10 @@ namespace UE::Audio::Insights
 		// This can happen if a node is not the root node but has no outputs
 		for (auto& [Key, Node] : FilteredEntries)
 		{
-			if (!Node.IsValid())
+			if (Node.IsValid() && !Node->NodeOrderIsValid())
 			{
-				continue;
-			}
-
-			if (!Node->NodeOrderIsValid())
-			{
-				int32 LowestNodeOrder = TNumericLimits<int32>::Max();
-				for (const FSignalFlowEntryKey& InputKey : Node->FilteredInputs)
-				{
-					const TSharedPtr<ISignalFlowNode>* InputNode = FilteredEntries.Find(InputKey);
-					if (InputNode && InputNode->IsValid() && (*InputNode)->NodeOrderIsValid())
-					{
-						LowestNodeOrder = FMath::Min((*InputNode)->GetNodeOrderID(), LowestNodeOrder);
-					}
-				}
-
-				if (LowestNodeOrder != TNumericLimits<int32>::Max())
-				{
-					Node->SetNodeOrderID(LowestNodeOrder);
-				}
+				// Force nodes with no outputs to one side of the graph
+				Node->SetNodeOrderID(TNumericLimits<int32>::Max());
 			}
 		}
 	}
@@ -1242,7 +1237,7 @@ namespace UE::Audio::Insights
 			PositionRowAroundOrigin(RowToProcess);
 
 			// Propagate change down the hierachy - moving nodes relative to their inputs
-			for (int32 PreviousRowsIndex = RowToProcessIndex + 1; PreviousRowsIndex < NodeIndexRows.Num(); ++PreviousRowsIndex)
+			for (int32 PreviousRowsIndex = RowToProcessIndex - 1; PreviousRowsIndex >= 0; --PreviousRowsIndex)
 			{
 				const TArray<int32> PrevRowOfIndexes = NodeIndexRows[PreviousRowsIndex];
 				TArray<TSharedPtr<ISignalFlowNode>> PrevPositionedNodes;
@@ -1356,6 +1351,11 @@ namespace UE::Audio::Insights
 
 	void FSignalFlowDashboardViewFactory::MakeSpaceForNodeInRow(const TSharedPtr<ISignalFlowNode>& Node, const TArray<TSharedPtr<ISignalFlowNode>>& PrevPositionedNodes, const int32 NumNodesInRow)
 	{
+		if (!Node->IsRealNode())
+		{
+			return;
+		}
+
 		// In cases where nodes share inputs/outputs, we can end up with nodes placed in the same XPos
 		// Check if we have hit this edge case and reposition the row accordingly
 		const TSharedPtr<ISignalFlowNode>* PrevNodeAtExistingPosition = PrevPositionedNodes.FindByPredicate([XPos = Node->XPos](const TSharedPtr<ISignalFlowNode>& AlreadyPositionedNode)
@@ -1365,18 +1365,18 @@ namespace UE::Audio::Insights
 
 		if (PrevNodeAtExistingPosition)
 		{
-			// there already is a node at this position - either shift all nodes to the left or put the new node on the right
-			const bool bShiftNodesLeft = static_cast<float>(PrevPositionedNodes.Num()) / static_cast<float>(NumNodesInRow) <= 0.5f;
-			if (bShiftNodesLeft)
+			// there already is a node at this position - shift nodes to the left
+			for (const TSharedPtr<ISignalFlowNode>& PrevNode : PrevPositionedNodes)
 			{
-				for (const TSharedPtr<ISignalFlowNode>& PrevNode : PrevPositionedNodes)
+				if (!PrevNode.IsValid())
 				{
-					PrevNode->XPos -= 1.0f;
+					continue;
 				}
-			}
-			else
-			{
-				Node->XPos = PrevPositionedNodes.Last()->XPos + 1.0f;
+			
+				if (PrevNode->XPos <= Node->XPos)
+				{
+					PrevNode->XPos -= 1.0f;
+				} 
 			}
 		}
 	}
diff --git a/Engine/Plugins/AudioInsights/Source/AudioInsights/Public/Views/SignalFlowDashboardViewFactory.h b/Engine/Plugins/AudioInsights/Source/AudioInsights/Public/Views/SignalFlowDashboardViewFactory.h
--- a/Engine/Plugins/AudioInsights/Source/AudioInsights/Public/Views/SignalFlowDashboardViewFactory.h
+++ b/Engine/Plugins/AudioInsights/Source/AudioInsights/Public/Views/SignalFlowDashboardViewFactory.h
@@ -97,7 +97,7 @@ namespace UE::Audio::Insights
 		void AssignNodeOrderRecursive(TSharedPtr<ISignalFlowNode> OutNode, const int32 Depth, int32& OutNodeProcessingOrder);
 		void SortNodeInputs(TSharedPtr<ISignalFlowNode> OutNode);
 		TreeDepthPair GetOriginTreeDepthForNode(const TSharedPtr<ISignalFlowNode>& Node) const;
-		void GetMinTreeDepthAndTimestamps(const TArray<FSignalFlowEntryKey>& NodeKeys, TreeDepthPair& OutMinTreeDepth, double& OutMinTimestamp) const;
+		void GetMinTimestampRecursive(const TArray<FSignalFlowEntryKey>& NodeKeys, double& OutMinTimestamp) const;
 
 		void FixUnassignedNodeOrders();
 		void SortNodeOrder();
