[Audio Insights] Add zooming in/out to the signal flow graph

#Synced-CL 49947352

#Synced-CL 49947352

diff --git a/Engine/Plugins/AudioInsights/Source/AudioInsights/Private/Views/SSignalFlowGraph.cpp b/Engine/Plugins/AudioInsights/Source/AudioInsights/Private/Views/SSignalFlowGraph.cpp
--- a/Engine/Plugins/AudioInsights/Source/AudioInsights/Private/Views/SSignalFlowGraph.cpp
+++ b/Engine/Plugins/AudioInsights/Source/AudioInsights/Private/Views/SSignalFlowGraph.cpp
@@ -9,6 +9,9 @@
 #include "Styling/StyleColors.h"
 #include "Views/SignalFlowNodes.h"
 #include "Widgets/Colors/SColorBlock.h"
+#include "Widgets/Layout/SBorder.h"
+#include "Widgets/Layout/SScaleBox.h"
+#include "Widgets/Layout/SScrollBox.h"
 #include "Widgets/SBoxPanel.h"
 #include "Widgets/Text/STextBlock.h"
 
@@ -34,6 +37,13 @@ namespace UE::Audio::Insights
 		constexpr float NodeBorderPadding = 12.0f;
 		constexpr float SendLevelTextMargin = 2.0f;
 
+		constexpr float AnimationSpeed = 15.0f;
+		constexpr float DefaultZoomScale = 1.0f;
+		constexpr float ZoomStepFactor = 0.1f;
+		constexpr float MinZoomScale = 0.5f;
+		constexpr float MaxZoomScale = 2.0f;
+		constexpr float ScrollAnimInRangeThreshold = (1.e-2f);
+
 		const FSlateRoundedBoxBrush& GetNodeRoundedBoxBrush()
 		{
 			static const FSlateRoundedBoxBrush RoundedBoxBrush(FLinearColor::White, NodeRoundedBrushRadius);
@@ -99,6 +109,7 @@ namespace UE::Audio::Insights
 		PositionInGraph = InArgs._PositionInGraph;
 		XPos = InArgs._XPos;
 		YPos = InArgs._YPos;
+		ZoomScaleFactor = InArgs._ZoomScaleFactor;
 		IsSelected = InArgs._IsSelected;
 		GraphOrientation = InArgs._GraphOrientation;
 		OnNodeSelected = InArgs._OnNodeSelected;
@@ -119,7 +130,7 @@ namespace UE::Audio::Insights
 			.AutoHeight()
 			.HAlign(HAlign_Center)
 			[
-				SNew(SBorder)
+				SAssignNew(SelectedNodeBorder, SBorder)
 				.BorderImage(&GetNodeRoundedBoxBrush())
 				.HAlign(HAlign_Fill)
 				.VAlign(VAlign_Fill)
@@ -129,7 +140,7 @@ namespace UE::Audio::Insights
 				})
 				.Padding(SelectedNodeBorderThickness)
 				[
-					SNew(SBorder)
+					SAssignNew(NodeBorder, SBorder)
 					.BorderImage(&GetNodeRoundedBoxBrush())
 					.HAlign(HAlign_Fill)
 					.VAlign(VAlign_Fill)
@@ -210,12 +221,12 @@ namespace UE::Audio::Insights
 		if (GraphOrientation.IsBound() && GraphOrientation.Get() == EOrientation::Orient_Horizontal)
 		{
 			// Input pos = left center of node
-			return { CachedPositionInGraph.X + SelectedNodeBorderThickness, CachedPositionInGraph.Y + (Size.Y / 2.0f) };
+			return FVector2D{ CachedPositionInGraph.X + SelectedNodeBorderThickness, CachedPositionInGraph.Y + (Size.Y / 2.0f) } * ZoomScaleFactor.Get();
 		}
 		else
 		{
 			// Input pos = top center of node
-			return { CachedPositionInGraph.X + (Size.X / 2.0f), CachedPositionInGraph.Y + SelectedNodeBorderThickness };
+			return FVector2D{ CachedPositionInGraph.X + (Size.X / 2.0f), CachedPositionInGraph.Y + SelectedNodeBorderThickness } * ZoomScaleFactor.Get();
 		}
 	}
 	
@@ -228,19 +239,32 @@ namespace UE::Audio::Insights
 		if (GraphOrientation.IsBound() && GraphOrientation.Get() == EOrientation::Orient_Horizontal)
 		{
 			// Out pos = right center of node
-			return { CachedPositionInGraph.X + Size.X - SelectedNodeBorderThickness, CachedPositionInGraph.Y + (Size.Y / 2.0f) };
+			return FVector2D{ CachedPositionInGraph.X + Size.X - SelectedNodeBorderThickness, CachedPositionInGraph.Y + (Size.Y / 2.0f) } * ZoomScaleFactor.Get();
 		}
 		else
 		{
 			// Output pos = bottom center of node
-			return { CachedPositionInGraph.X + (Size.X / 2.0f), CachedPositionInGraph.Y + Size.Y - SelectedNodeBorderThickness };
+			return FVector2D{ CachedPositionInGraph.X + (Size.X / 2.0f), CachedPositionInGraph.Y + Size.Y - SelectedNodeBorderThickness } * ZoomScaleFactor.Get();
 		}
 	}
 
-	FVector2D SSignalFlowGraphNode::UpdatePositionInGraph()
+	void SSignalFlowGraphNode::UpdateNodeLayout()
 	{
+		using namespace SSignalFlowGraphStyle;
+
+		// Ensure border thickness does not get scaled to avoid rendering artifacts
+		if (SelectedNodeBorder.IsValid())
+		{
+			SelectedNodeBorder->SetPadding(SelectedNodeBorderThickness / ZoomScaleFactor.Get());
+		}
+
+		if (NodeBorder.IsValid())
+		{
+			NodeBorder->SetPadding(NodeBorderThickness / ZoomScaleFactor.Get());
+		}
+
+		// Update the node's position in the graph
 		CachedPositionInGraph = PositionInGraph.Get();
-		return CachedPositionInGraph;
 	}
 
 	FVector2D SSignalFlowGraphNode::GetCachedPositionInGraph() const
@@ -257,6 +281,7 @@ namespace UE::Audio::Insights
 		ReceiverPositionInGraph = InArgs._ReceiverPositionInGraph;
 		SendLevel = InArgs._SendLevel;
 		GraphOrientation = InArgs._GraphOrientation;
+		ZoomScaleFactor = InArgs._ZoomScaleFactor;
 
 		ChildSlot
 		[
@@ -271,7 +296,7 @@ namespace UE::Audio::Insights
 			.HAlign(HAlign_Center)
 			.VAlign(VAlign_Center)
 			[
-				SNew(SBorder)
+				SAssignNew(LabelBorder, SBorder)
 				.BorderImage(&GetSendLabelRoundedBoxBrush())
 				.HAlign(HAlign_Fill)
 				.VAlign(VAlign_Fill)
@@ -318,12 +343,19 @@ namespace UE::Audio::Insights
 		return Size;
 	}
 
-	FVector2D SSignalFlowGraphSendLabel::UpdatePositionInGraph()
+	void SSignalFlowGraphSendLabel::UpdateLabelLayout()
 	{
 		using namespace SSignalFlowGraphStyle;
 
-		const FVector2D SenderNodePos = SenderPositionInGraph.Get();
-		const FVector2D ReceiverNodePos = ReceiverPositionInGraph.Get();
+		// Ensure border thickness does not get scaled to avoid rendering artifacts
+		if (LabelBorder.IsValid())
+		{
+			LabelBorder->SetPadding(SendLabelBorderThickness / ZoomScaleFactor.Get());
+		}
+
+		// Update the label's position in the graph
+		const FVector2D SenderNodePos = SenderPositionInGraph.Get() / ZoomScaleFactor.Get();
+		const FVector2D ReceiverNodePos = ReceiverPositionInGraph.Get() / ZoomScaleFactor.Get();
 
 		const EOrientation Orientation = GraphOrientation.Get();
 		const FVector2D PointAlongSpline = FMath::CubicInterp(SenderNodePos, GetSplineStartDirection(Orientation), ReceiverNodePos, GetSplineEndDirection(Orientation), SendLabelPropAlongLine);
@@ -331,7 +363,6 @@ namespace UE::Audio::Insights
 		const FVector2D LabelSize = GetDesiredSize();
 
 		CachedPositionInGraph = PointAlongSpline - (LabelSize / 2.0f);
-		return CachedPositionInGraph;
 	}
 
 	FVector2D SSignalFlowGraphSendLabel::GetCachedPositionInGraph() const
@@ -343,6 +374,9 @@ namespace UE::Audio::Insights
 	{
 		ListSource = InArgs._ListSource;
 		SelectedItem = InArgs._SelectedItem;
+		ActiveAudioDeviceEntry = InArgs._ActiveAudioDeviceEntry;
+		HorizontalScrollBox = InArgs._HorizontalScrollBox;
+		VerticalScrollBox = InArgs._VerticalScrollBox;
 		OnNodeSelected = InArgs._OnNodeSelected;
 
 		Clipping = EWidgetClipping::ClipToBounds;
@@ -353,7 +387,25 @@ namespace UE::Audio::Insights
 
 		ChildSlot
 		[
-			Canvas.ToSharedRef()
+			SNew(SScaleBox)
+			.Stretch(EStretch::UserSpecified)
+			.StretchDirection(EStretchDirection::Both)
+			.UserSpecifiedScale_Lambda([this]()
+			{
+				return ZoomScaleFactor;
+			})
+			.PixelSnappingMethod(EWidgetPixelSnapping::Disabled)
+			.HAlign(HAlign_Fill)
+			.VAlign(VAlign_Fill)
+			.Content()
+			[
+				SNew(SBox)
+				.HAlign(HAlign_Fill)
+				.VAlign(VAlign_Fill)
+				[
+					Canvas.ToSharedRef()
+				]
+			]
 		];
 	}
 
@@ -385,6 +437,11 @@ namespace UE::Audio::Insights
 		RefreshGraph();
 	}
 
+	void SSignalFlowGraph::ResetZoom()
+	{
+		bAnimateZoom = true;
+	}
+
 	TSharedPtr<SSignalFlowGraphNode> SSignalFlowGraph::FindNodeWidget(const FSignalFlowEntryKey& NodeEntryKey) const
 	{
 		const TSharedPtr<SSignalFlowGraphNode>* NodeWidget = NodeWidgets.Find(NodeEntryKey);
@@ -396,6 +453,58 @@ namespace UE::Audio::Insights
 		return nullptr;
 	}
 
+	void SSignalFlowGraph::Tick(const FGeometry& AllottedGeometry, const double InCurrentTime, const float InDeltaTime)
+	{
+		using namespace SSignalFlowGraphStyle;
+
+		if (!bAnimateZoom)
+		{
+			return;
+		}
+
+		ZoomScaleFactor = FMath::FInterpTo(ZoomScaleFactor, DefaultZoomScale, InDeltaTime, AnimationSpeed);
+
+		TSharedPtr<SSignalFlowGraphNode>* FocusedNodeWidget = nullptr;
+		if (SelectedItem.Get().IsValid())
+		{
+			FocusedNodeWidget = NodeWidgets.Find(SelectedItem.Get()->GetSignalFlowEntryKey());
+		}
+		else if (ActiveAudioDeviceEntry.Get().IsValid())
+		{
+			FocusedNodeWidget = NodeWidgets.Find(ActiveAudioDeviceEntry.Get()->GetSignalFlowEntryKey());
+		}
+
+		if (FocusedNodeWidget == nullptr || !FocusedNodeWidget->IsValid())
+		{
+			bAnimateZoom = false;
+			return;
+		}
+
+		const FVector2D SelectedNodeCenterPos = (*FocusedNodeWidget)->GetCachedPositionInGraph() + ((*FocusedNodeWidget)->GetDesiredSize() * 0.5f);
+
+		EAnimationResult XScrollAnimProgress = EAnimationResult::InProgress;
+		EAnimationResult YScrollAnimProgress = EAnimationResult::InProgress;
+
+		const TSharedPtr<SScrollBox> HorizontalScrollBoxPinned = HorizontalScrollBox.Get().Pin();
+		if (HorizontalScrollBoxPinned.IsValid())
+		{
+			XScrollAnimProgress = AnimateScrollTo(HorizontalScrollBoxPinned.ToSharedRef(), SelectedNodeCenterPos, InDeltaTime);
+		}
+
+		const TSharedPtr<SScrollBox> VerticalScrollBoxPinned = VerticalScrollBox.Get().Pin();
+		if (VerticalScrollBoxPinned.IsValid())
+		{
+			YScrollAnimProgress = AnimateScrollTo(VerticalScrollBoxPinned.ToSharedRef(), SelectedNodeCenterPos, InDeltaTime);
+		}
+
+		if (XScrollAnimProgress == EAnimationResult::Complete 
+			&& YScrollAnimProgress == EAnimationResult::Complete 
+			&& FMath::IsNearlyEqual(ZoomScaleFactor, DefaultZoomScale))
+		{
+			bAnimateZoom = false;
+		}
+	}
+
 	int32 SSignalFlowGraph::OnPaint(const FPaintArgs& Args, const FGeometry& AllottedGeometry, const FSlateRect& MyCullingRect, FSlateWindowElementList& OutDrawElements, int32 LayerId, const FWidgetStyle& InWidgetStyle, bool bParentEnabled) const
 	{
 		LayerId = OnPaintConnections(AllottedGeometry, OutDrawElements, LayerId);
@@ -403,6 +512,42 @@ namespace UE::Audio::Insights
 		return SCompoundWidget::OnPaint(Args, AllottedGeometry, MyCullingRect, OutDrawElements, LayerId, InWidgetStyle, ShouldBeEnabled(bParentEnabled));
 	}
 
+	FReply SSignalFlowGraph::OnMouseWheel(const FGeometry& MyGeometry, const FPointerEvent& MouseEvent)
+	{
+		using namespace SSignalFlowGraphStyle;
+
+		if (MouseEvent.IsControlDown())
+		{
+			bAnimateZoom = false;
+
+			const float OldZoomScaleFactor = ZoomScaleFactor;
+			ZoomScaleFactor = FMath::Clamp(ZoomScaleFactor + (MouseEvent.GetWheelDelta() * ZoomStepFactor), MinZoomScale, MaxZoomScale);
+
+			const TSharedPtr<SScrollBox> HorizontalScrollBoxPinned = HorizontalScrollBox.Get().Pin();
+			const TSharedPtr<SScrollBox> VerticalScrollBoxPinned = VerticalScrollBox.Get().Pin();
+
+			if (HorizontalScrollBoxPinned.IsValid() && VerticalScrollBoxPinned.IsValid())
+			{
+				const FVector2D MousePosInGraph = MyGeometry.AbsoluteToLocal(MouseEvent.GetScreenSpacePosition());
+
+				const FVector2D ViewportOffsetFraction(HorizontalScrollBoxPinned->GetViewOffsetFraction(), VerticalScrollBoxPinned->GetViewOffsetFraction());
+				const FVector2D ViewportPosInGraph = MyGeometry.GetLocalSize() * ViewportOffsetFraction;
+				const FVector2D MousePosInViewport = MousePosInGraph - ViewportPosInGraph;
+
+				const float ZoomScaleRatio = ZoomScaleFactor / OldZoomScaleFactor;
+				const FVector2D CurrentScrollOffset(HorizontalScrollBoxPinned->GetScrollOffset(), VerticalScrollBoxPinned->GetScrollOffset());
+				const FVector2D NewScrollOffset = ((MousePosInViewport + CurrentScrollOffset) * ZoomScaleRatio) - MousePosInViewport;
+
+				HorizontalScrollBoxPinned->SetScrollOffset(FMath::Clamp(NewScrollOffset.X, 0.0f, HorizontalScrollBoxPinned->GetScrollOffsetOfEnd() * ZoomScaleRatio));
+				VerticalScrollBoxPinned->SetScrollOffset(FMath::Clamp(NewScrollOffset.Y, 0.0f, VerticalScrollBoxPinned->GetScrollOffsetOfEnd() * ZoomScaleRatio));
+			}
+
+			return FReply::Handled();
+		}
+
+		return FReply::Unhandled();
+	}
+
 	int32 SSignalFlowGraph::OnPaintConnections(const FGeometry& AllottedGeometry, FSlateWindowElementList& OutDrawElements, int32 LayerId) const
 	{
 		using namespace SSignalFlowGraphStyle;
@@ -452,7 +597,7 @@ namespace UE::Audio::Insights
 						GetSplineStartDirection(Orientation),
 						ConnectionEndPos,
 						GetSplineEndDirection(Orientation),
-						2.0f,
+						FMath::Max(2.0f * ZoomScaleFactor, 1.0f),
 						ESlateDrawEffect::NoPixelSnapping,
 						GetConnectionColor().GetSpecifiedColor()
 					);
@@ -492,7 +637,7 @@ namespace UE::Audio::Insights
 			}
 		}
 
-		return FVector2D{ GraphWidth, GraphHeight };
+		return FVector2D{ GraphWidth, GraphHeight } * InLayoutScaleMultiplier * ZoomScaleFactor;
 	}
 
 	void SSignalFlowGraph::CreateGraphNodes(TArray<FSignalFlowEntryKey>& OutRemainingNodeKeys)
@@ -675,6 +820,7 @@ namespace UE::Audio::Insights
 			.Entry(EntryNode->Entry)
 			.ListSource(ListSource)
 			.DisplayName(EntryNode->Entry.IsValid() ? EntryNode->Entry->GetDisplayName() : FText::GetEmpty())
+			.PixelSnappingMethod(EWidgetPixelSnapping::SnapToPixel)
 			.PositionInGraph_Lambda([this, EntryKey]() -> FVector2D
 			{
 				const TSharedPtr<SSignalFlowGraphNode>* Node = NodeWidgets.Find(EntryKey);
@@ -720,6 +866,10 @@ namespace UE::Audio::Insights
 			.OnNodeSelected_Lambda([this](TSharedPtr<FSignalFlowDashboardEntry> SelectedEntry)
 			{
 				OnNodeSelected.ExecuteIfBound(SelectedEntry);
+			})
+			.ZoomScaleFactor_Lambda([this]()
+			{
+				return ZoomScaleFactor;
 			});
 
 		NodeWidgets.Add(Node->GetEntryKey(), NodeWidget);
@@ -730,7 +880,8 @@ namespace UE::Audio::Insights
 				const TSharedPtr<SSignalFlowGraphNode>* Node = NodeWidgets.Find(EntryKey);
 				if (Node && Node->IsValid())
 				{
-					return (*Node)->UpdatePositionInGraph();
+					(*Node)->UpdateNodeLayout();
+					return (*Node)->GetCachedPositionInGraph();
 				}
 
 				return FVector2D::ZeroVector;
@@ -762,6 +913,7 @@ namespace UE::Audio::Insights
 		}
 
 		TSharedPtr<SSignalFlowGraphSendLabel> SendLabelWidget = SNew(SSignalFlowGraphSendLabel)
+			.PixelSnappingMethod(EWidgetPixelSnapping::SnapToPixel)
 			.SendReceiverEntryKey(ReceiverEntryKey)
 			.SenderPositionInGraph_Lambda([this, SenderEntryKey]() -> FVector2D
 			{
@@ -801,6 +953,10 @@ namespace UE::Audio::Insights
 			.GraphOrientation_Lambda([this]()
 			{
 				return Orientation;
+			})
+			.ZoomScaleFactor_Lambda([this]()
+			{
+				return ZoomScaleFactor;
 			});
 
 		SendLevelWidgets.Add(SenderEntryKey, SendLabelWidget);
@@ -813,7 +969,8 @@ namespace UE::Audio::Insights
 
 			if (FoundExistingSendLabel.IsValid())
 			{
-				return FoundExistingSendLabel->UpdatePositionInGraph();
+				FoundExistingSendLabel->UpdateLabelLayout();
+				return FoundExistingSendLabel->GetCachedPositionInGraph();
 			}
 
 			return FVector2D::ZeroVector;
@@ -899,6 +1056,23 @@ namespace UE::Audio::Insights
 			return SendLabel.IsValid() && SendLabel->GetSendReceiverEntryKey() == ReceiverKey;
 		});
 	}
+
+	SSignalFlowGraph::EAnimationResult SSignalFlowGraph::AnimateScrollTo(const TSharedRef<SScrollBox>& ScrollBox, const FVector2D& ScrollToPos, const float DeltaTime)
+	{
+		using namespace SSignalFlowGraphStyle;
+
+		const FGeometry& ScrollBoxGeometry = ScrollBox->GetTickSpaceGeometry();
+		const FVector2D CenterPos = ScrollBoxGeometry.GetLocalSize() * 0.5f;
+		const FVector2D ScrollOffsetTarget = ScrollToPos - CenterPos;
+
+		const EOrientation ScrollOrientation = ScrollBox->GetOrientation();
+		const float TargetAlongOrientation = ScrollOrientation == EOrientation::Orient_Horizontal ? ScrollOffsetTarget.X : ScrollOffsetTarget.Y;
+
+		const float ScrollOffset = FMath::FInterpTo(ScrollBox->GetScrollOffset(), TargetAlongOrientation, DeltaTime, AnimationSpeed);
+		ScrollBox->SetScrollOffset(ScrollOffset);
+
+		return FMath::IsNearlyEqual(TargetAlongOrientation, ScrollOffset, ScrollAnimInRangeThreshold) ? EAnimationResult::Complete : EAnimationResult::InProgress;
+	}
 } // namespace UE::Audio::Insights
 
 #undef LOCTEXT_NAMESPACE
diff --git a/Engine/Plugins/AudioInsights/Source/AudioInsights/Private/Views/SignalFlowDashboardViewFactory.cpp b/Engine/Plugins/AudioInsights/Source/AudioInsights/Private/Views/SignalFlowDashboardViewFactory.cpp
--- a/Engine/Plugins/AudioInsights/Source/AudioInsights/Private/Views/SignalFlowDashboardViewFactory.cpp
+++ b/Engine/Plugins/AudioInsights/Source/AudioInsights/Private/Views/SignalFlowDashboardViewFactory.cpp
@@ -143,7 +143,7 @@ namespace UE::Audio::Insights
 			}
 			else if (InKeyEvent.GetKey() == EKeys::F && InKeyEvent.IsControlDown())
 			{
-				return ScrollGraphViewToSelectedNode();
+				return ScrollGraphViewToSelectedNode(true /*bResetZoom*/);
 			}
 		}
 
@@ -277,7 +277,7 @@ namespace UE::Audio::Insights
 			.ToolTipText(LOCTEXT("AudioDashboard_SignalFlow_CenterViewTooltip", "Focus the graph view on the selected node. (Ctrl + f)"))
 			.OnClicked_Lambda([this]()
 			{
-				return ScrollGraphViewToSelectedNode();
+				return ScrollGraphViewToSelectedNode(true /*bResetZoom*/);
 			})
 			[
 				Style.CreateButtonContentWidget("AudioInsights.Icon.SignalFlow.CenterView", FText::GetEmpty(), ButtonTextStyleName)
@@ -285,7 +285,7 @@ namespace UE::Audio::Insights
 		];
 	}
 
-	FReply FSignalFlowDashboardViewFactory::ScrollGraphViewToSelectedNode()
+	FReply FSignalFlowDashboardViewFactory::ScrollGraphViewToSelectedNode(const bool bResetZoom)
 	{
 		if (GraphWidgetContents.IsValid())
 		{
@@ -295,14 +295,21 @@ namespace UE::Audio::Insights
 				const TSharedPtr<SSignalFlowGraphNode> FocusedNodeWidget = GraphWidgetContents->FindNodeWidget(FocusedEntry->GetSignalFlowEntryKey());
 				if (FocusedNodeWidget.IsValid())
 				{
-					if (GraphHorizontalScrollBox.IsValid())
+					if (bResetZoom)
 					{
-						GraphHorizontalScrollBox->ScrollDescendantIntoView(FocusedNodeWidget, true, EDescendantScrollDestination::Center);
+						GraphWidgetContents->ResetZoom();
 					}
-
-					if (GraphVerticalScrollBox.IsValid())
+					else
 					{
-						GraphVerticalScrollBox->ScrollDescendantIntoView(FocusedNodeWidget, true, EDescendantScrollDestination::Center);
+						if (GraphHorizontalScrollBox.IsValid())
+						{
+							GraphHorizontalScrollBox->ScrollDescendantIntoView(FocusedNodeWidget, true, EDescendantScrollDestination::Center);
+						}
+
+						if (GraphVerticalScrollBox.IsValid())
+						{
+							GraphVerticalScrollBox->ScrollDescendantIntoView(FocusedNodeWidget, true, EDescendantScrollDestination::Center);
+						}
 					}
 
 					return FReply::Handled();
@@ -381,7 +388,7 @@ namespace UE::Audio::Insights
 
 	TSharedRef<SWidget> FSignalFlowDashboardViewFactory::CreateGraphSection()
 	{
-		return SAssignNew(GraphVerticalScrollBox, SScrollBox)
+		GraphVerticalScrollBox = SAssignNew(GraphVerticalScrollBox, SScrollBox)
 		.Orientation(EOrientation::Orient_Vertical)
 
 		+ SScrollBox::Slot()
@@ -413,7 +420,14 @@ namespace UE::Audio::Insights
 				[
 					SAssignNew(GraphWidgetContents, SSignalFlowGraph)
 					.ListSource(&FilteredGraphNodes)
-					.SelectedItem_Lambda([this]() { return SelectedItem; })
+					.SelectedItem_Lambda([this]()
+					{
+						return SelectedItem;
+					})
+					.ActiveAudioDeviceEntry_Lambda([this]()
+					{
+						return GetFilteredAudioDeviceEntry();
+					})
 					.OnNodeSelected_Lambda([this](TSharedPtr<FSignalFlowDashboardEntry> SelectedNodeEntry)
 					{
 						if (SelectedNodeEntry.IsValid())
@@ -465,6 +479,14 @@ namespace UE::Audio::Insights
 				]
 			]
 		];
+
+		if (GraphWidgetContents.IsValid())
+		{
+			GraphWidgetContents->SetHorizontalScrollBox(GraphHorizontalScrollBox);
+			GraphWidgetContents->SetVerticalScrollBox(GraphVerticalScrollBox);
+		}
+
+		return GraphVerticalScrollBox.ToSharedRef();
 	}
 
 	TSharedRef<SListView<TSharedPtr<FSignalFlowDashboardEntry>>> FSignalFlowDashboardViewFactory::CreateSelectableList(const FText& HeadingText, const TArray<TSharedPtr<FSignalFlowDashboardEntry>>& ListSource)
@@ -650,7 +672,7 @@ namespace UE::Audio::Insights
 
 			if (bAutoFocusAfterGraphRefresh)
 			{
-				ScrollGraphViewToSelectedNode();
+				ScrollGraphViewToSelectedNode(false /*bResetZoom*/);
 				bAutoFocusAfterGraphRefresh = false;
 			}
 		}
diff --git a/Engine/Plugins/AudioInsights/Source/AudioInsights/Public/Views/SSignalFlowGraph.h b/Engine/Plugins/AudioInsights/Source/AudioInsights/Public/Views/SSignalFlowGraph.h
--- a/Engine/Plugins/AudioInsights/Source/AudioInsights/Public/Views/SSignalFlowGraph.h
+++ b/Engine/Plugins/AudioInsights/Source/AudioInsights/Public/Views/SSignalFlowGraph.h
@@ -8,6 +8,9 @@
 
 #define UE_API AUDIOINSIGHTS_API
 
+class SBorder;
+class SScrollBox;
+
 namespace UE::Audio::Insights
 {
 	struct FSignalFlowDashboardEntry;
@@ -30,6 +33,7 @@ namespace UE::Audio::Insights
 			_PositionInGraph(FVector2d::ZeroVector),
 			_XPos(0.0f),
 			_YPos(0.0f),
+			_ZoomScaleFactor(1.0f),
 			_IsSelected(false),
 			_GraphOrientation(EOrientation::Orient_Vertical)
 		{
@@ -41,6 +45,7 @@ namespace UE::Audio::Insights
 		SLATE_ATTRIBUTE(FVector2d, PositionInGraph)
 		SLATE_ATTRIBUTE(float, XPos)
 		SLATE_ATTRIBUTE(float, YPos)
+		SLATE_ATTRIBUTE(float, ZoomScaleFactor)
 		SLATE_ATTRIBUTE(bool, IsSelected)
 		SLATE_ATTRIBUTE(EOrientation, GraphOrientation)
 		SLATE_EVENT(FOnSignalFlowNodeSelected, OnNodeSelected)
@@ -55,7 +60,7 @@ namespace UE::Audio::Insights
 		UE_API FVector2D GetLocalInputPos() const;
 		UE_API FVector2D GetLocalOutputPos() const;
 
-		UE_API FVector2D UpdatePositionInGraph();
+		UE_API void UpdateNodeLayout();
 		UE_API FVector2D GetCachedPositionInGraph() const;
 
 		const FSignalFlowEntryKey& GetNodeEntryKey() const { return NodeEntryKey; }
@@ -72,12 +77,16 @@ namespace UE::Audio::Insights
 		TAttribute<FVector2d> PositionInGraph;
 		TAttribute<float> XPos;
 		TAttribute<float> YPos;
+		TAttribute<float> ZoomScaleFactor;
 		TAttribute<bool> IsSelected;
 		TAttribute<EOrientation> GraphOrientation;
 
 		FOnSignalFlowNodeSelected OnNodeSelected;
 		FVector2d CachedPositionInGraph = FVector2d::Zero();
 		FSignalFlowEntryKey NodeEntryKey;
+
+		TSharedPtr<SBorder> SelectedNodeBorder;
+		TSharedPtr<SBorder> NodeBorder;
 	};
 
 	class SSignalFlowGraphSendLabel : public SCompoundWidget
@@ -88,7 +97,8 @@ namespace UE::Audio::Insights
 			_SenderPositionInGraph(FVector2d::ZeroVector),
 			_ReceiverPositionInGraph(FVector2d::ZeroVector),
 			_SendLevel(),
-			_GraphOrientation(EOrientation::Orient_Vertical)
+			_GraphOrientation(EOrientation::Orient_Vertical),
+			_ZoomScaleFactor(1.0f)
 		{
 		}
 
@@ -97,13 +107,14 @@ namespace UE::Audio::Insights
 		SLATE_ATTRIBUTE(FVector2d, ReceiverPositionInGraph)
 		SLATE_ATTRIBUTE(TOptional<float>, SendLevel)
 		SLATE_ATTRIBUTE(EOrientation, GraphOrientation)
+		SLATE_ATTRIBUTE(float, ZoomScaleFactor)
 
 		SLATE_END_ARGS()
 
 		UE_API void Construct(const FArguments& InArgs);
 		UE_API virtual FVector2D ComputeDesiredSize(float InLayoutScaleMultiplier) const override;
 		
-		UE_API FVector2D UpdatePositionInGraph();
+		UE_API void UpdateLabelLayout();
 		UE_API FVector2D GetCachedPositionInGraph() const;
 
 		const FSignalFlowEntryKey& GetSendReceiverEntryKey() const { return SendReceiverEntryKey; }
@@ -114,8 +125,11 @@ namespace UE::Audio::Insights
 		TAttribute<FVector2d> ReceiverPositionInGraph;
 		TAttribute<TOptional<float>> SendLevel;
 		TAttribute<EOrientation> GraphOrientation;
+		TAttribute<float> ZoomScaleFactor;
 
 		FVector2d CachedPositionInGraph = FVector2d::Zero();
+
+		TSharedPtr<SBorder> LabelBorder;
 	};
 
 	class SSignalFlowGraph : public SCompoundWidget
@@ -128,19 +142,30 @@ namespace UE::Audio::Insights
 
 		SLATE_ARGUMENT(TArray<TSharedPtr<ISignalFlowNode>>*, ListSource)
 		SLATE_ATTRIBUTE(TSharedPtr<FSignalFlowDashboardEntry>, SelectedItem)
+		SLATE_ATTRIBUTE(TSharedPtr<FSignalFlowDashboardEntry>, ActiveAudioDeviceEntry)
+		SLATE_ATTRIBUTE(TWeakPtr<SScrollBox>, HorizontalScrollBox)
+		SLATE_ATTRIBUTE(TWeakPtr<SScrollBox>, VerticalScrollBox)
 		SLATE_EVENT(FOnSignalFlowNodeSelected, OnNodeSelected)
 
 		SLATE_END_ARGS()
 
 		UE_API void Construct(const FArguments& InArgs);
+		void SetHorizontalScrollBox(const TWeakPtr<SScrollBox> ScrollBox) { HorizontalScrollBox = ScrollBox; }
+		void SetVerticalScrollBox(const TWeakPtr<SScrollBox> ScrollBox) { VerticalScrollBox = ScrollBox; }
+
 		UE_API void RefreshGraph();
 
 		UE_API void ToggleOrientation();
-		UE_API EOrientation GetGraphOrientation() const { return Orientation; }
+		EOrientation GetGraphOrientation() const { return Orientation; }
+
+		UE_API void ResetZoom();
+
 		UE_API TSharedPtr<SSignalFlowGraphNode> FindNodeWidget(const FSignalFlowEntryKey& NodeEntryKey) const;
 
+		UE_API virtual void Tick(const FGeometry& AllottedGeometry, const double InCurrentTime, const float InDeltaTime) override;
 		UE_API virtual int32 OnPaint(const FPaintArgs& Args, const FGeometry& AllottedGeometry, const FSlateRect& MyCullingRect, FSlateWindowElementList& OutDrawElements, int32 LayerId, const FWidgetStyle& InWidgetStyle, bool bParentEnabled) const override;
-
+		
+		UE_API virtual FReply OnMouseWheel(const FGeometry& MyGeometry, const FPointerEvent& MouseEvent) override;
 	protected:
 		UE_API virtual FVector2D ComputeDesiredSize(float InLayoutScaleMultiplier) const override;
 
@@ -166,6 +191,14 @@ namespace UE::Audio::Insights
 		TSharedPtr<SSignalFlowGraphSendLabel> FindSendLabel(const FSignalFlowEntryKey& SenderKey, const FSignalFlowEntryKey& ReceiverKey) const;
 		const TSharedPtr<SSignalFlowGraphSendLabel>* FindSendLabelInList(const TArray<TSharedPtr<SSignalFlowGraphSendLabel>>& SendLabelWidgets, const FSignalFlowEntryKey& ReceiverKey) const;
 
+		enum class EAnimationResult : uint8
+		{
+			InProgress,
+			Complete
+		};
+
+		EAnimationResult AnimateScrollTo(const TSharedRef<SScrollBox>& ScrollBox, const FVector2D& ScrollToPos, const float DeltaTime);
+
 		TSharedPtr<SCanvas> Canvas;
 
 		TMap<FSignalFlowEntryKey, TSharedPtr<SSignalFlowGraphNode>> NodeWidgets;
@@ -181,11 +214,18 @@ namespace UE::Audio::Insights
 
 		TArray<TSharedPtr<ISignalFlowNode>>* ListSource;
 		TAttribute<TSharedPtr<FSignalFlowDashboardEntry>> SelectedItem;
+		TAttribute<TSharedPtr<FSignalFlowDashboardEntry>> ActiveAudioDeviceEntry;
+
+		TAttribute<TWeakPtr<SScrollBox>> HorizontalScrollBox;
+		TAttribute<TWeakPtr<SScrollBox>> VerticalScrollBox;
 
 		FOnSignalFlowNodeSelected OnNodeSelected;
 		EOrientation Orientation = EOrientation::Orient_Vertical;
 
+		float ZoomScaleFactor = 1.0f;
+
 		bool bNeedsRefresh = false;
+		bool bAnimateZoom = false;
 	};
 } // namespace UE::Audio::Insights
 
diff --git a/Engine/Plugins/AudioInsights/Source/AudioInsights/Public/Views/SignalFlowDashboardViewFactory.h b/Engine/Plugins/AudioInsights/Source/AudioInsights/Public/Views/SignalFlowDashboardViewFactory.h
--- a/Engine/Plugins/AudioInsights/Source/AudioInsights/Public/Views/SignalFlowDashboardViewFactory.h
+++ b/Engine/Plugins/AudioInsights/Source/AudioInsights/Public/Views/SignalFlowDashboardViewFactory.h
@@ -67,7 +67,7 @@ namespace UE::Audio::Insights
 		TSharedRef<SWidget> CreateToolbarToggleControls();
 		TSharedRef<SWidget> CreateToolbarButtonActions();
 
-		FReply ScrollGraphViewToSelectedNode();
+		FReply ScrollGraphViewToSelectedNode(const bool bResetZoom);
 		TSharedPtr<FSignalFlowDashboardEntry> GetFilteredAudioDeviceEntry() const;
 
 		TSharedRef<SWidget> CreateSelectionPanel();
