[Audio Insights] Signal Flow Node Details part 01: Add toolbar toggle & Node Details section to signal flow nodes. Implement Amp (peak) detail for sound source nodes.

#Synced-CL 50041425

diff --git a/Engine/Plugins/AudioInsights/Source/AudioInsights/Private/Providers/SignalFlowTraceProvider.cpp b/Engine/Plugins/AudioInsights/Source/AudioInsights/Private/Providers/SignalFlowTraceProvider.cpp
--- a/Engine/Plugins/AudioInsights/Source/AudioInsights/Private/Providers/SignalFlowTraceProvider.cpp
+++ b/Engine/Plugins/AudioInsights/Source/AudioInsights/Private/Providers/SignalFlowTraceProvider.cpp
@@ -1,6 +1,7 @@
 // Copyright Epic Games, Inc. All Rights Reserved.
 #include "Providers/SignalFlowTraceProvider.h"
 
+#include "Async/ParallelFor.h"
 #include "AudioInsightsModule.h"
 #include "Cache/AudioInsightsCacheManager.h"
 
@@ -84,6 +85,8 @@ namespace UE::Audio::Insights
 				Builder.RouteEvent(RouteId_SoundToBusSend, "Audio", "SoundToBusSend");
 				Builder.RouteEvent(RouteId_SoundToSubmixSend, "Audio", "SoundToSubmixSend");
 				Builder.RouteEvent(RouteId_SubmixToSubmixSend, "Audio", "SubmixToSubmixSend");
+
+				Builder.RouteEvent(RouteId_SoundSourceEnvelope, "Audio", "SoundSourceEnvelope");
 			}
 
 			virtual bool OnHandleEvent(uint16 RouteId, EStyle Style, const FOnEventContext& Context) override
@@ -145,6 +148,12 @@ namespace UE::Audio::Insights
 						break;
 					}
 
+					case RouteId_SoundSourceEnvelope:
+					{
+						AddMessageToQueue<FSoundEnvelopeMessage>(Context, Messages.SoundEnvelopeMessages);
+						break;
+					}
+
 					// The second set of messages here are unique to signal flow.
 					// Make sure these are added to the cache.
 					case RouteID_SoundWaveStop:
@@ -209,7 +218,9 @@ namespace UE::Audio::Insights
 
 				RouteId_SoundToBusSend,
 				RouteId_SoundToSubmixSend,
-				RouteId_SubmixToSubmixSend
+				RouteId_SubmixToSubmixSend,
+
+				RouteId_SoundSourceEnvelope
 			};
 
 			TraceServices::IAnalysisSession& Session;
@@ -218,6 +229,13 @@ namespace UE::Audio::Insights
 		return new FSignalFlowTraceAnalyzer(AsShared(), InSession);
 	}
 
+	void FSignalFlowTraceProvider::Reset()
+	{
+		TDeviceDataMapTraceProvider<FSignalFlowEntryKey, TSharedPtr<FSignalFlowDashboardEntry>>::Reset();
+
+		OnResetGraph.Broadcast();
+	}
+
 	bool FSignalFlowTraceProvider::ProcessMessages()
 	{
 		auto GetOwnerEntry = [this](const FSoundStopMessage& Msg)
@@ -240,6 +258,11 @@ namespace UE::Audio::Insights
 			return FindDeviceEntry(Msg.DeviceId, { Msg.DeviceId, ESignalFlowEntryType::SoundSource, Msg.WaveInstancePlayOrder });
 		};
 
+		auto GetSoundwaveEntryFromParamMessage = [this](const FSoundParameterMessage& Msg)
+		{
+			return FindDeviceEntry(Msg.DeviceId, { Msg.DeviceId, ESignalFlowEntryType::SoundSource, Msg.WaveInstancePlayOrder });
+		};
+
 		auto GetEntryFromSendMessage = [&GetSoundwaveEntry](const FSignalFlowSendMessageBase& Msg)
 		{
 			return GetSoundwaveEntry(Msg.DeviceId, Msg.SenderID);
@@ -408,6 +431,15 @@ namespace UE::Audio::Insights
 			bForceGraphRefresh = true;
 		}
 
+		ProcessMessageQueue<FSoundEnvelopeMessage>(TraceMessages.SoundEnvelopeMessages, GetSoundwaveEntryFromParamMessage,
+		[this](const FSoundEnvelopeMessage& Msg, TSharedPtr<FSignalFlowDashboardEntry>* OutEntry)
+		{
+			if (OutEntry && OutEntry->IsValid())
+			{
+				(*OutEntry)->Amplitude = Msg.Envelope;
+			}
+		});
+
 		if (bForceGraphRefresh)
 		{
 			OnRequestGraphRefresh.Broadcast();
@@ -570,6 +602,8 @@ namespace UE::Audio::Insights
 			RemoveInputs(RemovedKeys);
 			bForceGraphRefresh = true;
 		}
+
+		CollectParamsForTimestamp(CacheManager, ActiveEntries, TimeMarker);
 				
 		if (bForceGraphRefresh)
 		{
@@ -634,6 +668,42 @@ namespace UE::Audio::Insights
 		}
 	}
 
+	void FSignalFlowTraceProvider::CollectParamsForTimestamp(const FAudioInsightsCacheManager& CacheManager, const TSet<FSignalFlowEntryKey>& ActiveEntries, const double TimeMarker)
+	{
+		const TArrayView<const FSignalFlowEntryKey> ActiveEntriesView = ActiveEntries.ArrayView();
+
+		ParallelFor(ActiveEntriesView.Num(), [&CacheManager, TimeMarker, &ActiveEntriesView, this](const int32 Index)
+		{
+			const FSignalFlowEntryKey& EntryKey = ActiveEntriesView[Index];
+			if (EntryKey.EntryType != ESignalFlowEntryType::SoundSource
+				&& EntryKey.EntryType != ESignalFlowEntryType::AudioBus
+				&& EntryKey.EntryType != ESignalFlowEntryType::Submix)
+			{
+				return;
+			}
+
+			TSharedPtr<FSignalFlowDashboardEntry>* Entry = FindDeviceEntry(EntryKey.DeviceID, EntryKey);
+			if (Entry == nullptr || !Entry->IsValid())
+			{
+				return;
+			}
+
+			if (EntryKey.EntryType == ESignalFlowEntryType::SoundSource)
+			{
+				auto MatchingParamPredicate = [&EntryKey](const FSoundParameterMessage& Msg)
+				{
+					return Msg.WaveInstancePlayOrder == EntryKey.EntryID;
+				};
+
+				const FSoundEnvelopeMessage* EnvelopeMessage = CacheManager.FindClosestMessageByPredicate<FSoundEnvelopeMessage>(SoundMessageNames::EnvelopeParam, TimeMarker, MatchingParamPredicate);
+				if (EnvelopeMessage)
+				{
+					(*Entry)->Amplitude = EnvelopeMessage->Envelope;
+				}
+			}
+		});
+	}
+
 	void FSignalFlowTraceProvider::OnTimeControlMethodReset()
 	{
 		Reset();
diff --git a/Engine/Plugins/AudioInsights/Source/AudioInsights/Private/Providers/SignalFlowTraceProvider.h b/Engine/Plugins/AudioInsights/Source/AudioInsights/Private/Providers/SignalFlowTraceProvider.h
--- a/Engine/Plugins/AudioInsights/Source/AudioInsights/Private/Providers/SignalFlowTraceProvider.h
+++ b/Engine/Plugins/AudioInsights/Source/AudioInsights/Private/Providers/SignalFlowTraceProvider.h
@@ -16,14 +16,20 @@ namespace UE::Audio::Insights
 		static FName GetName_Static();
 		
 		virtual UE::Trace::IAnalyzer* ConstructAnalyzer(TraceServices::IAnalysisSession& InSession) override;
+		virtual void Reset() override;
 
 		DECLARE_MULTICAST_DELEGATE(FOnRequestGraphRefresh);
 		FOnRequestGraphRefresh OnRequestGraphRefresh;
 
+		DECLARE_MULTICAST_DELEGATE(FOnResetGraph);
+		FOnResetGraph OnResetGraph;
+
 	private:
 		virtual bool ProcessMessages() override;
 		virtual void OnTimingViewTimeMarkerChanged(double TimeMarker) override;
+
 		void CollectSendInfoForTimestamp(const FSignalFlowEntryKey& EntryKey, TSharedPtr<FSignalFlowDashboardEntry>* Entry, const FAudioInsightsCacheManager& CacheManager, const double TimeMarker, const double RangeCheckStart);
+		void CollectParamsForTimestamp(const FAudioInsightsCacheManager& CacheManager, const TSet<FSignalFlowEntryKey>& ActiveEntries, const double TimeMarker);
 
 		virtual void OnTimeControlMethodReset() override;
 
diff --git a/Engine/Plugins/AudioInsights/Source/AudioInsights/Private/Views/SSignalFlowGraph.cpp b/Engine/Plugins/AudioInsights/Source/AudioInsights/Private/Views/SSignalFlowGraph.cpp
--- a/Engine/Plugins/AudioInsights/Source/AudioInsights/Private/Views/SSignalFlowGraph.cpp
+++ b/Engine/Plugins/AudioInsights/Source/AudioInsights/Private/Views/SSignalFlowGraph.cpp
@@ -22,7 +22,6 @@ namespace UE::Audio::Insights
 	namespace SSignalFlowGraphStyle
 	{
 		constexpr float NodeMaxWidth = 200.0f;
-		constexpr float NodeMaxHeight = 50.0f;
 
 		constexpr float NodeRoundedBrushRadius = 12.0f;
 		constexpr float SendLabelRoundedBrushRadius = 8.0f;
@@ -62,10 +61,18 @@ namespace UE::Audio::Insights
 			return NodeMaxWidth + XPaddingBetweenNodes;
 		}
 
-		float GetNodeGraphSlotHeight(const EOrientation GraphOrientation)
+		float GetNodeMaxHeight(const bool bShowNodeDetails)
+		{
+			constexpr float NodeMaxHeight = 50.0f;
+			constexpr float NodeDetailsRowHeight = 20.0f;
+
+			return bShowNodeDetails ? NodeMaxHeight + NodeDetailsRowHeight : NodeMaxHeight;
+		}
+
+		float GetNodeGraphSlotHeight(const EOrientation GraphOrientation, const bool bShowNodeDetails)
 		{
 			const float YPaddingBetweenNodes = GraphOrientation == EOrientation::Orient_Horizontal ? 8.0f : 48.0f;
-			return NodeMaxHeight + YPaddingBetweenNodes;
+			return GetNodeMaxHeight(bShowNodeDetails) + YPaddingBetweenNodes;
 		}
 
 		FSlateColor GetNodeAccentColor(const ESignalFlowEntryType NodeType)
@@ -99,6 +106,67 @@ namespace UE::Audio::Insights
 		{
 			return Orientation == EOrientation::Orient_Horizontal ? FVector2d{ SplineRadius, 0.0f } : FVector2d{ 0.0f, SplineRadius };
 		}
+
+		TOptional<float> GetNodeDetailParamValue(const TSharedPtr<FSignalFlowDashboardEntry>& Entry, const ESignalFlowNodeDetailParam ParamType)
+		{
+			if (!Entry.IsValid())
+			{
+				return NullOpt;
+			}
+
+			switch (ParamType)
+			{
+				case ESignalFlowNodeDetailParam::Amplitude:
+					return Entry->Amplitude;
+
+				default:
+					break;
+			}
+
+			return NullOpt;
+		}
+
+		FSlateColor GetNodeDetailParamColor(const ESignalFlowNodeDetailParam ParamType)
+		{
+			switch (ParamType)
+			{
+				case ESignalFlowNodeDetailParam::Amplitude:
+					return FStyleColors::AccentGreen;
+
+				default:
+					break;
+			}
+
+			return FStyleColors::White;
+		}
+
+		FText GetNodeDetailDisplayName(const ESignalFlowNodeDetailParam ParamType)
+		{
+			switch (ParamType)
+			{
+				case ESignalFlowNodeDetailParam::Amplitude:
+					return LOCTEXT("SignalFlowGraph_NodeDetails_Amplitude", "Amp (peak): {0}");
+
+				default:
+					break;
+			}
+
+			return FText::GetEmpty();
+		}
+
+		const FNumberFormattingOptions* GetNodeDetailNumberFormat(const ESignalFlowNodeDetailParam ParamType)
+		{
+			switch (ParamType)
+			{
+				case ESignalFlowNodeDetailParam::Amplitude:
+					return FSlateStyle::Get().GetLinearVolumeFloatFormat();
+
+				default:
+					break;
+			}
+
+			return FSlateStyle::Get().GetDefaultFloatFormat();
+		}
 	}
 
 	void SSignalFlowGraphNode::Construct(const SSignalFlowGraphNode::FArguments& InArgs)
@@ -111,6 +179,7 @@ namespace UE::Audio::Insights
 		YPos = InArgs._YPos;
 		ZoomScaleFactor = InArgs._ZoomScaleFactor;
 		IsSelected = InArgs._IsSelected;
+		ShowNodeDetails = InArgs._ShowNodeDetails;
 		GraphOrientation = InArgs._GraphOrientation;
 		OnNodeSelected = InArgs._OnNodeSelected;
 
@@ -159,25 +228,38 @@ namespace UE::Audio::Insights
 						.Padding(NodeBorderPadding)
 						.Content()
 						[
-							SNew(SHorizontalBox)
+							SNew(SVerticalBox)
 
-							+ SHorizontalBox::Slot()
-							.AutoWidth()
-							.Padding(0.0f, 0.0f, 4.0f, 0.0f)
+							+ SVerticalBox::Slot()
+							.AutoHeight()
 							[
-								SNew(SImage)
-								.Image(FSlateStyle::Get().CreateIcon(InArgs._Entry->IconName).GetIcon())
-								.ColorAndOpacity(NodeAccentColor)
+								SNew(SHorizontalBox)
+
+								+ SHorizontalBox::Slot()
+								.AutoWidth()
+								.Padding(0.0f, 0.0f, 4.0f, 0.0f)
+								[
+									SNew(SImage)
+									.Image(FSlateStyle::Get().CreateIcon(InArgs._Entry->IconName).GetIcon())
+									.ColorAndOpacity(NodeAccentColor)
+								]
+
+								+ SHorizontalBox::Slot()
+								.FillWidth(1.0f)
+								[
+									SNew(STextBlock)
+									.Text(InArgs._DisplayName)
+									.Justification(ETextJustify::Left)
+									.OverflowPolicy(ETextOverflowPolicy::Ellipsis)
+									.Clipping(EWidgetClipping::ClipToBounds)
+								]
 							]
 
-							+ SHorizontalBox::Slot()
-							.FillWidth(1.0f)
+							+ SVerticalBox::Slot()
+							.Padding(20.0f, 0.0f)
+							.AutoHeight()
 							[
-								SNew(STextBlock)
-								.Text(InArgs._DisplayName)
-								.Justification(ETextJustify::Left)
-								.OverflowPolicy(ETextOverflowPolicy::Ellipsis)
-								.Clipping(EWidgetClipping::ClipToBounds)
+								CreateNodeDetails()
 							]
 						]
 					]
@@ -208,7 +290,7 @@ namespace UE::Audio::Insights
 		using namespace SSignalFlowGraphStyle;
 
 		const FVector2D Size = SCompoundWidget::ComputeDesiredSize(InLayoutScaleMultiplier);
-		const FVector2D ClampedSize = { FMath::Min(Size.X, NodeMaxWidth), FMath::Min(Size.Y, NodeMaxHeight) };
+		const FVector2D ClampedSize = { FMath::Min(Size.X, NodeMaxWidth), FMath::Min(Size.Y, GetNodeMaxHeight(ShowNodeDetails.Get())) };
 		return ClampedSize;
 	}
 
@@ -272,6 +354,57 @@ namespace UE::Audio::Insights
 		return CachedPositionInGraph;
 	}
 
+	TSharedRef<SVerticalBox> SSignalFlowGraphNode::CreateNodeDetails() const
+	{
+		return SNew(SVerticalBox)
+				.Visibility_Lambda([this]() { return ShowNodeDetails.Get() ? EVisibility::Visible : EVisibility::Collapsed; })
+
+				+ SVerticalBox::Slot()
+				.Padding(0.0f, 8.0f, 0.0f, 0.0f)
+				.AutoHeight()
+				[
+					CreateNodeDetailsRow(ESignalFlowNodeDetailParam::Amplitude)
+				];
+	}
+
+	TSharedRef<SHorizontalBox> SSignalFlowGraphNode::CreateNodeDetailsRow(const ESignalFlowNodeDetailParam ParamType) const
+	{
+		using namespace SSignalFlowGraphStyle;
+
+		return SNew(SHorizontalBox)
+				.Visibility_Lambda([this, ParamType]()
+				{
+					return Entry.IsValid() && GetNodeDetailParamValue(Entry, ParamType).IsSet() ? EVisibility::Visible : EVisibility::Collapsed;
+				})
+
+				+ SHorizontalBox::Slot()
+				.Padding(0.0f, 0.0f, 4.0f, 0.0f)
+				.AutoWidth()
+				[
+					SNew(SImage)
+					.Image(FSlateStyle::Get().CreateIcon("AudioInsights.Icon.SignalFlow.Pill").GetIcon())
+					.ColorAndOpacity(GetNodeDetailParamColor(ParamType))
+				]
+
+				+ SHorizontalBox::Slot()
+				.FillWidth(1.0f)
+				[
+					SNew(STextBlock)
+					.Text_Lambda([this, ParamType]()
+					{
+						const TOptional<float> ParamValueOpt = GetNodeDetailParamValue(Entry, ParamType);
+						if (!ParamValueOpt.IsSet())
+						{
+							return FText::GetEmpty();
+						}
+
+						return FText::Format(GetNodeDetailDisplayName(ParamType), FText::AsNumber(ParamValueOpt.GetValue(), GetNodeDetailNumberFormat(ParamType)));
+					})
+					.Font(FCoreStyle::Get().GetFontStyle("SmallFont"))
+					.Justification(ETextJustify::Left)
+				];
+	}
+
 	void SSignalFlowGraphSendLabel::Construct(const SSignalFlowGraphSendLabel::FArguments& InArgs)
 	{
 		using namespace SSignalFlowGraphStyle;
@@ -377,6 +510,7 @@ namespace UE::Audio::Insights
 		ActiveAudioDeviceEntry = InArgs._ActiveAudioDeviceEntry;
 		HorizontalScrollBox = InArgs._HorizontalScrollBox;
 		VerticalScrollBox = InArgs._VerticalScrollBox;
+		ShowNodeDetails = InArgs._ShowNodeDetails;
 		OnNodeSelected = InArgs._OnNodeSelected;
 
 		Clipping = EWidgetClipping::ClipToBounds;
@@ -616,7 +750,7 @@ namespace UE::Audio::Insights
 		float GraphHeight = 0.0f;
 
 		const float NodeGraphSlotWidth = GetNodeGraphSlotWidth(GetGraphOrientation());
-		const float NodeGraphSlotHeight = GetNodeGraphSlotHeight(GetGraphOrientation());
+		const float NodeGraphSlotHeight = GetNodeGraphSlotHeight(GetGraphOrientation(), ShowNodeDetails.Get());
 
 		for (const auto& [Key, NodeWidget] : NodeWidgets)
 		{
@@ -828,7 +962,7 @@ namespace UE::Audio::Insights
 				{
 					const FVector2D NodeSize = (*Node)->GetDesiredSize();
 					const float NodeGraphSlotWidth = GetNodeGraphSlotWidth(GetGraphOrientation());
-					const float NodeGraphSlotHeight = GetNodeGraphSlotHeight(GetGraphOrientation());
+					const float NodeGraphSlotHeight = GetNodeGraphSlotHeight(GetGraphOrientation(), ShowNodeDetails.Get());
 
 					const float XOffset = (NodeGraphSlotWidth - NodeSize.X) / 2.0f;
 					const float YOffset = (NodeGraphSlotHeight - NodeSize.Y) / 2.0f;
@@ -859,6 +993,10 @@ namespace UE::Audio::Insights
 			{
 				return (SelectedItem.IsBound() && SelectedItem.Get().IsValid() && SelectedItem.Get()->GetSignalFlowEntryKey() == EntryKey);
 			})
+			.ShowNodeDetails_Lambda([this]()
+			{
+				return ShowNodeDetails.Get();
+			})
 			.GraphOrientation_Lambda([this]()
 			{
 				return Orientation;
diff --git a/Engine/Plugins/AudioInsights/Source/AudioInsights/Private/Views/SignalFlowDashboardViewFactory.cpp b/Engine/Plugins/AudioInsights/Source/AudioInsights/Private/Views/SignalFlowDashboardViewFactory.cpp
--- a/Engine/Plugins/AudioInsights/Source/AudioInsights/Private/Views/SignalFlowDashboardViewFactory.cpp
+++ b/Engine/Plugins/AudioInsights/Source/AudioInsights/Private/Views/SignalFlowDashboardViewFactory.cpp
@@ -63,6 +63,7 @@ namespace UE::Audio::Insights
 		};
 
 		SignalFlowProvider->OnRequestGraphRefresh.AddRaw(this, &FSignalFlowDashboardViewFactory::OnRequestGraphRefresh);
+		SignalFlowProvider->OnResetGraph.AddRaw(this, &FSignalFlowDashboardViewFactory::Reset);
 	}
 
 	FSignalFlowDashboardViewFactory::~FSignalFlowDashboardViewFactory()
@@ -71,6 +72,7 @@ namespace UE::Audio::Insights
 		if (Provider.IsValid())
 		{
 			Provider->OnRequestGraphRefresh.RemoveAll(this);
+			Provider->OnResetGraph.RemoveAll(this);
 		}
 	}
 
@@ -134,23 +136,34 @@ namespace UE::Audio::Insights
 
 	FReply FSignalFlowDashboardViewFactory::OnKeyDown(const FGeometry& MyGeometry, const FKeyEvent& InKeyEvent)
 	{
-		if (GraphWidgetContents.IsValid())
+		if (GraphWidgetContents.IsValid() && InKeyEvent.IsControlDown())
 		{
-			if (InKeyEvent.GetKey() == EKeys::H && InKeyEvent.IsControlDown())
+			if (InKeyEvent.GetKey() == EKeys::H)
 			{
 				GraphWidgetContents->ToggleOrientation();
 
 				return FReply::Handled();
 			}
-			else if (InKeyEvent.GetKey() == EKeys::F && InKeyEvent.IsControlDown())
+			else if (InKeyEvent.GetKey() == EKeys::F)
 			{
 				return ScrollGraphViewToSelectedNode(true /*bResetZoom*/);
 			}
+			else if (InKeyEvent.GetKey() == EKeys::D)
+			{
+				ToggleShowNodeDetails();
+
+				return FReply::Handled();
+			}
 		}
 
 		return FReply::Unhandled();
 	}
 
+	void FSignalFlowDashboardViewFactory::Reset()
+	{
+		SetNewSelection(nullptr);
+	}
+
 	void FSignalFlowDashboardViewFactory::SRowWidget::Construct(const FArguments& InArgs, const TSharedRef<STableViewBase>& InOwnerTable, TSharedPtr<FSignalFlowDashboardEntry> Entry)
 	{
 		if (!Entry.IsValid())
@@ -232,6 +245,11 @@ namespace UE::Audio::Insights
 			return GraphWidgetContents.IsValid() && GraphWidgetContents->GetGraphOrientation() == Orient_Horizontal ? FSlateColor(FColor::Green) : FSlateColor::UseForeground();
 		};
 
+		TFunction<FSlateColor()> NodeDetailsColorFunc = [this]()
+		{ 
+			return bShowNodeDetails ? FSlateColor(FColor::Green) : FSlateColor::UseForeground();
+		};
+
 		return SNew(SHorizontalBox)
 
 		+ SHorizontalBox::Slot()
@@ -257,6 +275,28 @@ namespace UE::Audio::Insights
 												[]() { return LOCTEXT("AudioDashboard_SignalFlow_HorizontalFlow", "Horizontal Flow"); },
 												ToolbarButtonHeight)
 			]
+		]
+		
+		+ SHorizontalBox::Slot()
+		.AutoWidth()
+		.Padding(ToolbarButtonPaddingX, ToolbarButtonPaddingY)
+		[
+			SNew(SButton)
+			.ButtonStyle(&FAppStyle::Get().GetWidgetStyle<FButtonStyle>(ButtonStyleName))
+			.ToolTipText(LOCTEXT("AudioDashboard_SignalFlow_NodeDetailsTooltip", "Toggles visibility of node details. (Ctrl + d)"))
+			.OnClicked_Lambda([this]()
+			{
+				ToggleShowNodeDetails();
+
+				return FReply::Handled();
+			})
+			[
+				Style.CreateToggleButtonContent("AudioInsights.Icon.SignalFlow.NodeDetails",
+												ButtonTextStyleName,
+												NodeDetailsColorFunc,
+												[]() { return LOCTEXT("AudioDashboard_SignalFlow_NodeDetails", "Node Details"); },
+												ToolbarButtonHeight)
+			]
 		];
 	}
 	
@@ -321,6 +361,12 @@ namespace UE::Audio::Insights
 		return FReply::Unhandled();
 	}
 
+	void FSignalFlowDashboardViewFactory::ToggleShowNodeDetails()
+	{
+		bShowNodeDetails = !bShowNodeDetails;
+		OnRequestGraphRefresh();
+	}
+
 	TSharedPtr<FSignalFlowDashboardEntry> FSignalFlowDashboardViewFactory::GetFilteredAudioDeviceEntry() const
 	{
 		const TSharedPtr<const FSignalFlowTraceProvider> Provider = FindProvider<const FSignalFlowTraceProvider>();
@@ -477,6 +523,10 @@ namespace UE::Audio::Insights
 
 						OnSelectionChanged(SelectedNodeEntry, ESelectInfo::OnMouseClick);
 					})
+					.ShowNodeDetails_Lambda([this]()
+					{
+						return bShowNodeDetails;
+					})
 				]
 			]
 		];
@@ -533,13 +583,19 @@ namespace UE::Audio::Insights
 
 	void FSignalFlowDashboardViewFactory::OnSelectionChanged(TSharedPtr<FSignalFlowDashboardEntry> InSelectedItem, ESelectInfo::Type SelectInfo)
 	{
-		if (!ActiveSourcesListView.IsValid() || !ActiveBusesListView.IsValid() || !ActiveSubmixesListView.IsValid())
+		if (SelectInfo == ESelectInfo::Direct)
 		{
 			return;
 		}
 
-		if (SelectInfo == ESelectInfo::Direct)
+		SetNewSelection(InSelectedItem);
+	}
+
+	void FSignalFlowDashboardViewFactory::SetNewSelection(const TSharedPtr<FSignalFlowDashboardEntry>& InSelectedItem)
+	{
+		if (!ActiveSourcesListView.IsValid() || !ActiveBusesListView.IsValid() || !ActiveSubmixesListView.IsValid())
 		{
+			SelectedItem.Reset();
 			return;
 		}
 
diff --git a/Engine/Plugins/AudioInsights/Source/AudioInsights/Public/AudioInsightsStyle.h b/Engine/Plugins/AudioInsights/Source/AudioInsights/Public/AudioInsightsStyle.h
--- a/Engine/Plugins/AudioInsights/Source/AudioInsights/Public/AudioInsightsStyle.h
+++ b/Engine/Plugins/AudioInsights/Source/AudioInsights/Public/AudioInsightsStyle.h
@@ -254,6 +254,7 @@ namespace UE::Audio::Insights
 			Set("AudioInsights.Icon.SignalFlow.NodeDetails",		  new IMAGE_BRUSH_SVG(TEXT("Icons/SignalFlow/node_details"),		 Icon20));
 			Set("AudioInsights.Icon.SignalFlow.CenterView",			  new IMAGE_BRUSH_SVG(TEXT("Icons/SignalFlow/center_view"),			 Icon20));
 			Set("AudioInsights.Icon.SignalFlow.Settings",			  new IMAGE_BRUSH_SVG(TEXT("Icons/SignalFlow/settings"),			 Icon20));
+			Set("AudioInsights.Icon.SignalFlow.Pill",				  new IMAGE_BRUSH_SVG(TEXT("Icons/SignalFlow/pill"),				 Icon16));
 
 			/* Modulation */
 			// Icons
diff --git a/Engine/Plugins/AudioInsights/Source/AudioInsights/Public/Messages/SignalFlowEntryKey.h b/Engine/Plugins/AudioInsights/Source/AudioInsights/Public/Messages/SignalFlowEntryKey.h
--- a/Engine/Plugins/AudioInsights/Source/AudioInsights/Public/Messages/SignalFlowEntryKey.h
+++ b/Engine/Plugins/AudioInsights/Source/AudioInsights/Public/Messages/SignalFlowEntryKey.h
@@ -23,6 +23,11 @@ namespace UE::Audio::Insights
 		InvalidMax = TNumericLimits<int32>::Max()
 	};
 
+	enum class ESignalFlowNodeDetailParam : uint8
+	{
+		Amplitude = 0u,
+	};
+
 	struct FSignalFlowEntryKey
 	{
 		FSignalFlowEntryKey() = default;
diff --git a/Engine/Plugins/AudioInsights/Source/AudioInsights/Public/Messages/SignalFlowTraceMessages.h b/Engine/Plugins/AudioInsights/Source/AudioInsights/Public/Messages/SignalFlowTraceMessages.h
--- a/Engine/Plugins/AudioInsights/Source/AudioInsights/Public/Messages/SignalFlowTraceMessages.h
+++ b/Engine/Plugins/AudioInsights/Source/AudioInsights/Public/Messages/SignalFlowTraceMessages.h
@@ -100,6 +100,8 @@ namespace UE::Audio::Insights
 		TAnalyzerMessageQueue<FSoundToSubmixSendMessage> SoundToSubmixSendMessages;
 		TAnalyzerMessageQueue<FSubmixToSubmixSendMessage> SubmixToSubmixSendMessages;
 
+		TAnalyzerMessageQueue<FSoundEnvelopeMessage> SoundEnvelopeMessages;
+
 		friend class FSignalFlowTraceProvider;
 	};
 
@@ -135,5 +137,7 @@ namespace UE::Audio::Insights
 
 		TSet<FSignalFlowEntryKey> Inputs;
 		TSortedMap<FSignalFlowEntryKey, FSignalFlowOutputData> Outputs;
+
+		TOptional<float> Amplitude;
 	};
 } // namespace UE::Audio::Insights
diff --git a/Engine/Plugins/AudioInsights/Source/AudioInsights/Public/Views/SSignalFlowGraph.h b/Engine/Plugins/AudioInsights/Source/AudioInsights/Public/Views/SSignalFlowGraph.h
--- a/Engine/Plugins/AudioInsights/Source/AudioInsights/Public/Views/SSignalFlowGraph.h
+++ b/Engine/Plugins/AudioInsights/Source/AudioInsights/Public/Views/SSignalFlowGraph.h
@@ -9,7 +9,9 @@
 #define UE_API AUDIOINSIGHTS_API
 
 class SBorder;
+class SHorizontalBox;
 class SScrollBox;
+class SVerticalBox;
 
 namespace UE::Audio::Insights
 {
@@ -35,6 +37,7 @@ namespace UE::Audio::Insights
 			_YPos(0.0f),
 			_ZoomScaleFactor(1.0f),
 			_IsSelected(false),
+			_ShowNodeDetails(false),
 			_GraphOrientation(EOrientation::Orient_Vertical)
 		{
 		}
@@ -47,6 +50,7 @@ namespace UE::Audio::Insights
 		SLATE_ATTRIBUTE(float, YPos)
 		SLATE_ATTRIBUTE(float, ZoomScaleFactor)
 		SLATE_ATTRIBUTE(bool, IsSelected)
+		SLATE_ATTRIBUTE(bool, ShowNodeDetails)
 		SLATE_ATTRIBUTE(EOrientation, GraphOrientation)
 		SLATE_EVENT(FOnSignalFlowNodeSelected, OnNodeSelected)
 
@@ -73,12 +77,16 @@ namespace UE::Audio::Insights
 		void SetYPos(const float Pos) { YPos = Pos; }
 
 	private:
+		TSharedRef<SVerticalBox> CreateNodeDetails() const;
+		TSharedRef<SHorizontalBox> CreateNodeDetailsRow(const ESignalFlowNodeDetailParam ParamType) const;
+
 		TSharedPtr<FSignalFlowDashboardEntry> Entry;
 		TAttribute<FVector2d> PositionInGraph;
 		TAttribute<float> XPos;
 		TAttribute<float> YPos;
 		TAttribute<float> ZoomScaleFactor;
 		TAttribute<bool> IsSelected;
+		TAttribute<bool> ShowNodeDetails;
 		TAttribute<EOrientation> GraphOrientation;
 
 		FOnSignalFlowNodeSelected OnNodeSelected;
@@ -137,14 +145,18 @@ namespace UE::Audio::Insights
 	public:
 		SLATE_BEGIN_ARGS(SSignalFlowGraph)
 			: _ListSource(nullptr)
+			, _ShowNodeDetails(false)
 		{
 		}
 
 		SLATE_ARGUMENT(TArray<TSharedPtr<ISignalFlowNode>>*, ListSource)
+
 		SLATE_ATTRIBUTE(TSharedPtr<FSignalFlowDashboardEntry>, SelectedItem)
 		SLATE_ATTRIBUTE(TSharedPtr<FSignalFlowDashboardEntry>, ActiveAudioDeviceEntry)
 		SLATE_ATTRIBUTE(TWeakPtr<SScrollBox>, HorizontalScrollBox)
 		SLATE_ATTRIBUTE(TWeakPtr<SScrollBox>, VerticalScrollBox)
+		SLATE_ATTRIBUTE(bool, ShowNodeDetails)
+
 		SLATE_EVENT(FOnSignalFlowNodeSelected, OnNodeSelected)
 
 		SLATE_END_ARGS()
@@ -215,6 +227,7 @@ namespace UE::Audio::Insights
 		TArray<TSharedPtr<ISignalFlowNode>>* ListSource;
 		TAttribute<TSharedPtr<FSignalFlowDashboardEntry>> SelectedItem;
 		TAttribute<TSharedPtr<FSignalFlowDashboardEntry>> ActiveAudioDeviceEntry;
+		TAttribute<bool> ShowNodeDetails;
 
 		TAttribute<TWeakPtr<SScrollBox>> HorizontalScrollBox;
 		TAttribute<TWeakPtr<SScrollBox>> VerticalScrollBox;
diff --git a/Engine/Plugins/AudioInsights/Source/AudioInsights/Public/Views/SignalFlowDashboardViewFactory.h b/Engine/Plugins/AudioInsights/Source/AudioInsights/Public/Views/SignalFlowDashboardViewFactory.h
--- a/Engine/Plugins/AudioInsights/Source/AudioInsights/Public/Views/SignalFlowDashboardViewFactory.h
+++ b/Engine/Plugins/AudioInsights/Source/AudioInsights/Public/Views/SignalFlowDashboardViewFactory.h
@@ -45,6 +45,7 @@ namespace UE::Audio::Insights
 		UE_API virtual TSharedRef<SWidget> MakeWidget(TSharedRef<SDockTab> OwnerTab, const FSpawnTabArgs& SpawnTabArgs) override;
 
 		UE_API FReply OnKeyDown(const FGeometry& MyGeometry, const FKeyEvent& InKeyEvent);
+		UE_API void Reset();
 
 	private:
 		struct SRowWidget : public STableRow<TSharedPtr<FSignalFlowDashboardEntry>>
@@ -68,6 +69,8 @@ namespace UE::Audio::Insights
 		TSharedRef<SWidget> CreateToolbarButtonActions();
 
 		FReply ScrollGraphViewToSelectedNode(const bool bResetZoom);
+		void ToggleShowNodeDetails();
+
 		TSharedPtr<FSignalFlowDashboardEntry> GetFilteredAudioDeviceEntry() const;
 
 		TSharedRef<SWidget> CreateSelectionPanel();
@@ -76,6 +79,7 @@ namespace UE::Audio::Insights
 
 		void SetSearchBoxFilterText(const FText& NewText);
 		void OnSelectionChanged(TSharedPtr<FSignalFlowDashboardEntry> SelectedItem, ESelectInfo::Type SelectInfo);
+		void SetNewSelection(const TSharedPtr<FSignalFlowDashboardEntry>& InSelectedItem);
 
 		virtual void ProcessEntries(EProcessReason Reason) override;
 		void BuildSelectionMenus(const TSortedMap<FSignalFlowEntryKey, TSharedPtr<FSignalFlowDashboardEntry>>& DeviceData, const EProcessReason Reason);
@@ -132,6 +136,7 @@ namespace UE::Audio::Insights
 
 		bool bGraphRequiresRefresh = false;
 		bool bAutoFocusAfterGraphRefresh = false;
+		bool bShowNodeDetails = false;
 	};
 } // namespace UE::Audio::Insights
 
