[Audio Insights] Add r-click mouse drag to move signal flow graph

#Synced-CL 50041425

diff --git a/Engine/Plugins/AudioInsights/Source/AudioInsights/Private/Views/SSignalFlowGraph.cpp b/Engine/Plugins/AudioInsights/Source/AudioInsights/Private/Views/SSignalFlowGraph.cpp
--- a/Engine/Plugins/AudioInsights/Source/AudioInsights/Private/Views/SSignalFlowGraph.cpp
+++ b/Engine/Plugins/AudioInsights/Source/AudioInsights/Private/Views/SSignalFlowGraph.cpp
@@ -682,6 +682,78 @@ namespace UE::Audio::Insights
 		return FReply::Unhandled();
 	}
 
+	FReply SSignalFlowGraph::OnMouseButtonDown(const FGeometry& MyGeometry, const FPointerEvent& MouseEvent)
+	{
+		if (MouseEvent.GetEffectingButton() == EKeys::RightMouseButton)
+		{
+			return FReply::Handled().DetectDrag(SharedThis(this), EKeys::RightMouseButton);
+		}
+
+		return FReply::Unhandled();
+	}
+
+	FReply SSignalFlowGraph::OnDragDetected(const FGeometry& MyGeometry, const FPointerEvent& MouseEvent)
+	{
+		if (MouseEvent.IsMouseButtonDown(EKeys::RightMouseButton))
+		{
+			bIsDraggingToMove = true;
+
+			return FReply::Handled().CaptureMouse(SharedThis(this));
+		}
+
+		return FReply::Unhandled();
+	}
+
+	FReply SSignalFlowGraph::OnMouseMove(const FGeometry& MyGeometry, const FPointerEvent& MouseEvent)
+	{
+		if (bIsDraggingToMove)
+		{
+			if (!MouseEvent.IsMouseButtonDown(EKeys::RightMouseButton))
+			{
+				bIsDraggingToMove = false;
+
+				return FReply::Handled().ReleaseMouseCapture();
+			}
+
+			const FVector2D MoveDelta = MouseEvent.GetCursorDelta();
+
+			const TSharedPtr<SScrollBox> HorizontalScrollBoxPinned = HorizontalScrollBox.Get().Pin();
+			if (HorizontalScrollBoxPinned.IsValid())
+			{
+				const float NewOffsetX = HorizontalScrollBoxPinned->GetScrollOffset() - MoveDelta.X;
+				HorizontalScrollBoxPinned->SetScrollOffset(FMath::Clamp(NewOffsetX, 0.0f, HorizontalScrollBoxPinned->GetScrollOffsetOfEnd()));
+			}
+
+			const TSharedPtr<SScrollBox> VerticalScrollBoxPinned = VerticalScrollBox.Get().Pin();
+			if (VerticalScrollBoxPinned.IsValid())
+			{
+				const float NewOffsetY = VerticalScrollBoxPinned->GetScrollOffset() - MoveDelta.Y;
+				VerticalScrollBoxPinned->SetScrollOffset(FMath::Clamp(NewOffsetY, 0.0f, VerticalScrollBoxPinned->GetScrollOffsetOfEnd()));
+			}
+
+			return FReply::Handled();
+		}
+
+		return FReply::Unhandled();
+	}
+
+	FReply SSignalFlowGraph::OnMouseButtonUp(const FGeometry& MyGeometry, const FPointerEvent& MouseEvent)
+	{
+		if (bIsDraggingToMove && MouseEvent.GetEffectingButton() == EKeys::RightMouseButton)
+		{
+			bIsDraggingToMove = false;
+
+			return FReply::Handled().ReleaseMouseCapture();
+		}
+
+		return FReply::Unhandled();
+	}
+
+	FCursorReply SSignalFlowGraph::OnCursorQuery(const FGeometry& MyGeometry, const FPointerEvent& CursorEvent) const
+	{
+		return bIsDraggingToMove ? FCursorReply::Cursor(EMouseCursor::GrabHandClosed) : FCursorReply::Unhandled();
+	}
+
 	int32 SSignalFlowGraph::OnPaintConnections(const FGeometry& AllottedGeometry, FSlateWindowElementList& OutDrawElements, int32 LayerId) const
 	{
 		using namespace SSignalFlowGraphStyle;
diff --git a/Engine/Plugins/AudioInsights/Source/AudioInsights/Public/Views/SSignalFlowGraph.h b/Engine/Plugins/AudioInsights/Source/AudioInsights/Public/Views/SSignalFlowGraph.h
--- a/Engine/Plugins/AudioInsights/Source/AudioInsights/Public/Views/SSignalFlowGraph.h
+++ b/Engine/Plugins/AudioInsights/Source/AudioInsights/Public/Views/SSignalFlowGraph.h
@@ -178,6 +178,11 @@ namespace UE::Audio::Insights
 		UE_API virtual int32 OnPaint(const FPaintArgs& Args, const FGeometry& AllottedGeometry, const FSlateRect& MyCullingRect, FSlateWindowElementList& OutDrawElements, int32 LayerId, const FWidgetStyle& InWidgetStyle, bool bParentEnabled) const override;
 		
 		UE_API virtual FReply OnMouseWheel(const FGeometry& MyGeometry, const FPointerEvent& MouseEvent) override;
+		UE_API virtual FReply OnMouseButtonDown(const FGeometry& MyGeometry, const FPointerEvent& MouseEvent) override;
+		UE_API virtual FReply OnDragDetected(const FGeometry& MyGeometry, const FPointerEvent& MouseEvent) override;
+		UE_API virtual FReply OnMouseMove(const FGeometry& MyGeometry, const FPointerEvent& MouseEvent) override;
+		UE_API virtual FReply OnMouseButtonUp(const FGeometry& MyGeometry, const FPointerEvent& MouseEvent) override;
+		UE_API virtual FCursorReply OnCursorQuery(const FGeometry& MyGeometry, const FPointerEvent& CursorEvent) const override;
 	protected:
 		UE_API virtual FVector2D ComputeDesiredSize(float InLayoutScaleMultiplier) const override;
 
@@ -239,6 +244,7 @@ namespace UE::Audio::Insights
 
 		bool bNeedsRefresh = false;
 		bool bAnimateZoom = false;
+		bool bIsDraggingToMove = false;
 	};
 } // namespace UE::Audio::Insights
 
