[Audio Insights] Spelling fix

#Synced-CL 50041425

diff --git a/Engine/Plugins/AudioInsights/Source/AudioInsights/Private/Views/SignalFlowDashboardViewFactory.cpp b/Engine/Plugins/AudioInsights/Source/AudioInsights/Private/Views/SignalFlowDashboardViewFactory.cpp
--- a/Engine/Plugins/AudioInsights/Source/AudioInsights/Private/Views/SignalFlowDashboardViewFactory.cpp
+++ b/Engine/Plugins/AudioInsights/Source/AudioInsights/Private/Views/SignalFlowDashboardViewFactory.cpp
@@ -1315,7 +1315,7 @@ namespace UE::Audio::Insights
 			// new largest row, position the nodes around X = 0.0
 			PositionRowAroundOrigin(RowToProcess);
 
-			// Propagate change down the hierachy - moving nodes relative to their inputs
+			// Propagate change down the hierarchy - moving nodes relative to their inputs
 			for (int32 PreviousRowsIndex = RowToProcessIndex - 1; PreviousRowsIndex >= 0; --PreviousRowsIndex)
 			{
 				const TArray<int32> PrevRowOfIndexes = NodeIndexRows[PreviousRowsIndex];
