[Audio Insights] Linux compiler fix

#Synced-CL 50156409

#Synced-CL 50156409

diff --git a/Engine/Plugins/AudioInsights/Source/AudioInsights/Private/Providers/SignalFlowTraceProvider.cpp b/Engine/Plugins/AudioInsights/Source/AudioInsights/Private/Providers/SignalFlowTraceProvider.cpp
--- a/Engine/Plugins/AudioInsights/Source/AudioInsights/Private/Providers/SignalFlowTraceProvider.cpp
+++ b/Engine/Plugins/AudioInsights/Source/AudioInsights/Private/Providers/SignalFlowTraceProvider.cpp
@@ -4,6 +4,8 @@
 #include "Async/ParallelFor.h"
 #include "AudioInsightsModule.h"
 #include "Cache/AudioInsightsCacheManager.h"
+#include "Containers/ArrayView.h"
+#include "Containers/Set.h"
 
 #define LOCTEXT_NAMESPACE "AudioInsights"
 
