[Audio Insights] Signal Flow Node Details part 02: Add remaining node detail param types

#Synced-CL 50156409

diff --git a/Engine/Plugins/AudioInsights/Source/AudioInsights/Private/Messages/SoundTraceMessages.cpp b/Engine/Plugins/AudioInsights/Source/AudioInsights/Private/Messages/SoundTraceMessages.cpp
--- a/Engine/Plugins/AudioInsights/Source/AudioInsights/Private/Messages/SoundTraceMessages.cpp
+++ b/Engine/Plugins/AudioInsights/Source/AudioInsights/Private/Messages/SoundTraceMessages.cpp
@@ -29,6 +29,14 @@ namespace UE::Audio::Insights
 		{
 			return (static_cast<uint64>(ActiveSoundPlayOrder) << 32) + static_cast<uint64>(WaveInstancePlayOrder);
 		}
+
+		bool IsPriorityValueSetToMax(const float PriorityValue)
+		{
+			// Max priority as defined in SoundWave.cpp
+			static constexpr float VolumeWeightedMaxPriority = TNumericLimits<float>::Max() / MAX_VOLUME;
+
+			return PriorityValue >= VolumeWeightedMaxPriority;
+		}
 	};
 
 	FSoundMessageBase::FSoundMessageBase(const Trace::IAnalyzer::FOnEventContext& InContext)
diff --git a/Engine/Plugins/AudioInsights/Source/AudioInsights/Private/Providers/SignalFlowTraceProvider.cpp b/Engine/Plugins/AudioInsights/Source/AudioInsights/Private/Providers/SignalFlowTraceProvider.cpp
--- a/Engine/Plugins/AudioInsights/Source/AudioInsights/Private/Providers/SignalFlowTraceProvider.cpp
+++ b/Engine/Plugins/AudioInsights/Source/AudioInsights/Private/Providers/SignalFlowTraceProvider.cpp
@@ -89,6 +89,16 @@ namespace UE::Audio::Insights
 				Builder.RouteEvent(RouteId_SubmixToSubmixSend, "Audio", "SubmixToSubmixSend");
 
 				Builder.RouteEvent(RouteId_SoundSourceEnvelope, "Audio", "SoundSourceEnvelope");
+				Builder.RouteEvent(RouteId_SoundSourceVolume, "Audio", "SoundVolume");
+				Builder.RouteEvent(RouteId_SoundSourcePitch, "Audio", "SoundPitch");
+				Builder.RouteEvent(RouteId_SoundSourceFilters, "Audio", "SoundSourceFilters");
+				Builder.RouteEvent(RouteId_SoundSourcePriority, "Audio", "SoundPriority");
+				Builder.RouteEvent(RouteId_SoundSourceDistance, "Audio", "SoundDistance");
+				Builder.RouteEvent(RouteId_SoundSourceAttenuation, "Audio", "SoundDistanceAttenuation");
+				Builder.RouteEvent(RouteId_SoundSourceRelativeRenderCost, "Audio", "SoundRelativeRenderCost");
+
+				Builder.RouteEvent(RouteId_AudioBusEnvelope, "Audio", "AudioBusEnvelopeValues");
+				Builder.RouteEvent(RouteId_SubmixEnvelope, "Audio", "SubmixEnvelopeValues");
 			}
 
 			virtual bool OnHandleEvent(uint16 RouteId, EStyle Style, const FOnEventContext& Context) override
@@ -156,6 +166,61 @@ namespace UE::Audio::Insights
 						break;
 					}
 
+					case RouteId_SoundSourceVolume:
+					{
+						AddMessageToQueue<FSoundVolumeMessage>(Context, Messages.SoundVolumeMessages);
+						break;
+					}
+
+					case RouteId_SoundSourcePitch:
+					{
+						AddMessageToQueue<FSoundPitchMessage>(Context, Messages.SoundPitchMessages);
+						break;
+					}
+
+					case RouteId_SoundSourceFilters:
+					{
+						AddMessageToQueue<FSoundLPFFreqMessage>(Context, Messages.SoundLPFFreqMessages);
+						AddMessageToQueue<FSoundHPFFreqMessage>(Context, Messages.SoundHPFFreqMessages);
+						break;
+					}
+
+					case RouteId_SoundSourcePriority:
+					{
+						AddMessageToQueue<FSoundPriorityMessage>(Context, Messages.SoundPriorityMessages);
+						break;
+					}
+
+					case RouteId_SoundSourceDistance:
+					{
+						AddMessageToQueue<FSoundDistanceMessage>(Context, Messages.SoundDistanceMessages);
+						break;
+					}
+
+					case RouteId_SoundSourceAttenuation:
+					{
+						AddMessageToQueue<FSoundDistanceAttenuationMessage>(Context, Messages.SoundAttenuationMessages);
+						break;
+					}
+
+					case RouteId_SoundSourceRelativeRenderCost:
+					{
+						AddMessageToQueue<FSoundRelativeRenderCostMessage>(Context, Messages.SoundRelativeRenderCostMessages);
+						break;
+					}
+
+					case RouteId_AudioBusEnvelope:
+					{
+						AddMessageToQueue<FAudioBusEnvelopeValuesMessage>(Context, Messages.AudioBusEnvelopeMessages);
+						break;
+					}
+
+					case RouteId_SubmixEnvelope:
+					{
+						AddMessageToQueue<FSubmixEnvelopeValuesMessage>(Context, Messages.SubmixEnvelopeMessages);
+						break;
+					}
+
 					// The second set of messages here are unique to signal flow.
 					// Make sure these are added to the cache.
 					case RouteID_SoundWaveStop:
@@ -222,7 +287,16 @@ namespace UE::Audio::Insights
 				RouteId_SoundToSubmixSend,
 				RouteId_SubmixToSubmixSend,
 
-				RouteId_SoundSourceEnvelope
+				RouteId_SoundSourceEnvelope,
+				RouteId_SoundSourceVolume,
+				RouteId_SoundSourcePitch,
+				RouteId_SoundSourceFilters,
+				RouteId_SoundSourcePriority,
+				RouteId_SoundSourceDistance,
+				RouteId_SoundSourceAttenuation,
+				RouteId_SoundSourceRelativeRenderCost,
+				RouteId_AudioBusEnvelope,
+				RouteId_SubmixEnvelope
 			};
 
 			TraceServices::IAnalysisSession& Session;
@@ -442,6 +516,90 @@ namespace UE::Audio::Insights
 			}
 		});
 
+		ProcessMessageQueue<FSoundVolumeMessage>(TraceMessages.SoundVolumeMessages, GetSoundwaveEntryFromParamMessage,
+		[this](const FSoundVolumeMessage& Msg, TSharedPtr<FSignalFlowDashboardEntry>* OutEntry)
+		{
+			if (OutEntry && OutEntry->IsValid())
+			{
+				(*OutEntry)->Volume = Msg.Volume;
+			}
+		});
+
+		ProcessMessageQueue<FSoundPitchMessage>(TraceMessages.SoundPitchMessages, GetSoundwaveEntryFromParamMessage,
+		[this](const FSoundPitchMessage& Msg, TSharedPtr<FSignalFlowDashboardEntry>* OutEntry)
+		{
+			if (OutEntry && OutEntry->IsValid())
+			{
+				(*OutEntry)->Pitch = Msg.Pitch;
+			}
+		});
+
+		ProcessMessageQueue<FSoundLPFFreqMessage>(TraceMessages.SoundLPFFreqMessages, GetSoundwaveEntryFromParamMessage,
+		[this](const FSoundLPFFreqMessage& Msg, TSharedPtr<FSignalFlowDashboardEntry>* OutEntry)
+		{
+			if (OutEntry && OutEntry->IsValid())
+			{
+				(*OutEntry)->LPFFreq = Msg.LPFFrequency;
+			}
+		});
+
+		ProcessMessageQueue<FSoundHPFFreqMessage>(TraceMessages.SoundHPFFreqMessages, GetSoundwaveEntryFromParamMessage,
+		[this](const FSoundHPFFreqMessage& Msg, TSharedPtr<FSignalFlowDashboardEntry>* OutEntry)
+		{
+			if (OutEntry && OutEntry->IsValid())
+			{
+				(*OutEntry)->HPFFreq = Msg.HPFFrequency;
+			}
+		});
+
+		ProcessMessageQueue<FSoundPriorityMessage>(TraceMessages.SoundPriorityMessages, GetSoundwaveEntryFromParamMessage,
+		[this](const FSoundPriorityMessage& Msg, TSharedPtr<FSignalFlowDashboardEntry>* OutEntry)
+		{
+			if (OutEntry && OutEntry->IsValid())
+			{
+				(*OutEntry)->Priority = Msg.Priority;
+			}
+		});
+
+		ProcessMessageQueue<FSoundDistanceMessage>(TraceMessages.SoundDistanceMessages, GetSoundwaveEntryFromParamMessage,
+		[this](const FSoundDistanceMessage& Msg, TSharedPtr<FSignalFlowDashboardEntry>* OutEntry)
+		{
+			if (OutEntry && OutEntry->IsValid())
+			{
+				(*OutEntry)->Distance = Msg.Distance;
+			}
+		});
+
+		ProcessMessageQueue<FSoundDistanceAttenuationMessage>(TraceMessages.SoundAttenuationMessages, GetSoundwaveEntryFromParamMessage,
+		[this](const FSoundDistanceAttenuationMessage& Msg, TSharedPtr<FSignalFlowDashboardEntry>* OutEntry)
+		{
+			if (OutEntry && OutEntry->IsValid())
+			{
+				(*OutEntry)->Attenuation = Msg.DistanceAttenuation;
+			}
+		});
+
+		ProcessMessageQueue<FSoundRelativeRenderCostMessage>(TraceMessages.SoundRelativeRenderCostMessages, GetSoundwaveEntryFromParamMessage,
+		[this](const FSoundRelativeRenderCostMessage& Msg, TSharedPtr<FSignalFlowDashboardEntry>* OutEntry)
+		{
+			if (OutEntry && OutEntry->IsValid())
+			{
+				(*OutEntry)->RelativeRenderCost = Msg.RelativeRenderCost;
+			}
+		});
+
+		ProcessMessageQueue<FAudioBusEnvelopeValuesMessage>(TraceMessages.AudioBusEnvelopeMessages, GetBusEntry,
+		[this](const FAudioBusEnvelopeValuesMessage& Msg, TSharedPtr<FSignalFlowDashboardEntry>* OutEntry)
+		{
+			ProcessMultiChannelEnvelopeMessage(OutEntry, Msg.EnvelopeValues);
+		});
+
+		ProcessMessageQueue<FSubmixEnvelopeValuesMessage>(TraceMessages.SubmixEnvelopeMessages, GetSubmixEntry,
+		[this](const FSubmixEnvelopeValuesMessage& Msg, TSharedPtr<FSignalFlowDashboardEntry>* OutEntry)
+		{
+			ProcessMultiChannelEnvelopeMessage(OutEntry, Msg.EnvelopeValues);
+		});
+
 		if (bForceGraphRefresh)
 		{
 			OnRequestGraphRefresh.Broadcast();
@@ -702,6 +860,70 @@ namespace UE::Audio::Insights
 				{
 					(*Entry)->Amplitude = EnvelopeMessage->Envelope;
 				}
+
+				const FSoundVolumeMessage* VolumeMessage = CacheManager.FindClosestMessageByPredicate<FSoundVolumeMessage>(SoundMessageNames::VolumeParam, TimeMarker, MatchingParamPredicate);
+				if (VolumeMessage)
+				{
+					(*Entry)->Volume = VolumeMessage->Volume;
+				}
+
+				const FSoundPitchMessage* PitchMessage = CacheManager.FindClosestMessageByPredicate<FSoundPitchMessage>(SoundMessageNames::PitchParam, TimeMarker, MatchingParamPredicate);
+				if (PitchMessage)
+				{
+					(*Entry)->Pitch = PitchMessage->Pitch;
+				}
+
+				const FSoundLPFFreqMessage* LPFMessage = CacheManager.FindClosestMessageByPredicate<FSoundLPFFreqMessage>(SoundMessageNames::LPFFreqParam, TimeMarker, MatchingParamPredicate);
+				if (LPFMessage)
+				{
+					(*Entry)->LPFFreq = LPFMessage->LPFFrequency;
+				}
+
+				const FSoundHPFFreqMessage* HPFMessage = CacheManager.FindClosestMessageByPredicate<FSoundHPFFreqMessage>(SoundMessageNames::HPFFreqParam, TimeMarker, MatchingParamPredicate);
+				if (HPFMessage)
+				{
+					(*Entry)->HPFFreq = HPFMessage->HPFFrequency;
+				}
+
+				const FSoundPriorityMessage* PriorityMessage = CacheManager.FindClosestMessageByPredicate<FSoundPriorityMessage>(SoundMessageNames::PriorityParam, TimeMarker, MatchingParamPredicate);
+				if (PriorityMessage)
+				{
+					(*Entry)->Priority = PriorityMessage->Priority;
+				}
+
+				const FSoundDistanceMessage* DistanceMessage = CacheManager.FindClosestMessageByPredicate<FSoundDistanceMessage>(SoundMessageNames::DistanceParam, TimeMarker, MatchingParamPredicate);
+				if (DistanceMessage)
+				{
+					(*Entry)->Distance = DistanceMessage->Distance;
+				}
+
+				const FSoundDistanceAttenuationMessage* AttenuationMessage = CacheManager.FindClosestMessageByPredicate<FSoundDistanceAttenuationMessage>(SoundMessageNames::DistanceAttenuationParam, TimeMarker, MatchingParamPredicate);
+				if (AttenuationMessage)
+				{
+					(*Entry)->Attenuation = AttenuationMessage->DistanceAttenuation;
+				}
+
+				const FSoundRelativeRenderCostMessage* RelativeRenderCostMessage = CacheManager.FindClosestMessageByPredicate<FSoundRelativeRenderCostMessage>(SoundMessageNames::RelativeRenderCostParam, TimeMarker, MatchingParamPredicate);
+				if (RelativeRenderCostMessage)
+				{
+					(*Entry)->RelativeRenderCost = RelativeRenderCostMessage->RelativeRenderCost;
+				}
+			}
+			else if (EntryKey.EntryType == ESignalFlowEntryType::AudioBus)
+			{
+				const FAudioBusEnvelopeValuesMessage* AudioBusEnvelopeMessage = CacheManager.FindClosestMessage<FAudioBusEnvelopeValuesMessage>(AudioBusMessageNames::EnvelopeValues, TimeMarker, EntryKey.EntryID);
+				if (AudioBusEnvelopeMessage)
+				{
+					ProcessMultiChannelEnvelopeMessage(Entry, AudioBusEnvelopeMessage->EnvelopeValues);
+				}
+			}
+			else if (EntryKey.EntryType == ESignalFlowEntryType::Submix)
+			{
+				const FSubmixEnvelopeValuesMessage* SubmixEnvelopeMessage = CacheManager.FindClosestMessage<FSubmixEnvelopeValuesMessage>(SubmixMessageNames::EnvelopeValues, TimeMarker, EntryKey.EntryID);
+				if (SubmixEnvelopeMessage)
+				{
+					ProcessMultiChannelEnvelopeMessage(Entry, SubmixEnvelopeMessage->EnvelopeValues);
+				}
 			}
 		});
 	}
@@ -876,6 +1098,25 @@ namespace UE::Audio::Insights
 		}
 	}
 
+	void FSignalFlowTraceProvider::ProcessMultiChannelEnvelopeMessage(TSharedPtr<FSignalFlowDashboardEntry>* OutEntry, const TArray<float>& EnvelopeValues)
+	{
+		if (OutEntry == nullptr || !OutEntry->IsValid() || EnvelopeValues.IsEmpty())
+		{
+			return;
+		}
+
+		float MaxPeakAmp = TNumericLimits<float>::Lowest();
+		for (const float AmpValue : EnvelopeValues)
+		{
+			MaxPeakAmp = FMath::Max(MaxPeakAmp, AmpValue);
+		}
+
+		if (MaxPeakAmp > TNumericLimits<float>::Lowest())
+		{
+			(*OutEntry)->Amplitude = MaxPeakAmp;
+		}
+	}
+
 	void FSignalFlowTraceProvider::RemoveInputs(const TSet<FSignalFlowEntryKey>& RemovedKeys)
 	{
 		for (auto& [DeviceID, Entries] : DeviceDataMap)
diff --git a/Engine/Plugins/AudioInsights/Source/AudioInsights/Private/Providers/SignalFlowTraceProvider.h b/Engine/Plugins/AudioInsights/Source/AudioInsights/Private/Providers/SignalFlowTraceProvider.h
--- a/Engine/Plugins/AudioInsights/Source/AudioInsights/Private/Providers/SignalFlowTraceProvider.h
+++ b/Engine/Plugins/AudioInsights/Source/AudioInsights/Private/Providers/SignalFlowTraceProvider.h
@@ -46,6 +46,7 @@ namespace UE::Audio::Insights
 		void ProcessStopMessage(const ::Audio::FDeviceId DeviceID, const double MsgTimestamp, TSharedPtr<FSignalFlowDashboardEntry>* OutEntry, TSet<FSignalFlowEntryKey>& OutRemovedKeys);
 
 		void ProcessOutputSendMessage(TSharedPtr<FSignalFlowDashboardEntry>* OutEntry, const FSignalFlowEntryKey& OutputEntryKey, const float SendLevel);
+		void ProcessMultiChannelEnvelopeMessage(TSharedPtr<FSignalFlowDashboardEntry>* OutEntry, const TArray<float>& EnvelopeValues);
 
 		void LinkInput(const FSignalFlowDashboardEntry& Entry, const FSignalFlowEntryKey& OutputKey, const ::Audio::FDeviceId DeviceID);
 		void RemoveInputs(const TSet<FSignalFlowEntryKey>& RemovedKeys);
diff --git a/Engine/Plugins/AudioInsights/Source/AudioInsights/Private/Views/SSignalFlowGraph.cpp b/Engine/Plugins/AudioInsights/Source/AudioInsights/Private/Views/SSignalFlowGraph.cpp
--- a/Engine/Plugins/AudioInsights/Source/AudioInsights/Private/Views/SSignalFlowGraph.cpp
+++ b/Engine/Plugins/AudioInsights/Source/AudioInsights/Private/Views/SSignalFlowGraph.cpp
@@ -21,7 +21,7 @@ namespace UE::Audio::Insights
 {
 	namespace SSignalFlowGraphStyle
 	{
-		constexpr float NodeMaxWidth = 200.0f;
+		constexpr float NodeMaxWidth = 240.0f;
 
 		constexpr float NodeRoundedBrushRadius = 12.0f;
 		constexpr float SendLabelRoundedBrushRadius = 8.0f;
@@ -64,9 +64,10 @@ namespace UE::Audio::Insights
 		float GetNodeMaxHeight(const bool bShowNodeDetails)
 		{
 			constexpr float NodeMaxHeight = 50.0f;
-			constexpr float NodeDetailsRowHeight = 20.0f;
+			constexpr float NodeDetailsRowHeight = 24.0f;
+			constexpr uint32 MaxNumNodeDetails = 9; // TODO - Scale this with the number enabled in settings when UE-358203 is implemented
 
-			return bShowNodeDetails ? NodeMaxHeight + NodeDetailsRowHeight : NodeMaxHeight;
+			return bShowNodeDetails ? NodeMaxHeight + (NodeDetailsRowHeight * MaxNumNodeDetails) : NodeMaxHeight;
 		}
 
 		float GetNodeGraphSlotHeight(const EOrientation GraphOrientation, const bool bShowNodeDetails)
@@ -119,6 +120,30 @@ namespace UE::Audio::Insights
 				case ESignalFlowNodeDetailParam::Amplitude:
 					return Entry->Amplitude;
 
+				case ESignalFlowNodeDetailParam::Volume:
+					return Entry->Volume;
+
+				case ESignalFlowNodeDetailParam::Pitch:
+					return Entry->Pitch;
+
+				case ESignalFlowNodeDetailParam::LPFFreq:
+					return Entry->LPFFreq;
+
+				case ESignalFlowNodeDetailParam::HPFFreq:
+					return Entry->HPFFreq;
+
+				case ESignalFlowNodeDetailParam::Priority:
+					return Entry->Priority;
+
+				case ESignalFlowNodeDetailParam::Distance:
+					return Entry->Distance;
+
+				case ESignalFlowNodeDetailParam::Attenuation:
+					return Entry->Attenuation;
+
+				case ESignalFlowNodeDetailParam::RelativeRenderCost:
+					return Entry->RelativeRenderCost;
+
 				default:
 					break;
 			}
@@ -133,6 +158,30 @@ namespace UE::Audio::Insights
 				case ESignalFlowNodeDetailParam::Amplitude:
 					return FStyleColors::AccentGreen;
 
+				case ESignalFlowNodeDetailParam::Volume:
+					return FStyleColors::AccentYellow;
+
+				case ESignalFlowNodeDetailParam::Pitch:
+					return FStyleColors::AccentOrange;
+
+				case ESignalFlowNodeDetailParam::LPFFreq:
+					return FStyleColors::AccentBrown;
+
+				case ESignalFlowNodeDetailParam::HPFFreq:
+					return FStyleColors::AccentRed;
+
+				case ESignalFlowNodeDetailParam::Priority:
+					return FStyleColors::AccentPink;
+
+				case ESignalFlowNodeDetailParam::Distance:
+					return FStyleColors::AccentPurple;
+
+				case ESignalFlowNodeDetailParam::Attenuation:
+					return FStyleColors::AccentBlue;
+
+				case ESignalFlowNodeDetailParam::RelativeRenderCost:
+					return FStyleColors::AccentGray;
+
 				default:
 					break;
 			}
@@ -147,6 +196,30 @@ namespace UE::Audio::Insights
 				case ESignalFlowNodeDetailParam::Amplitude:
 					return LOCTEXT("SignalFlowGraph_NodeDetails_Amplitude", "Amp (peak): {0}");
 
+				case ESignalFlowNodeDetailParam::Volume:
+					return LOCTEXT("SignalFlowGraph_NodeDetails_Volume", "Volume: {0}");
+
+				case ESignalFlowNodeDetailParam::Pitch:
+					return LOCTEXT("SignalFlowGraph_NodeDetails_Pitch", "Pitch: {0}");
+
+				case ESignalFlowNodeDetailParam::LPFFreq:
+					return LOCTEXT("SignalFlowGraph_NodeDetails_LPFFreq", "LPF Freq (Hz): {0}");
+
+				case ESignalFlowNodeDetailParam::HPFFreq:
+					return LOCTEXT("SignalFlowGraph_NodeDetails_HPFFreq", "HPF Freq (Hz): {0}");
+
+				case ESignalFlowNodeDetailParam::Priority:
+					return LOCTEXT("SignalFlowGraph_NodeDetails_Priority", "Priority: {0}");
+
+				case ESignalFlowNodeDetailParam::Distance:
+					return LOCTEXT("SignalFlowGraph_NodeDetails_Distance", "Distance: {0}");
+
+				case ESignalFlowNodeDetailParam::Attenuation:
+					return LOCTEXT("SignalFlowGraph_NodeDetails_Attenuation", "Distance/Occlusion: {0}");
+
+				case ESignalFlowNodeDetailParam::RelativeRenderCost:
+					return LOCTEXT("SignalFlowGraph_NodeDetails_RekativeRenderCost", "Rel. Render Cost: {0}");
+
 				default:
 					break;
 			}
@@ -161,6 +234,30 @@ namespace UE::Audio::Insights
 				case ESignalFlowNodeDetailParam::Amplitude:
 					return FSlateStyle::Get().GetLinearVolumeFloatFormat();
 
+				case ESignalFlowNodeDetailParam::Volume:
+					return FSlateStyle::Get().GetLinearVolumeFloatFormat();
+
+				case ESignalFlowNodeDetailParam::Pitch:
+					return FSlateStyle::Get().GetPitchFloatFormat();
+
+				case ESignalFlowNodeDetailParam::LPFFreq:
+					return FSlateStyle::Get().GetFreqFloatFormat();
+
+				case ESignalFlowNodeDetailParam::HPFFreq:
+					return FSlateStyle::Get().GetFreqFloatFormat();
+
+				case ESignalFlowNodeDetailParam::Priority:
+					return FSlateStyle::Get().GetLinearVolumeFloatFormat();
+
+				case ESignalFlowNodeDetailParam::Distance:
+					return FSlateStyle::Get().GetDefaultFloatFormat();
+
+				case ESignalFlowNodeDetailParam::Attenuation:
+					return FSlateStyle::Get().GetLinearVolumeFloatFormat();
+
+				case ESignalFlowNodeDetailParam::RelativeRenderCost:
+					return FSlateStyle::Get().GetRelativeRenderCostFormat();
+
 				default:
 					break;
 			}
@@ -189,7 +286,6 @@ namespace UE::Audio::Insights
 		}
 
 		NodeEntryKey = Entry->GetSignalFlowEntryKey();
-		const FSlateColor NodeAccentColor = GetNodeAccentColor(NodeEntryKey.EntryType);
 
 		ChildSlot
 		[
@@ -213,7 +309,7 @@ namespace UE::Audio::Insights
 					.BorderImage(&GetNodeRoundedBoxBrush())
 					.HAlign(HAlign_Fill)
 					.VAlign(VAlign_Fill)
-					.BorderBackgroundColor(NodeAccentColor)
+					.BorderBackgroundColor(GetNodeAccentColor(NodeEntryKey.EntryType))
 					.Padding(NodeBorderThickness)
 					.Content()
 					[
@@ -233,26 +329,7 @@ namespace UE::Audio::Insights
 							+ SVerticalBox::Slot()
 							.AutoHeight()
 							[
-								SNew(SHorizontalBox)
-
-								+ SHorizontalBox::Slot()
-								.AutoWidth()
-								.Padding(0.0f, 0.0f, 4.0f, 0.0f)
-								[
-									SNew(SImage)
-									.Image(FSlateStyle::Get().CreateIcon(InArgs._Entry->IconName).GetIcon())
-									.ColorAndOpacity(NodeAccentColor)
-								]
-
-								+ SHorizontalBox::Slot()
-								.FillWidth(1.0f)
-								[
-									SNew(STextBlock)
-									.Text(InArgs._DisplayName)
-									.Justification(ETextJustify::Left)
-									.OverflowPolicy(ETextOverflowPolicy::Ellipsis)
-									.Clipping(EWidgetClipping::ClipToBounds)
-								]
+								CreateNodeName(InArgs)
 							]
 
 							+ SVerticalBox::Slot()
@@ -354,6 +431,32 @@ namespace UE::Audio::Insights
 		return CachedPositionInGraph;
 	}
 
+	TSharedRef<SHorizontalBox> SSignalFlowGraphNode::CreateNodeName(const SSignalFlowGraphNode::FArguments& InArgs) const
+	{
+		using namespace SSignalFlowGraphStyle;
+
+		return SNew(SHorizontalBox)
+
+			+ SHorizontalBox::Slot()
+			.AutoWidth()
+			.Padding(0.0f, 0.0f, 4.0f, 0.0f)
+			[
+				SNew(SImage)
+				.Image(FSlateStyle::Get().CreateIcon(InArgs._Entry->IconName).GetIcon())
+				.ColorAndOpacity(GetNodeAccentColor(NodeEntryKey.EntryType))
+			]
+
+			+ SHorizontalBox::Slot()
+			.FillWidth(1.0f)
+			[
+				SNew(STextBlock)
+				.Text(InArgs._DisplayName)
+				.Justification(ETextJustify::Left)
+				.OverflowPolicy(ETextOverflowPolicy::Ellipsis)
+				.Clipping(EWidgetClipping::ClipToBounds)
+			];
+	}
+
 	TSharedRef<SVerticalBox> SSignalFlowGraphNode::CreateNodeDetails() const
 	{
 		return SNew(SVerticalBox)
@@ -364,6 +467,62 @@ namespace UE::Audio::Insights
 				.AutoHeight()
 				[
 					CreateNodeDetailsRow(ESignalFlowNodeDetailParam::Amplitude)
+				]
+			
+				+ SVerticalBox::Slot()
+				.Padding(0.0f, 8.0f, 0.0f, 0.0f)
+				.AutoHeight()
+				[
+					CreateNodeDetailsRow(ESignalFlowNodeDetailParam::Volume)
+				]
+
+				+ SVerticalBox::Slot()
+				.Padding(0.0f, 8.0f, 0.0f, 0.0f)
+				.AutoHeight()
+				[
+					CreateNodeDetailsRow(ESignalFlowNodeDetailParam::Pitch)
+				]
+
+				+ SVerticalBox::Slot()
+				.Padding(0.0f, 8.0f, 0.0f, 0.0f)
+				.AutoHeight()
+				[
+					CreateNodeDetailsRow(ESignalFlowNodeDetailParam::LPFFreq)
+				]
+			
+				+ SVerticalBox::Slot()
+				.Padding(0.0f, 8.0f, 0.0f, 0.0f)
+				.AutoHeight()
+				[
+					CreateNodeDetailsRow(ESignalFlowNodeDetailParam::HPFFreq)
+				]
+
+				+ SVerticalBox::Slot()
+				.Padding(0.0f, 8.0f, 0.0f, 0.0f)
+				.AutoHeight()
+				[
+					CreateNodeDetailsRow(ESignalFlowNodeDetailParam::Priority)
+				]
+
+				+ SVerticalBox::Slot()
+				.Padding(0.0f, 8.0f, 0.0f, 0.0f)
+				.AutoHeight()
+				[
+					CreateNodeDetailsRow(ESignalFlowNodeDetailParam::Distance)
+				]
+
+				+ SVerticalBox::Slot()
+				.Padding(0.0f, 8.0f, 0.0f, 0.0f)
+				.AutoHeight()
+				[
+					CreateNodeDetailsRow(ESignalFlowNodeDetailParam::Attenuation)
+				]
+
+				+ SVerticalBox::Slot()
+				.Padding(0.0f, 8.0f, 0.0f, 0.0f)
+				.AutoHeight()
+				[
+					CreateNodeDetailsRow(ESignalFlowNodeDetailParam::RelativeRenderCost)
 				];
 	}
 
@@ -398,6 +557,11 @@ namespace UE::Audio::Insights
 							return FText::GetEmpty();
 						}
 
+						if (ParamType == ESignalFlowNodeDetailParam::Priority && SoundMessageUtils::IsPriorityValueSetToMax(ParamValueOpt.GetValue()))
+						{
+							return FText::Format(GetNodeDetailDisplayName(ParamType), LOCTEXT("SignalFlowGraph_NodeDetails_Max", "MAX"));
+						}
+
 						return FText::Format(GetNodeDetailDisplayName(ParamType), FText::AsNumber(ParamValueOpt.GetValue(), GetNodeDetailNumberFormat(ParamType)));
 					})
 					.Font(FCoreStyle::Get().GetFontStyle("SmallFont"))
diff --git a/Engine/Plugins/AudioInsights/Source/AudioInsights/Private/Views/SoundDashboardViewFactory.cpp b/Engine/Plugins/AudioInsights/Source/AudioInsights/Private/Views/SoundDashboardViewFactory.cpp
--- a/Engine/Plugins/AudioInsights/Source/AudioInsights/Private/Views/SoundDashboardViewFactory.cpp
+++ b/Engine/Plugins/AudioInsights/Source/AudioInsights/Private/Views/SoundDashboardViewFactory.cpp
@@ -1930,10 +1930,7 @@ namespace UE::Audio::Insights
 
 						const float PriorityValue = SoundDashboardEntry.PriorityDataPoint;
 
-						// Max priority as defined in SoundWave.cpp
-						static constexpr float VolumeWeightedMaxPriority = TNumericLimits<float>::Max() / MAX_VOLUME;
-
-						return PriorityValue >= VolumeWeightedMaxPriority 
+						return SoundMessageUtils::IsPriorityValueSetToMax(PriorityValue)
 							? LOCTEXT("AudioDashboard_Sounds_Max", "MAX") 
 							: FText::AsNumber(PriorityValue, FSlateStyle::Get().GetLinearVolumeFloatFormat());
 					}
diff --git a/Engine/Plugins/AudioInsights/Source/AudioInsights/Public/Messages/SignalFlowEntryKey.h b/Engine/Plugins/AudioInsights/Source/AudioInsights/Public/Messages/SignalFlowEntryKey.h
--- a/Engine/Plugins/AudioInsights/Source/AudioInsights/Public/Messages/SignalFlowEntryKey.h
+++ b/Engine/Plugins/AudioInsights/Source/AudioInsights/Public/Messages/SignalFlowEntryKey.h
@@ -26,6 +26,14 @@ namespace UE::Audio::Insights
 	enum class ESignalFlowNodeDetailParam : uint8
 	{
 		Amplitude = 0u,
+		Volume,
+		Pitch,
+		LPFFreq,
+		HPFFreq,
+		Priority,
+		Distance,
+		Attenuation,
+		RelativeRenderCost
 	};
 
 	struct FSignalFlowEntryKey
diff --git a/Engine/Plugins/AudioInsights/Source/AudioInsights/Public/Messages/SignalFlowTraceMessages.h b/Engine/Plugins/AudioInsights/Source/AudioInsights/Public/Messages/SignalFlowTraceMessages.h
--- a/Engine/Plugins/AudioInsights/Source/AudioInsights/Public/Messages/SignalFlowTraceMessages.h
+++ b/Engine/Plugins/AudioInsights/Source/AudioInsights/Public/Messages/SignalFlowTraceMessages.h
@@ -101,6 +101,17 @@ namespace UE::Audio::Insights
 		TAnalyzerMessageQueue<FSubmixToSubmixSendMessage> SubmixToSubmixSendMessages;
 
 		TAnalyzerMessageQueue<FSoundEnvelopeMessage> SoundEnvelopeMessages;
+		TAnalyzerMessageQueue<FSoundVolumeMessage> SoundVolumeMessages;
+		TAnalyzerMessageQueue<FSoundPitchMessage> SoundPitchMessages;
+		TAnalyzerMessageQueue<FSoundLPFFreqMessage> SoundLPFFreqMessages;
+		TAnalyzerMessageQueue<FSoundHPFFreqMessage> SoundHPFFreqMessages;
+		TAnalyzerMessageQueue<FSoundPriorityMessage> SoundPriorityMessages;
+		TAnalyzerMessageQueue<FSoundDistanceMessage> SoundDistanceMessages;
+		TAnalyzerMessageQueue<FSoundDistanceAttenuationMessage> SoundAttenuationMessages;
+		TAnalyzerMessageQueue<FSoundRelativeRenderCostMessage> SoundRelativeRenderCostMessages;
+
+		TAnalyzerMessageQueue<FAudioBusEnvelopeValuesMessage> AudioBusEnvelopeMessages;
+		TAnalyzerMessageQueue<FSubmixEnvelopeValuesMessage> SubmixEnvelopeMessages;
 
 		friend class FSignalFlowTraceProvider;
 	};
@@ -139,5 +150,13 @@ namespace UE::Audio::Insights
 		TSortedMap<FSignalFlowEntryKey, FSignalFlowOutputData> Outputs;
 
 		TOptional<float> Amplitude;
+		TOptional<float> Volume;
+		TOptional<float> Pitch;
+		TOptional<float> LPFFreq;
+		TOptional<float> HPFFreq;
+		TOptional<float> Priority;
+		TOptional<float> Distance;
+		TOptional<float> Attenuation;
+		TOptional<float> RelativeRenderCost;
 	};
 } // namespace UE::Audio::Insights
diff --git a/Engine/Plugins/AudioInsights/Source/AudioInsights/Public/Messages/SoundTraceMessages.h b/Engine/Plugins/AudioInsights/Source/AudioInsights/Public/Messages/SoundTraceMessages.h
--- a/Engine/Plugins/AudioInsights/Source/AudioInsights/Public/Messages/SoundTraceMessages.h
+++ b/Engine/Plugins/AudioInsights/Source/AudioInsights/Public/Messages/SoundTraceMessages.h
@@ -61,6 +61,7 @@ namespace UE::Audio::Insights
 	namespace SoundMessageUtils
 	{
 		uint64 GeneratePlayOrderUniqueID(const uint32 ActiveSoundPlayOrder, const uint32 WaveInstancePlayOrder);
+		bool IsPriorityValueSetToMax(const float PriorityValue);
 	};
 
 	// Trace messages
diff --git a/Engine/Plugins/AudioInsights/Source/AudioInsights/Public/Views/SSignalFlowGraph.h b/Engine/Plugins/AudioInsights/Source/AudioInsights/Public/Views/SSignalFlowGraph.h
--- a/Engine/Plugins/AudioInsights/Source/AudioInsights/Public/Views/SSignalFlowGraph.h
+++ b/Engine/Plugins/AudioInsights/Source/AudioInsights/Public/Views/SSignalFlowGraph.h
@@ -77,6 +77,7 @@ namespace UE::Audio::Insights
 		void SetYPos(const float Pos) { YPos = Pos; }
 
 	private:
+		TSharedRef<SHorizontalBox> CreateNodeName(const SSignalFlowGraphNode::FArguments& InArgs) const;
 		TSharedRef<SVerticalBox> CreateNodeDetails() const;
 		TSharedRef<SHorizontalBox> CreateNodeDetailsRow(const ESignalFlowNodeDetailParam ParamType) const;
 
