[Audio Insights] Remove usage of TArrayView inside TSet in SignalFlowTraceProvider - fix for Linux and Mac builds

#Synced-CL 50197703

#Synced-CL 50197703

diff --git a/Engine/Plugins/AudioInsights/Source/AudioInsights/Private/Providers/SignalFlowTraceProvider.cpp b/Engine/Plugins/AudioInsights/Source/AudioInsights/Private/Providers/SignalFlowTraceProvider.cpp
--- a/Engine/Plugins/AudioInsights/Source/AudioInsights/Private/Providers/SignalFlowTraceProvider.cpp
+++ b/Engine/Plugins/AudioInsights/Source/AudioInsights/Private/Providers/SignalFlowTraceProvider.cpp
@@ -4,7 +4,7 @@
 #include "Async/ParallelFor.h"
 #include "AudioInsightsModule.h"
 #include "Cache/AudioInsightsCacheManager.h"
-#include "Containers/ArrayView.h"
+#include "Containers/Array.h"
 #include "Containers/Set.h"
 
 #define LOCTEXT_NAMESPACE "AudioInsights"
@@ -830,11 +830,18 @@ namespace UE::Audio::Insights
 
 	void FSignalFlowTraceProvider::CollectParamsForTimestamp(const FAudioInsightsCacheManager& CacheManager, const TSet<FSignalFlowEntryKey>& ActiveEntries, const double TimeMarker)
 	{
-		const TArrayView<const FSignalFlowEntryKey> ActiveEntriesView = ActiveEntries.ArrayView();
+		TArray<const FSignalFlowEntryKey> ActiveEntriesArray;
+		ActiveEntriesArray.Reserve(ActiveEntries.Num());
 
-		ParallelFor(ActiveEntriesView.Num(), [&CacheManager, TimeMarker, &ActiveEntriesView, this](const int32 Index)
+		for (auto Itr = ActiveEntries.CreateConstIterator(); Itr; ++Itr)
 		{
-			const FSignalFlowEntryKey& EntryKey = ActiveEntriesView[Index];
+			ActiveEntriesArray.Add(*Itr);
+		}
+
+		ParallelFor(ActiveEntriesArray.Num(), [&CacheManager, TimeMarker, &ActiveEntriesArray, this](const int32 Index)
+		{
+			const FSignalFlowEntryKey& EntryKey = ActiveEntriesArray[Index];
+
 			if (EntryKey.EntryType != ESignalFlowEntryType::SoundSource
 				&& EntryKey.EntryType != ESignalFlowEntryType::AudioBus
 				&& EntryKey.EntryType != ESignalFlowEntryType::Submix)
