[AudioInsights] - Support 'CheckBoxCheck' and 'CheckBoxIndeterminate' icons from figma

https://p4-swarm.epicgames.net/reviews/50307860/

Passes Bughawk

#Synced-CL 50390951

diff --git a/Engine/Plugins/AudioInsights/Source/AudioInsights/Public/AudioInsightsStyle.h b/Engine/Plugins/AudioInsights/Source/AudioInsights/Public/AudioInsightsStyle.h
--- a/Engine/Plugins/AudioInsights/Source/AudioInsights/Public/AudioInsightsStyle.h
+++ b/Engine/Plugins/AudioInsights/Source/AudioInsights/Public/AudioInsightsStyle.h
@@ -192,6 +192,8 @@ namespace UE::Audio::Insights
 			Set("AudioInsights.Icon", new IMAGE_BRUSH_SVG(TEXT("Icons/audio_insights_icon"), Icon16));
 			Set("AudioInsights.Icon.ClearLog", new IMAGE_BRUSH_SVG(TEXT("Icons/audio_clear_log"), Icon20));
 			Set("AudioInsights.Icon.ContentBrowser", new IMAGE_BRUSH_SVG(TEXT("Icons/content_browser"), Icon20));
+			Set("AudioInsights.Icon.CheckBoxCheck", new IMAGE_BRUSH_SVG(TEXT("Icons/checkboxcheck"), Icon16));
+			Set("AudioInsights.Icon.CheckBoxIndeterminate", new IMAGE_BRUSH_SVG(TEXT("Icons/checkboxindeterminate"), Icon16));
 			Set("AudioInsights.Icon.Dashboard", new IMAGE_BRUSH_SVG(TEXT("Icons/audio_dashboard"), Icon16));
 			Set("AudioInsights.Icon.Event", new IMAGE_BRUSH_SVG(TEXT("Icons/audio_event"), Icon16));
 			Set("AudioInsights.Icon.Log", new IMAGE_BRUSH_SVG(TEXT("Icons/audio_log"), Icon16));
