[Audio Insights] First pass on adding settings menu to Signal Flow.
* Create persistent settings struct (FSignalFlowSettings) for Signal Flow
* Create new button and drop down menu for signal flow settings
* Create new UI Commands class (FSignalFlowEditorCommands) for signal flow
* Begin by integrating "Horizontal Flow" toggle setting into the persistent settings and settings menu
* Set default for signal flow orientation to Horizontal

#Synced-CL 50345329

diff --git a/Engine/Plugins/AudioInsights/Source/AudioInsights/Private/Settings/SignalFlowSettings.cpp b/Engine/Plugins/AudioInsights/Source/AudioInsights/Private/Settings/SignalFlowSettings.cpp
new file mode 100644
--- /dev/null
+++ b/Engine/Plugins/AudioInsights/Source/AudioInsights/Private/Settings/SignalFlowSettings.cpp
@@ -0,0 +1,13 @@
+// Copyright Epic Games, Inc. All Rights Reserved.
+#include "Settings/SignalFlowSettings.h"
+
+#include UE_INLINE_GENERATED_CPP_BY_NAME(SignalFlowSettings)
+
+#if WITH_EDITOR
+
+FSignalFlowSettings::FOnReadSignalFlowSettings FSignalFlowSettings::OnReadSettings;
+FSignalFlowSettings::FOnWriteSignalFlowSettings FSignalFlowSettings::OnWriteSettings;
+FSignalFlowSettings::FOnRequestReadSignalFlowSettings FSignalFlowSettings::OnRequestReadSettings;
+FSignalFlowSettings::FOnRequestWriteSignalFlowSettings FSignalFlowSettings::OnRequestWriteSettings;
+
+#endif // WITH_EDITOR
diff --git a/Engine/Plugins/AudioInsights/Source/AudioInsights/Private/SignalFlowEditorCommands.cpp b/Engine/Plugins/AudioInsights/Source/AudioInsights/Private/SignalFlowEditorCommands.cpp
new file mode 100644
--- /dev/null
+++ b/Engine/Plugins/AudioInsights/Source/AudioInsights/Private/SignalFlowEditorCommands.cpp
@@ -0,0 +1,24 @@
+// Copyright Epic Games, Inc. All Rights Reserved.
+#include "SignalFlowEditorCommands.h"
+
+#include "AudioInsightsStyle.h"
+#include "Framework/Commands/InputChord.h"
+#include "Framework/Commands/UICommandInfo.h"
+
+#define LOCTEXT_NAMESPACE "AudioInsights"
+
+namespace UE::Audio::Insights
+{
+	FSignalFlowEditorCommands::FSignalFlowEditorCommands()
+		: TCommands<FSignalFlowEditorCommands>("SignalFlowEditorCommands", LOCTEXT("SignalFlowEditorCommands_ContextDescText", "Signal Flow Editor Commands"), NAME_None, UE::Audio::Insights::FSlateStyle::GetStyleName())
+	{
+	}
+
+	void FSignalFlowEditorCommands::RegisterCommands()
+	{
+		UI_COMMAND(ToggleHorizontalView, "Horizontal Flow", "Toggles between horizontally/vertically orienting the Signal Flow graph.", EUserInterfaceActionType::ToggleButton, FInputChord(EKeys::H, EModifierKey::Control));
+	}
+
+} // namespace UE::Audio::Insights
+
+#undef LOCTEXT_NAMESPACE
diff --git a/Engine/Plugins/AudioInsights/Source/AudioInsights/Private/SignalFlowEditorCommands.h b/Engine/Plugins/AudioInsights/Source/AudioInsights/Private/SignalFlowEditorCommands.h
new file mode 100644
--- /dev/null
+++ b/Engine/Plugins/AudioInsights/Source/AudioInsights/Private/SignalFlowEditorCommands.h
@@ -0,0 +1,23 @@
+// Copyright Epic Games, Inc. All Rights Reserved.
+#pragma once
+
+#include "Framework/Commands/Commands.h"
+#include "Templates/SharedPointer.h"
+
+class FUICommandInfo;
+
+namespace UE::Audio::Insights
+{
+	class FSignalFlowEditorCommands : public TCommands<FSignalFlowEditorCommands>
+	{
+	public:
+		FSignalFlowEditorCommands();
+
+		virtual void RegisterCommands() override;
+
+		TSharedPtr<const FUICommandInfo> GetToggleHorizontalViewCommand() const { return ToggleHorizontalView; }
+
+	private:
+		TSharedPtr<FUICommandInfo> ToggleHorizontalView;
+	};
+} // namespace UE::Audio::Insights
diff --git a/Engine/Plugins/AudioInsights/Source/AudioInsights/Private/Views/SignalFlowDashboardViewFactory.cpp b/Engine/Plugins/AudioInsights/Source/AudioInsights/Private/Views/SignalFlowDashboardViewFactory.cpp
--- a/Engine/Plugins/AudioInsights/Source/AudioInsights/Private/Views/SignalFlowDashboardViewFactory.cpp
+++ b/Engine/Plugins/AudioInsights/Source/AudioInsights/Private/Views/SignalFlowDashboardViewFactory.cpp
@@ -4,13 +4,16 @@
 #include "AudioInsightsModule.h"
 #include "AudioInsightsStyle.h"
 #include "AudioInsightsTimingViewExtender.h"
+#include "Framework/MultiBox/MultiBoxBuilder.h"
 #include "IAudioInsightsModule.h"
 #include "Internationalization/Text.h"
 #include "Providers/SignalFlowTraceProvider.h"
+#include "SignalFlowEditorCommands.h"
 #include "Styling/StyleColors.h"
 #include "Views/SignalFlowNodes.h"
 #include "Widgets/Colors/SColorBlock.h"
 #include "Widgets/Input/SButton.h"
+#include "Widgets/Input/SComboButton.h"
 #include "Widgets/Layout/SSeparator.h"
 
 #if WITH_EDITOR
@@ -62,8 +65,16 @@ namespace UE::Audio::Insights
 			SignalFlowProvider
 		};
 
+		FSignalFlowEditorCommands::Register();
+		BindCommands();
+
 		SignalFlowProvider->OnRequestGraphRefresh.AddRaw(this, &FSignalFlowDashboardViewFactory::OnRequestGraphRefresh);
 		SignalFlowProvider->OnResetGraph.AddRaw(this, &FSignalFlowDashboardViewFactory::Reset);
+
+#if WITH_EDITOR
+		FSignalFlowSettings::OnReadSettings.AddRaw(this, &FSignalFlowDashboardViewFactory::OnReadEditorSettings);
+		FSignalFlowSettings::OnWriteSettings.AddRaw(this, &FSignalFlowDashboardViewFactory::OnWriteEditorSettings);
+#endif // WITH_EDITOR
 	}
 
 	FSignalFlowDashboardViewFactory::~FSignalFlowDashboardViewFactory()
@@ -74,6 +85,13 @@ namespace UE::Audio::Insights
 			Provider->OnRequestGraphRefresh.RemoveAll(this);
 			Provider->OnResetGraph.RemoveAll(this);
 		}
+
+#if WITH_EDITOR
+		FSignalFlowSettings::OnReadSettings.RemoveAll(this);
+		FSignalFlowSettings::OnWriteSettings.RemoveAll(this);
+#endif // WITH_EDITOR
+
+		FSignalFlowEditorCommands::Unregister();
 	}
 
 	FName FSignalFlowDashboardViewFactory::GetName() const
@@ -131,6 +149,11 @@ namespace UE::Audio::Insights
 			DashboardWidget->SetOnKeyDownHandler(FOnKeyDown::CreateSP(this, &FSignalFlowDashboardViewFactory::OnKeyDown));
 		}
 
+#if WITH_EDITOR
+		// Read the editor settings after the widget has been created
+		FSignalFlowSettings::OnRequestReadSettings.Broadcast();
+#endif // WITH_EDITOR
+
 		return DashboardWidget->AsShared();
 	}
 
@@ -138,10 +161,8 @@ namespace UE::Audio::Insights
 	{
 		if (GraphWidgetContents.IsValid() && InKeyEvent.IsControlDown())
 		{
-			if (InKeyEvent.GetKey() == EKeys::H)
+			if (CommandList && CommandList->ProcessCommandBindings(InKeyEvent))
 			{
-				GraphWidgetContents->ToggleOrientation();
-
 				return FReply::Handled();
 			}
 			else if (InKeyEvent.GetKey() == EKeys::F)
@@ -207,6 +228,19 @@ namespace UE::Audio::Insights
 		return &FSlateStyle::Get().GetWidgetStyle<FTableRowStyle>("TreeDashboard.TableViewRow");
 	}
 
+	void FSignalFlowDashboardViewFactory::BindCommands()
+	{
+		CommandList = MakeShared<FUICommandList>();
+
+		const FSignalFlowEditorCommands& Commands = FSignalFlowEditorCommands::Get();
+
+		CommandList->MapAction(
+			Commands.GetToggleHorizontalViewCommand(),
+			FExecuteAction::CreateLambda([this]() { ToggleGraphOrientation(); }),
+			FCanExecuteAction(),
+			FIsActionChecked::CreateLambda([this]() { return GraphWidgetContents.IsValid() && GraphWidgetContents->GetGraphOrientation() == Orient_Horizontal; }));
+	}
+
 	TSharedRef<SWidget> FSignalFlowDashboardViewFactory::CreateToolbar()
 	{
 		using namespace FSignalFlowDashboardViewFactoryPrivate;
@@ -214,12 +248,14 @@ namespace UE::Audio::Insights
 		return SNew(SHorizontalBox)
 
 		+ SHorizontalBox::Slot()
+		.HAlign(HAlign_Left)
 		.AutoWidth()
 		[
 			CreateToolbarToggleControls()
 		]
 			
 		+ SHorizontalBox::Slot()
+		.HAlign(HAlign_Left)
 		.AutoWidth()
 		.Padding(ToolbarButtonPaddingX, 0.0f)
 		[
@@ -228,12 +264,58 @@ namespace UE::Audio::Insights
 		]
 
 		+ SHorizontalBox::Slot()
+		.HAlign(HAlign_Left)
 		.AutoWidth()
 		[
 			CreateToolbarButtonActions()
+		]
+
+		+ SHorizontalBox::Slot()
+		.HAlign(HAlign_Left)
+		.AutoWidth()
+		.Padding(ToolbarButtonPaddingX, 0.0f)
+		[
+			SNew(SSeparator)
+			.Orientation(EOrientation::Orient_Vertical)
+		]
+
+		+ SHorizontalBox::Slot()
+		.FillWidth(1.0f)
+			
+		+ SHorizontalBox::Slot()
+		.HAlign(HAlign_Right)
+		.AutoWidth()
+		[
+			CreateToolbarSettingsMenu()
 		];
 	}
 
+	TSharedRef<SWidget> FSignalFlowDashboardViewFactory::CreateToolbarSettingsMenu()
+	{
+		using namespace FSignalFlowDashboardViewFactoryPrivate;
+		
+		return SNew(SComboButton)
+			.ComboButtonStyle(&FAppStyle::Get().GetWidgetStyle<FComboButtonStyle>("SimpleComboButton"))
+			.OnGetMenuContent(this, &FSignalFlowDashboardViewFactory::OnGetSettingsMenuContent)
+			.MenuPlacement(EMenuPlacement::MenuPlacement_ComboBoxRight)
+			.HasDownArrow(false)
+			.ToolTipText(LOCTEXT("AudioDashboard_SignalFlow_SettingsTooltip", "Open the settings menu."))
+			.ButtonContent()
+			[
+				SNew(SHorizontalBox)
+
+				+ SHorizontalBox::Slot()
+				.AutoWidth()
+				.VAlign(VAlign_Center)
+				.Padding(5.0, 0.0f)
+				[
+					SNew(SImage)
+					.ColorAndOpacity(FSlateColor::UseForeground())
+					.Image(FSlateStyle::Get().GetBrush("AudioInsights.Icon.Settings"))
+				]
+			];
+	}
+
 	TSharedRef<SWidget> FSignalFlowDashboardViewFactory::CreateToolbarToggleControls()
 	{
 		using namespace FSignalFlowDashboardViewFactoryPrivate;
@@ -261,10 +343,7 @@ namespace UE::Audio::Insights
 			.ToolTipText(LOCTEXT("AudioDashboard_SignalFlow_HorizontalFlowTooltip", "Toggles between horizontally/vertically orienting the Signal Flow graph. (Ctrl + h)"))
 			.OnClicked_Lambda([this]()
 			{
-				if (GraphWidgetContents.IsValid())
-				{
-					GraphWidgetContents->ToggleOrientation();
-				}
+				ToggleGraphOrientation();
 
 				return FReply::Handled();
 			})
@@ -326,6 +405,27 @@ namespace UE::Audio::Insights
 		];
 	}
 
+	TSharedRef<SWidget> FSignalFlowDashboardViewFactory::OnGetSettingsMenuContent()
+	{
+		if (!CommandList.IsValid())
+		{
+			return SNullWidget::NullWidget;
+		}
+
+		FMenuBuilder MenuBuilder(false /*bShouldCloseWindowAfterMenuSelection*/, CommandList);
+
+		const FSignalFlowEditorCommands& Commands = FSignalFlowEditorCommands::Get();
+
+		MenuBuilder.BeginSection("SignalFlowSettingActions", LOCTEXT("AudioDashboard_SignalFlow_Settings_HeaderText", "Settings"));
+		{
+			MenuBuilder.AddMenuEntry(Commands.GetToggleHorizontalViewCommand());
+		}
+
+		MenuBuilder.EndSection();
+
+		return MenuBuilder.MakeWidget();
+	}
+
 	FReply FSignalFlowDashboardViewFactory::ScrollGraphViewToSelectedNode(const bool bResetZoom)
 	{
 		if (GraphWidgetContents.IsValid())
@@ -361,6 +461,18 @@ namespace UE::Audio::Insights
 		return FReply::Unhandled();
 	}
 
+	void FSignalFlowDashboardViewFactory::ToggleGraphOrientation()
+	{
+		if (GraphWidgetContents.IsValid())
+		{
+			GraphWidgetContents->ToggleOrientation();
+		}
+
+#if WITH_EDITOR
+		FSignalFlowSettings::OnRequestWriteSettings.Broadcast();
+#endif // WITH_EDITOR	
+	}
+
 	void FSignalFlowDashboardViewFactory::ToggleShowNodeDetails()
 	{
 		bShowNodeDetails = !bShowNodeDetails;
@@ -1459,6 +1571,29 @@ namespace UE::Audio::Insights
 			}
 		}
 	}
+
+#if WITH_EDITOR
+	void FSignalFlowDashboardViewFactory::OnReadEditorSettings(const FSignalFlowSettings& InSettings)
+	{
+		if (GraphWidgetContents.IsValid())
+		{
+			const bool bGraphIsHorizontal = GraphWidgetContents->GetGraphOrientation() == EOrientation::Orient_Horizontal;
+			if (InSettings.bHorizontalFlow != bGraphIsHorizontal)
+			{
+				GraphWidgetContents->ToggleOrientation();
+			}
+		}
+	}
+
+	void FSignalFlowDashboardViewFactory::OnWriteEditorSettings(FSignalFlowSettings& OutSettings)
+	{
+		if (GraphWidgetContents.IsValid())
+		{
+			const bool bGraphIsHorizontal = GraphWidgetContents->GetGraphOrientation() == EOrientation::Orient_Horizontal;
+			OutSettings.bHorizontalFlow = bGraphIsHorizontal;
+		}
+	}
+#endif // WITH_EDITOR
 } // namespace UE::Audio::Insights
 
 #undef LOCTEXT_NAMESPACE
diff --git a/Engine/Plugins/AudioInsights/Source/AudioInsights/Public/Settings/SignalFlowSettings.h b/Engine/Plugins/AudioInsights/Source/AudioInsights/Public/Settings/SignalFlowSettings.h
new file mode 100644
--- /dev/null
+++ b/Engine/Plugins/AudioInsights/Source/AudioInsights/Public/Settings/SignalFlowSettings.h
@@ -0,0 +1,36 @@
+// Copyright Epic Games, Inc. All Rights Reserved.
+
+#pragma once
+
+#include "Delegates/DelegateCombinations.h"
+#include "Engine/DeveloperSettings.h"
+
+#include "SignalFlowSettings.generated.h"
+
+#define LOCTEXT_NAMESPACE "AudioInsights"
+
+USTRUCT()
+struct FSignalFlowSettings
+{
+	GENERATED_BODY()
+
+	/** Toggles between horizontally/vertically orienting the Signal Flow graph. */
+	UPROPERTY(EditAnywhere, config, Category = SignalFlow)
+	bool bHorizontalFlow = true;
+
+#if WITH_EDITOR
+	DECLARE_MULTICAST_DELEGATE_OneParam(FOnReadSignalFlowSettings, const FSignalFlowSettings&);
+	static AUDIOINSIGHTS_API FOnReadSignalFlowSettings OnReadSettings;
+
+	DECLARE_MULTICAST_DELEGATE_OneParam(FOnWriteSignalFlowSettings, FSignalFlowSettings&);
+	static AUDIOINSIGHTS_API FOnWriteSignalFlowSettings OnWriteSettings;
+
+	DECLARE_MULTICAST_DELEGATE(FOnRequestReadSignalFlowSettings);
+	static AUDIOINSIGHTS_API FOnRequestReadSignalFlowSettings OnRequestReadSettings;
+
+	DECLARE_MULTICAST_DELEGATE(FOnRequestWriteSignalFlowSettings);
+	static AUDIOINSIGHTS_API FOnRequestWriteSignalFlowSettings OnRequestWriteSettings;
+#endif // WITH_EDITOR
+};
+
+#undef LOCTEXT_NAMESPACE
\ No newline at end of file
diff --git a/Engine/Plugins/AudioInsights/Source/AudioInsights/Public/Views/SignalFlowDashboardViewFactory.h b/Engine/Plugins/AudioInsights/Source/AudioInsights/Public/Views/SignalFlowDashboardViewFactory.h
--- a/Engine/Plugins/AudioInsights/Source/AudioInsights/Public/Views/SignalFlowDashboardViewFactory.h
+++ b/Engine/Plugins/AudioInsights/Source/AudioInsights/Public/Views/SignalFlowDashboardViewFactory.h
@@ -2,6 +2,7 @@
 #pragma once
 
 #include "Messages/SignalFlowEntryKey.h"
+#include "Settings/SignalFlowSettings.h"
 #include "Templates/SharedPointer.h"
 #include "Views/SSignalFlowGraph.h"
 #include "Views/TableDashboardViewFactory.h"
@@ -26,7 +27,7 @@ namespace UE::Audio::Insights
 		UE_API virtual FReply OnKeyDown(const FGeometry& MyGeometry, const FKeyEvent& InKeyEvent) override;
 		virtual bool SupportsKeyboardFocus() const override { return true; }
 
-		UE_API void SetOnKeyDownHandler(FOnKeyDown Handler) { OnKeyDownHandler = Handler; }
+		void SetOnKeyDownHandler(FOnKeyDown Handler) { OnKeyDownHandler = Handler; }
 
 	private:
 		FOnKeyDown OnKeyDownHandler;
@@ -64,11 +65,17 @@ namespace UE::Audio::Insights
 
 		const FTableRowStyle* GetRowStyle() const;
 
+		void BindCommands();
+
 		TSharedRef<SWidget> CreateToolbar();
+		TSharedRef<SWidget> CreateToolbarSettingsMenu();
 		TSharedRef<SWidget> CreateToolbarToggleControls();
 		TSharedRef<SWidget> CreateToolbarButtonActions();
 
+		TSharedRef<SWidget> OnGetSettingsMenuContent();
+
 		FReply ScrollGraphViewToSelectedNode(const bool bResetZoom);
+		void ToggleGraphOrientation();
 		void ToggleShowNodeDetails();
 
 		TSharedPtr<FSignalFlowDashboardEntry> GetFilteredAudioDeviceEntry() const;
@@ -112,6 +119,11 @@ namespace UE::Audio::Insights
 		void CalcXPositionForNode(const TSharedPtr<ISignalFlowNode>& Node, const TArray<FSignalFlowEntryKey>& NodesToPositionRelativeTo, const TArray<TSharedPtr<ISignalFlowNode>>& PrevPositionedNodes, const int32 NumNodesInRow);
 		void MakeSpaceForNodeInRow(const TSharedPtr<ISignalFlowNode>& Node, const TArray<TSharedPtr<ISignalFlowNode>>& PrevPositionedNodes, const int32 NumNodesInRow);
 
+#if WITH_EDITOR
+		void OnReadEditorSettings(const FSignalFlowSettings& InSettings);
+		void OnWriteEditorSettings(FSignalFlowSettings& OutSettings);
+#endif // WITH_EDITOR
+
 		TSharedPtr<FSignalFlowTraceProvider> SignalFlowProvider;
 		TSharedPtr<SSignalFlowDashboard> DashboardWidget;
 
@@ -131,6 +143,8 @@ namespace UE::Audio::Insights
 		TSharedPtr<SScrollBox> GraphHorizontalScrollBox;
 		TSharedPtr<SScrollBox> GraphVerticalScrollBox;
 
+		TSharedPtr<FUICommandList> CommandList;
+
 		TSharedPtr<FSignalFlowDashboardEntry> SelectedItem;
 		FText SearchBoxFilterText;
 
diff --git a/Engine/Plugins/AudioInsights/Source/AudioInsightsEditor/Private/AudioInsightsEditorSettings.cpp b/Engine/Plugins/AudioInsights/Source/AudioInsightsEditor/Private/AudioInsightsEditorSettings.cpp
--- a/Engine/Plugins/AudioInsights/Source/AudioInsightsEditor/Private/AudioInsightsEditorSettings.cpp
+++ b/Engine/Plugins/AudioInsights/Source/AudioInsightsEditor/Private/AudioInsightsEditorSettings.cpp
@@ -69,6 +69,17 @@ void UAudioInsightsEditorSettings::PostInitProperties()
 		FSoundDashboardSettings::OnWriteSettings.Broadcast(SoundDashboardSettings);
 		SaveConfig();
 	});
+
+	OnRequestReadSignalFlowSettingsHandle = FSignalFlowSettings::OnRequestReadSettings.AddLambda([this]()
+	{
+		FSignalFlowSettings::OnReadSettings.Broadcast(SignalFlowSettings);
+	});
+
+	OnRequestWriteSignalFlowSettingsHandle = FSignalFlowSettings::OnRequestWriteSettings.AddLambda([this]()
+	{
+		FSignalFlowSettings::OnWriteSettings.Broadcast(SignalFlowSettings);
+		SaveConfig();
+	});
 }
 
 void UAudioInsightsEditorSettings::BeginDestroy()
@@ -80,6 +91,9 @@ void UAudioInsightsEditorSettings::BeginDestroy()
 
 	FSoundDashboardSettings::OnRequestReadSettings.Remove(OnRequestReadSoundDashboardSettingsHandle);
 	FSoundDashboardSettings::OnRequestWriteSettings.Remove(OnRequestWriteSoundDashboardSettingsHandle);
+
+	FSignalFlowSettings::OnRequestReadSettings.Remove(OnRequestReadSignalFlowSettingsHandle);
+	FSignalFlowSettings::OnRequestWriteSettings.Remove(OnRequestWriteSignalFlowSettingsHandle);
 }
 
 void UAudioInsightsEditorSettings::PostEditChangeChainProperty(FPropertyChangedChainEvent& PropertyChangedEvent)
@@ -115,6 +129,11 @@ void UAudioInsightsEditorSettings::PostEditChangeChainProperty(FPropertyChangedC
 	{
 		FSoundDashboardSettings::OnReadSettings.Broadcast(SoundDashboardSettings);
 	}
+
+	if (ActiveMemberNode->GetValue()->GetFName() == GET_MEMBER_NAME_CHECKED(UAudioInsightsEditorSettings, SignalFlowSettings))
+	{
+		FSignalFlowSettings::OnReadSettings.Broadcast(SignalFlowSettings);
+	}
 }
 
 void UAudioInsightsEditorSettings::PostEditUndo()
@@ -123,6 +142,7 @@ void UAudioInsightsEditorSettings::PostEditUndo()
 
 	FAudioEventLogSettings::OnReadSettings.Broadcast(EventLogSettings);
 	FSoundDashboardSettings::OnReadSettings.Broadcast(SoundDashboardSettings);
+	FSignalFlowSettings::OnReadSettings.Broadcast(SignalFlowSettings);
 }
 
 bool UAudioInsightsEditorSettings::VerifyEventLogInput(const TDoubleLinkedList<FProperty*>::TDoubleLinkedListNode& ActiveMemberNode, FPropertyChangedChainEvent& PropertyChangedEvent)
diff --git a/Engine/Plugins/AudioInsights/Source/AudioInsightsEditor/Public/AudioInsightsEditorSettings.h b/Engine/Plugins/AudioInsights/Source/AudioInsightsEditor/Public/AudioInsightsEditorSettings.h
--- a/Engine/Plugins/AudioInsights/Source/AudioInsightsEditor/Public/AudioInsightsEditorSettings.h
+++ b/Engine/Plugins/AudioInsights/Source/AudioInsightsEditor/Public/AudioInsightsEditorSettings.h
@@ -7,6 +7,7 @@
 #include "Engine/DeveloperSettings.h"
 #include "LoudnessMeterRackUnitSettings.h"
 #include "Settings/AudioEventLogSettings.h"
+#include "Settings/SignalFlowSettings.h"
 #include "Settings/SoundDashboardSettings.h"
 
 #include "AudioInsightsEditorSettings.generated.h"
@@ -58,6 +59,10 @@ public:
 	UPROPERTY(EditAnywhere, config, Category = SoundDashboard, meta = (ShowOnlyInnerProperties))
 	FSoundDashboardSettings SoundDashboardSettings;
 
+	/** Settings for the Signal Flow Dashboard */
+	UPROPERTY(EditAnywhere, config, Category = SignalFlow, meta = (ShowOnlyInnerProperties))
+	FSignalFlowSettings SignalFlowSettings;
+
 	DECLARE_MULTICAST_DELEGATE_OneParam(FOnOnlyTraceAudioChannelsPropertyChanged, const bool);
 	FOnOnlyTraceAudioChannelsPropertyChanged OnOnlyTraceAudioChannelsPropertyChanged;
 
@@ -86,6 +91,9 @@ private:
 	FDelegateHandle OnRequestReadSoundDashboardSettingsHandle;
 	FDelegateHandle OnRequestWriteSoundDashboardSettingsHandle;
 
+	FDelegateHandle OnRequestReadSignalFlowSettingsHandle;
+	FDelegateHandle OnRequestWriteSignalFlowSettingsHandle;
+
 	const TMap<FString, TSet<FString>> InBuiltAudioLogEventTypes;
 };
 
