[Audio Insights] Add Node Detail filters to Signal Flow.
* Detail visibility is controlled via the settings menu, and are stored inside FSignalFlowNodeDetailFilterSettings
* Height of nodes in SSignalFlowGraph are scaled with the number of parameters visible

#Synced-CL 50395601

diff --git a/Engine/Plugins/AudioInsights/Source/AudioInsights/Private/Settings/SignalFlowSettings.cpp b/Engine/Plugins/AudioInsights/Source/AudioInsights/Private/Settings/SignalFlowSettings.cpp
--- a/Engine/Plugins/AudioInsights/Source/AudioInsights/Private/Settings/SignalFlowSettings.cpp
+++ b/Engine/Plugins/AudioInsights/Source/AudioInsights/Private/Settings/SignalFlowSettings.cpp
@@ -11,3 +11,144 @@ FSignalFlowSettings::FOnRequestReadSignalFlowSettings FSignalFlowSettings::OnReq
 FSignalFlowSettings::FOnRequestWriteSignalFlowSettings FSignalFlowSettings::OnRequestWriteSettings;
 
 #endif // WITH_EDITOR
+
+bool FSignalFlowNodeDetailFilterSettings::GetParameterIsVisible(const UE::Audio::Insights::ESignalFlowNodeDetailParam Param) const
+{
+	using namespace UE::Audio::Insights;
+
+	switch (Param)
+	{
+		case ESignalFlowNodeDetailParam::Amplitude:
+			return bAmplitude;
+
+		case ESignalFlowNodeDetailParam::Volume:
+			return bVolume;
+
+		case ESignalFlowNodeDetailParam::Pitch:
+			return bPitch;
+
+		case ESignalFlowNodeDetailParam::LPFFreq:
+			return bLPFFreq;
+
+		case ESignalFlowNodeDetailParam::HPFFreq:
+			return bHPFFreq;
+
+		case ESignalFlowNodeDetailParam::Priority:
+			return bPriority;
+
+		case ESignalFlowNodeDetailParam::Distance:
+			return bDistance;
+
+		case ESignalFlowNodeDetailParam::Attenuation:
+			return bDistanceOcclusionAttenuation;
+
+		case ESignalFlowNodeDetailParam::RelativeRenderCost:
+			return bRelativeRenderCost;
+
+		default:
+			ensureMsgf(false, TEXT("Unrecognized ESignalFlowNodeDetailParam value"));
+			break;
+	}
+
+    return false;
+}
+
+int32 FSignalFlowNodeDetailFilterSettings::GetNumVisibleParams() const
+{
+	using namespace UE::Audio::Insights;
+
+	int32 NumVisibleParams = 0;
+
+	const int32 NumFilters = static_cast<int32>(ESignalFlowNodeDetailParam::MAX);
+	for (uint8 FilterID = 0u; FilterID < NumFilters; ++FilterID)
+	{
+		if (GetParameterIsVisible(static_cast<ESignalFlowNodeDetailParam>(FilterID)))
+		{
+			NumVisibleParams++;
+		}
+	}
+
+	return NumVisibleParams;
+}
+
+ECheckBoxState FSignalFlowNodeDetailFilterSettings::GetShowAllNodeDetailFiltersCheckboxState() const
+{
+	using namespace UE::Audio::Insights;
+
+	const int32 NumVisibleParams = GetNumVisibleParams();
+
+	if (NumVisibleParams == 0)
+	{
+		return ECheckBoxState::Unchecked;
+	}
+
+	const int32 NumFilters = static_cast<int32>(ESignalFlowNodeDetailParam::MAX);
+	if (NumVisibleParams == NumFilters)
+	{
+		return ECheckBoxState::Checked;
+	}
+
+	return ECheckBoxState::Undetermined;
+}
+
+void FSignalFlowNodeDetailFilterSettings::ToggleEnableAllNodeDetailFilters()
+{
+	using namespace UE::Audio::Insights;
+
+	const bool bAllFiltersAreEnabled = GetShowAllNodeDetailFiltersCheckboxState() == ECheckBoxState::Checked;
+
+	const int32 NumFilters = static_cast<int32>(ESignalFlowNodeDetailParam::MAX);
+
+	for (int32 FilterID = 0; FilterID < NumFilters; ++FilterID)
+	{
+		SetParameterVisibility(static_cast<ESignalFlowNodeDetailParam>(FilterID), !bAllFiltersAreEnabled);
+	}
+}
+
+void FSignalFlowNodeDetailFilterSettings::SetParameterVisibility(const UE::Audio::Insights::ESignalFlowNodeDetailParam Param, const bool bIsVisible)
+{
+	using namespace UE::Audio::Insights;
+
+	switch (Param)
+	{
+		case ESignalFlowNodeDetailParam::Amplitude:
+			bAmplitude = bIsVisible;
+			break;
+
+		case ESignalFlowNodeDetailParam::Volume:
+			bVolume = bIsVisible;
+			break;
+
+		case ESignalFlowNodeDetailParam::Pitch:
+			bPitch = bIsVisible;
+			break;
+
+		case ESignalFlowNodeDetailParam::LPFFreq:
+			bLPFFreq = bIsVisible;
+			break;
+
+		case ESignalFlowNodeDetailParam::HPFFreq:
+			bHPFFreq = bIsVisible;
+			break;
+
+		case ESignalFlowNodeDetailParam::Priority:
+			bPriority = bIsVisible;
+			break;
+
+		case ESignalFlowNodeDetailParam::Distance:
+			bDistance = bIsVisible;
+			break;
+
+		case ESignalFlowNodeDetailParam::Attenuation:
+			bDistanceOcclusionAttenuation = bIsVisible;
+			break;
+
+		case ESignalFlowNodeDetailParam::RelativeRenderCost:
+			bRelativeRenderCost = bIsVisible;
+			break;
+
+		default:
+			ensureMsgf(false, TEXT("Unrecognized ESignalFlowNodeDetailParam value"));
+			break;
+	}
+}
\ No newline at end of file
diff --git a/Engine/Plugins/AudioInsights/Source/AudioInsights/Private/SignalFlowEditorCommands.cpp b/Engine/Plugins/AudioInsights/Source/AudioInsights/Private/SignalFlowEditorCommands.cpp
--- a/Engine/Plugins/AudioInsights/Source/AudioInsights/Private/SignalFlowEditorCommands.cpp
+++ b/Engine/Plugins/AudioInsights/Source/AudioInsights/Private/SignalFlowEditorCommands.cpp
@@ -17,6 +17,7 @@ namespace UE::Audio::Insights
 	void FSignalFlowEditorCommands::RegisterCommands()
 	{
 		UI_COMMAND(ToggleHorizontalView, "Horizontal Flow", "Toggles between horizontally/vertically orienting the Signal Flow graph.", EUserInterfaceActionType::ToggleButton, FInputChord(EKeys::H, EModifierKey::Control));
+		UI_COMMAND(ShowNodeDetails, "Show Node Details", "Toggles visibility of node details.", EUserInterfaceActionType::ToggleButton, FInputChord(EKeys::D, EModifierKey::Control));
 	}
 
 } // namespace UE::Audio::Insights
diff --git a/Engine/Plugins/AudioInsights/Source/AudioInsights/Private/SignalFlowEditorCommands.h b/Engine/Plugins/AudioInsights/Source/AudioInsights/Private/SignalFlowEditorCommands.h
--- a/Engine/Plugins/AudioInsights/Source/AudioInsights/Private/SignalFlowEditorCommands.h
+++ b/Engine/Plugins/AudioInsights/Source/AudioInsights/Private/SignalFlowEditorCommands.h
@@ -16,8 +16,10 @@ namespace UE::Audio::Insights
 		virtual void RegisterCommands() override;
 
 		TSharedPtr<const FUICommandInfo> GetToggleHorizontalViewCommand() const { return ToggleHorizontalView; }
+		TSharedPtr<const FUICommandInfo> GetShowNodeDetailsCommand() const { return ShowNodeDetails; }
 
 	private:
 		TSharedPtr<FUICommandInfo> ToggleHorizontalView;
+		TSharedPtr<FUICommandInfo> ShowNodeDetails;
 	};
 } // namespace UE::Audio::Insights
diff --git a/Engine/Plugins/AudioInsights/Source/AudioInsights/Private/Views/SSignalFlowGraph.cpp b/Engine/Plugins/AudioInsights/Source/AudioInsights/Private/Views/SSignalFlowGraph.cpp
--- a/Engine/Plugins/AudioInsights/Source/AudioInsights/Private/Views/SSignalFlowGraph.cpp
+++ b/Engine/Plugins/AudioInsights/Source/AudioInsights/Private/Views/SSignalFlowGraph.cpp
@@ -6,6 +6,7 @@
 #include "Brushes/SlateRoundedBoxBrush.h"
 #include "Internationalization/Text.h"
 #include "Messages/SignalFlowTraceMessages.h"
+#include "Settings/SignalFlowSettings.h"
 #include "Styling/StyleColors.h"
 #include "Views/SignalFlowNodes.h"
 #include "Widgets/Colors/SColorBlock.h"
@@ -61,19 +62,18 @@ namespace UE::Audio::Insights
 			return NodeMaxWidth + XPaddingBetweenNodes;
 		}
 
-		float GetNodeMaxHeight(const bool bShowNodeDetails)
+		float GetNodeMaxHeight(const int32 NumParamsVisible, const bool bShowNodeDetails)
 		{
 			constexpr float NodeMaxHeight = 50.0f;
 			constexpr float NodeDetailsRowHeight = 24.0f;
-			constexpr uint32 MaxNumNodeDetails = 9; // TODO - Scale this with the number enabled in settings when UE-358203 is implemented
 
-			return bShowNodeDetails ? NodeMaxHeight + (NodeDetailsRowHeight * MaxNumNodeDetails) : NodeMaxHeight;
+			return bShowNodeDetails ? NodeMaxHeight + (NodeDetailsRowHeight * NumParamsVisible) : NodeMaxHeight;
 		}
 
-		float GetNodeGraphSlotHeight(const EOrientation GraphOrientation, const bool bShowNodeDetails)
+		float GetNodeGraphSlotHeight(const EOrientation GraphOrientation, const int32 NumParamsVisible, const bool bShowNodeDetails)
 		{
 			const float YPaddingBetweenNodes = GraphOrientation == EOrientation::Orient_Horizontal ? 8.0f : 48.0f;
-			return GetNodeMaxHeight(bShowNodeDetails) + YPaddingBetweenNodes;
+			return GetNodeMaxHeight(NumParamsVisible, bShowNodeDetails) + YPaddingBetweenNodes;
 		}
 
 		FSlateColor GetNodeAccentColor(const ESignalFlowEntryType NodeType)
@@ -145,6 +145,7 @@ namespace UE::Audio::Insights
 					return Entry->RelativeRenderCost;
 
 				default:
+					ensureMsgf(false, TEXT("Unrecognized ESignalFlowNodeDetailParam value"));
 					break;
 			}
 
@@ -183,6 +184,7 @@ namespace UE::Audio::Insights
 					return FStyleColors::AccentGray;
 
 				default:
+					ensureMsgf(false, TEXT("Unrecognized ESignalFlowNodeDetailParam value"));
 					break;
 			}
 
@@ -221,6 +223,7 @@ namespace UE::Audio::Insights
 					return LOCTEXT("SignalFlowGraph_NodeDetails_RekativeRenderCost", "Rel. Render Cost: {0}");
 
 				default:
+					ensureMsgf(false, TEXT("Unrecognized ESignalFlowNodeDetailParam value"));
 					break;
 			}
 
@@ -259,6 +262,7 @@ namespace UE::Audio::Insights
 					return FSlateStyle::Get().GetRelativeRenderCostFormat();
 
 				default:
+					ensureMsgf(false, TEXT("Unrecognized ESignalFlowNodeDetailParam value"));
 					break;
 			}
 
@@ -271,6 +275,7 @@ namespace UE::Audio::Insights
 		using namespace SSignalFlowGraphStyle;
 
 		Entry = InArgs._Entry;
+		NodeDetailFilterSettings = InArgs._NodeDetailFilterSettings;
 		PositionInGraph = InArgs._PositionInGraph;
 		XPos = InArgs._XPos;
 		YPos = InArgs._YPos;
@@ -366,8 +371,10 @@ namespace UE::Audio::Insights
 	{
 		using namespace SSignalFlowGraphStyle;
 
+		const int32 NumParamsVisible = NodeDetailFilterSettings ? NodeDetailFilterSettings->GetNumVisibleParams() : 0;
+
 		const FVector2D Size = SCompoundWidget::ComputeDesiredSize(InLayoutScaleMultiplier);
-		const FVector2D ClampedSize = { FMath::Min(Size.X, NodeMaxWidth), FMath::Min(Size.Y, GetNodeMaxHeight(ShowNodeDetails.Get())) };
+		const FVector2D ClampedSize = { FMath::Min(Size.X, NodeMaxWidth), FMath::Min(Size.Y, GetNodeMaxHeight(NumParamsVisible, ShowNodeDetails.Get())) };
 		return ClampedSize;
 	}
 
@@ -533,7 +540,8 @@ namespace UE::Audio::Insights
 		return SNew(SHorizontalBox)
 				.Visibility_Lambda([this, ParamType]()
 				{
-					return Entry.IsValid() && GetNodeDetailParamValue(Entry, ParamType).IsSet() ? EVisibility::Visible : EVisibility::Collapsed;
+					const bool bParamIsFiltered = NodeDetailFilterSettings && NodeDetailFilterSettings->GetParameterIsVisible(ParamType);
+					return bParamIsFiltered && Entry.IsValid() && GetNodeDetailParamValue(Entry, ParamType).IsSet() ? EVisibility::Visible : EVisibility::Collapsed;
 				})
 
 				+ SHorizontalBox::Slot()
@@ -670,6 +678,7 @@ namespace UE::Audio::Insights
 	void SSignalFlowGraph::Construct(const SSignalFlowGraph::FArguments& InArgs)
 	{
 		ListSource = InArgs._ListSource;
+		NodeDetailFilterSettings = InArgs._NodeDetailFilterSettings;
 		SelectedItem = InArgs._SelectedItem;
 		ActiveAudioDeviceEntry = InArgs._ActiveAudioDeviceEntry;
 		HorizontalScrollBox = InArgs._HorizontalScrollBox;
@@ -985,8 +994,10 @@ namespace UE::Audio::Insights
 		float GraphWidth = 0.0f;
 		float GraphHeight = 0.0f;
 
+		const int32 NumParamsVisible = NodeDetailFilterSettings ? NodeDetailFilterSettings->GetNumVisibleParams() : 0;
+
 		const float NodeGraphSlotWidth = GetNodeGraphSlotWidth(GetGraphOrientation());
-		const float NodeGraphSlotHeight = GetNodeGraphSlotHeight(GetGraphOrientation(), ShowNodeDetails.Get());
+		const float NodeGraphSlotHeight = GetNodeGraphSlotHeight(GetGraphOrientation(), NumParamsVisible, ShowNodeDetails.Get());
 
 		for (const auto& [Key, NodeWidget] : NodeWidgets)
 		{
@@ -1188,7 +1199,7 @@ namespace UE::Audio::Insights
 
 		TSharedPtr<SSignalFlowGraphNode> NodeWidget = SNew(SSignalFlowGraphNode)
 			.Entry(EntryNode->Entry)
-			.ListSource(ListSource)
+			.NodeDetailFilterSettings(NodeDetailFilterSettings)
 			.DisplayName(EntryNode->Entry.IsValid() ? EntryNode->Entry->GetDisplayName() : FText::GetEmpty())
 			.PixelSnappingMethod(EWidgetPixelSnapping::SnapToPixel)
 			.PositionInGraph_Lambda([this, EntryKey]() -> FVector2D
@@ -1196,9 +1207,11 @@ namespace UE::Audio::Insights
 				const TSharedPtr<SSignalFlowGraphNode>* Node = NodeWidgets.Find(EntryKey);
 				if (Node && Node->IsValid())
 				{
+					const int32 NumParamsVisible = NodeDetailFilterSettings ? NodeDetailFilterSettings->GetNumVisibleParams() : 0;
+
 					const FVector2D NodeSize = (*Node)->GetDesiredSize();
 					const float NodeGraphSlotWidth = GetNodeGraphSlotWidth(GetGraphOrientation());
-					const float NodeGraphSlotHeight = GetNodeGraphSlotHeight(GetGraphOrientation(), ShowNodeDetails.Get());
+					const float NodeGraphSlotHeight = GetNodeGraphSlotHeight(GetGraphOrientation(), NumParamsVisible, ShowNodeDetails.Get());
 
 					const float XOffset = (NodeGraphSlotWidth - NodeSize.X) / 2.0f;
 					const float YOffset = (NodeGraphSlotHeight - NodeSize.Y) / 2.0f;
diff --git a/Engine/Plugins/AudioInsights/Source/AudioInsights/Private/Views/SignalFlowDashboardViewFactory.cpp b/Engine/Plugins/AudioInsights/Source/AudioInsights/Private/Views/SignalFlowDashboardViewFactory.cpp
--- a/Engine/Plugins/AudioInsights/Source/AudioInsights/Private/Views/SignalFlowDashboardViewFactory.cpp
+++ b/Engine/Plugins/AudioInsights/Source/AudioInsights/Private/Views/SignalFlowDashboardViewFactory.cpp
@@ -40,6 +40,45 @@ namespace UE::Audio::Insights
 		constexpr float ToolbarButtonHeight = 20.0f;
 		constexpr float ToolbarButtonPaddingX = 8.0f;
 		constexpr float ToolbarButtonPaddingY = 8.0f;
+
+		FText GetParamFilterName(const ESignalFlowNodeDetailParam Param)
+		{
+			switch (Param)
+			{
+				case ESignalFlowNodeDetailParam::Amplitude:
+					return LOCTEXT("AudioDashboard_SignalFlow_Param_Amplitude", "Amplitude");
+
+				case ESignalFlowNodeDetailParam::Volume:
+					return LOCTEXT("AudioDashboard_SignalFlow_Param_Volume", "Volume");
+
+				case ESignalFlowNodeDetailParam::Pitch:
+					return LOCTEXT("AudioDashboard_SignalFlow_Param_Pitch", "Pitch");
+
+				case ESignalFlowNodeDetailParam::LPFFreq:
+					return LOCTEXT("AudioDashboard_SignalFlow_Param_LPFFreq", "LPF Freq");
+
+				case ESignalFlowNodeDetailParam::HPFFreq:
+					return LOCTEXT("AudioDashboard_SignalFlow_Param_HPFFreq", "HPF Freq");
+
+				case ESignalFlowNodeDetailParam::Priority:
+					return LOCTEXT("AudioDashboard_SignalFlow_Param_Priority", "Priority");;
+
+				case ESignalFlowNodeDetailParam::Distance:
+					return LOCTEXT("AudioDashboard_SignalFlow_Param_Distance", "Distance");;
+
+				case ESignalFlowNodeDetailParam::Attenuation:
+					return LOCTEXT("AudioDashboard_SignalFlow_Param_Attenuation", "Distance/Occlusion Attenuation");;
+
+				case ESignalFlowNodeDetailParam::RelativeRenderCost:
+					return LOCTEXT("AudioDashboard_SignalFlow_Param_RelativeRanderCost", "Rel. Render Cost");;
+
+				default:
+					ensureMsgf(false, TEXT("Unrecognized ESignalFlowNodeDetailParam value"));
+					break;
+			}
+
+			return FText::GetEmpty();
+		}
 	}
 
 	FReply SSignalFlowDashboard::OnKeyDown(const FGeometry& MyGeometry, const FKeyEvent& InKeyEvent)
@@ -169,12 +208,6 @@ namespace UE::Audio::Insights
 			{
 				return ScrollGraphViewToSelectedNode(true /*bResetZoom*/);
 			}
-			else if (InKeyEvent.GetKey() == EKeys::D)
-			{
-				ToggleShowNodeDetails();
-
-				return FReply::Handled();
-			}
 		}
 
 		return FReply::Unhandled();
@@ -239,6 +272,12 @@ namespace UE::Audio::Insights
 			FExecuteAction::CreateLambda([this]() { ToggleGraphOrientation(); }),
 			FCanExecuteAction(),
 			FIsActionChecked::CreateLambda([this]() { return GraphWidgetContents.IsValid() && GraphWidgetContents->GetGraphOrientation() == Orient_Horizontal; }));
+
+		CommandList->MapAction(
+			Commands.GetShowNodeDetailsCommand(),
+			FExecuteAction::CreateLambda([this]() { ToggleShowNodeDetails(); }),
+			FCanExecuteAction(),
+			FIsActionChecked::CreateLambda([this]() { return bShowNodeDetails; }));
 	}
 
 	TSharedRef<SWidget> FSignalFlowDashboardViewFactory::CreateToolbar()
@@ -419,6 +458,14 @@ namespace UE::Audio::Insights
 		MenuBuilder.BeginSection("SignalFlowSettingActions", LOCTEXT("AudioDashboard_SignalFlow_Settings_HeaderText", "Settings"));
 		{
 			MenuBuilder.AddMenuEntry(Commands.GetToggleHorizontalViewCommand());
+
+			MenuBuilder.AddMenuEntry(Commands.GetShowNodeDetailsCommand());
+			MenuBuilder.AddSubMenu
+			(
+				LOCTEXT("AudioDashboard_SignalFlow_Settings_NodeDetailFilters", "Node Detail Filters"),
+				LOCTEXT("AudioDashboard_SignalFlow_Settings_NodeDetailFiltersTooltip", "Choose which parameters are visible on the Signal Flow graph"),
+				FNewMenuDelegate::CreateRaw(this, &FSignalFlowDashboardViewFactory::BuildNodeDetailFiltersMenuContent)
+			);
 		}
 
 		MenuBuilder.EndSection();
@@ -426,6 +473,61 @@ namespace UE::Audio::Insights
 		return MenuBuilder.MakeWidget();
 	}
 
+	void FSignalFlowDashboardViewFactory::BuildNodeDetailFiltersMenuContent(FMenuBuilder& MenuBuilder)
+	{
+		using namespace FSignalFlowDashboardViewFactoryPrivate;
+
+		const FSignalFlowEditorCommands& Commands = FSignalFlowEditorCommands::Get();
+
+		MenuBuilder.AddMenuEntry
+		(
+			LOCTEXT("AudioDashboard_SignalFlow_Settings_ToggleAllNodeDetails", "Show all"),
+			FText::GetEmpty(),
+			FSlateIcon(),
+			FUIAction
+			(
+				FExecuteAction::CreateSP(this, &FSignalFlowDashboardViewFactory::ToggleEnableAllNodeDetailFilters),
+				FCanExecuteAction::CreateLambda([]() { return true; }),
+				FGetActionCheckState::CreateLambda([this]() { return NodeDetailFilters.GetShowAllNodeDetailFiltersCheckboxState(); })
+			),
+			NAME_None,
+			EUserInterfaceActionType::ToggleButton
+		);
+
+		const uint8 NumFilters = static_cast<uint8>(ESignalFlowNodeDetailParam::MAX);
+
+		for (uint8 FilterID = 0u; FilterID < NumFilters; ++FilterID)
+		{
+			const ESignalFlowNodeDetailParam ParamType = static_cast<ESignalFlowNodeDetailParam>(FilterID);
+
+			MenuBuilder.AddMenuEntry
+			(
+				GetParamFilterName(ParamType),
+				FText::GetEmpty(),
+				FSlateIcon(),
+				FUIAction
+				(
+					FExecuteAction::CreateLambda([this, ParamType]()
+					{
+						const bool bParamIsVisible = NodeDetailFilters.GetParameterIsVisible(ParamType);
+						NodeDetailFilters.SetParameterVisibility(ParamType, !bParamIsVisible);
+
+#if WITH_EDITOR
+						FSignalFlowSettings::OnRequestWriteSettings.Broadcast();
+#endif // WITH_EDITOR	
+					}),
+					FCanExecuteAction::CreateLambda([]() { return true; }),
+					FGetActionCheckState::CreateLambda([this, ParamType]()
+					{
+						return NodeDetailFilters.GetParameterIsVisible(ParamType) ? ECheckBoxState::Checked : ECheckBoxState::Unchecked;
+					})
+				),
+				NAME_None,
+				EUserInterfaceActionType::ToggleButton
+			);
+		}
+	}
+
 	FReply FSignalFlowDashboardViewFactory::ScrollGraphViewToSelectedNode(const bool bResetZoom)
 	{
 		if (GraphWidgetContents.IsValid())
@@ -477,6 +579,19 @@ namespace UE::Audio::Insights
 	{
 		bShowNodeDetails = !bShowNodeDetails;
 		OnRequestGraphRefresh();
+
+#if WITH_EDITOR
+		FSignalFlowSettings::OnRequestWriteSettings.Broadcast();
+#endif // WITH_EDITOR	
+	}
+
+	void FSignalFlowDashboardViewFactory::ToggleEnableAllNodeDetailFilters()
+	{
+		NodeDetailFilters.ToggleEnableAllNodeDetailFilters();
+
+#if WITH_EDITOR
+		FSignalFlowSettings::OnRequestWriteSettings.Broadcast();
+#endif // WITH_EDITOR	
 	}
 
 	TSharedPtr<FSignalFlowDashboardEntry> FSignalFlowDashboardViewFactory::GetFilteredAudioDeviceEntry() const
@@ -579,6 +694,7 @@ namespace UE::Audio::Insights
 				[
 					SAssignNew(GraphWidgetContents, SSignalFlowGraph)
 					.ListSource(&FilteredGraphNodes)
+					.NodeDetailFilterSettings(&NodeDetailFilters)
 					.SelectedItem_Lambda([this]()
 					{
 						return SelectedItem;
@@ -1583,6 +1699,9 @@ namespace UE::Audio::Insights
 				GraphWidgetContents->ToggleOrientation();
 			}
 		}
+
+		bShowNodeDetails = InSettings.bShowNodeDetails;
+		NodeDetailFilters = InSettings.NodeDetailFilters;
 	}
 
 	void FSignalFlowDashboardViewFactory::OnWriteEditorSettings(FSignalFlowSettings& OutSettings)
@@ -1592,6 +1711,9 @@ namespace UE::Audio::Insights
 			const bool bGraphIsHorizontal = GraphWidgetContents->GetGraphOrientation() == EOrientation::Orient_Horizontal;
 			OutSettings.bHorizontalFlow = bGraphIsHorizontal;
 		}
+
+		OutSettings.bShowNodeDetails = bShowNodeDetails;
+		OutSettings.NodeDetailFilters = NodeDetailFilters;
 	}
 #endif // WITH_EDITOR
 } // namespace UE::Audio::Insights
diff --git a/Engine/Plugins/AudioInsights/Source/AudioInsights/Public/Messages/SignalFlowEntryKey.h b/Engine/Plugins/AudioInsights/Source/AudioInsights/Public/Messages/SignalFlowEntryKey.h
--- a/Engine/Plugins/AudioInsights/Source/AudioInsights/Public/Messages/SignalFlowEntryKey.h
+++ b/Engine/Plugins/AudioInsights/Source/AudioInsights/Public/Messages/SignalFlowEntryKey.h
@@ -33,7 +33,10 @@ namespace UE::Audio::Insights
 		Priority,
 		Distance,
 		Attenuation,
-		RelativeRenderCost
+		RelativeRenderCost,
+
+		///
+		MAX
 	};
 
 	struct FSignalFlowEntryKey
diff --git a/Engine/Plugins/AudioInsights/Source/AudioInsights/Public/Settings/SignalFlowSettings.h b/Engine/Plugins/AudioInsights/Source/AudioInsights/Public/Settings/SignalFlowSettings.h
--- a/Engine/Plugins/AudioInsights/Source/AudioInsights/Public/Settings/SignalFlowSettings.h
+++ b/Engine/Plugins/AudioInsights/Source/AudioInsights/Public/Settings/SignalFlowSettings.h
@@ -4,11 +4,62 @@
 
 #include "Delegates/DelegateCombinations.h"
 #include "Engine/DeveloperSettings.h"
+#include "Messages/SignalFlowEntryKey.h"
+#include "Styling/SlateTypes.h"
 
 #include "SignalFlowSettings.generated.h"
 
 #define LOCTEXT_NAMESPACE "AudioInsights"
 
+USTRUCT()
+struct FSignalFlowNodeDetailFilterSettings
+{
+	GENERATED_BODY()
+
+	bool GetParameterIsVisible(const UE::Audio::Insights::ESignalFlowNodeDetailParam Param) const;
+	int32 GetNumVisibleParams() const;
+	ECheckBoxState GetShowAllNodeDetailFiltersCheckboxState() const;
+
+	void ToggleEnableAllNodeDetailFilters();
+	void SetParameterVisibility(const UE::Audio::Insights::ESignalFlowNodeDetailParam Param, const bool bIsVisible);
+
+	/** Toggles visibility of amplitude parameters. */
+	UPROPERTY(EditAnywhere, config, Category = SignalFlow)
+	bool bAmplitude = true;
+
+	/** Toggles visibility of volume parameters. */
+	UPROPERTY(EditAnywhere, config, Category = SignalFlow)
+	bool bVolume = true;
+
+	/** Toggles visibility of pitch parameters. */
+	UPROPERTY(EditAnywhere, config, Category = SignalFlow)
+	bool bPitch = false;
+
+	/** Toggles visibility of low pass filter frequency parameters. */
+	UPROPERTY(EditAnywhere, config, Category = SignalFlow)
+	bool bLPFFreq = true;
+
+	/** Toggles visibility of high pass filter frequency parameters. */
+	UPROPERTY(EditAnywhere, config, Category = SignalFlow)
+	bool bHPFFreq = true;
+
+	/** Toggles visibility of priority parameters. */
+	UPROPERTY(EditAnywhere, config, Category = SignalFlow)
+	bool bPriority = false;
+
+	/** Toggles visibility of distance parameters. */
+	UPROPERTY(EditAnywhere, config, Category = SignalFlow)
+	bool bDistance = false;
+
+	/** Toggles visibility of distance/occlusion attenuation parameters. */
+	UPROPERTY(EditAnywhere, config, Category = SignalFlow)
+	bool bDistanceOcclusionAttenuation = false;
+
+	/** Toggles visibility of relative render cost parameters. */
+	UPROPERTY(EditAnywhere, config, Category = SignalFlow)
+	bool bRelativeRenderCost = false;
+};
+
 USTRUCT()
 struct FSignalFlowSettings
 {
@@ -18,6 +69,14 @@ struct FSignalFlowSettings
 	UPROPERTY(EditAnywhere, config, Category = SignalFlow)
 	bool bHorizontalFlow = true;
 
+	/** Toggles visibility of node details. */
+	UPROPERTY(EditAnywhere, config, Category = SignalFlow)
+	bool bShowNodeDetails = false;
+
+	/** Which parameter types are visible in the Signal Flow graph when Node Details is enabled */
+	UPROPERTY(EditAnywhere, config, Category = SignalFlow)
+	FSignalFlowNodeDetailFilterSettings NodeDetailFilters;
+
 #if WITH_EDITOR
 	DECLARE_MULTICAST_DELEGATE_OneParam(FOnReadSignalFlowSettings, const FSignalFlowSettings&);
 	static AUDIOINSIGHTS_API FOnReadSignalFlowSettings OnReadSettings;
diff --git a/Engine/Plugins/AudioInsights/Source/AudioInsights/Public/Views/SSignalFlowGraph.h b/Engine/Plugins/AudioInsights/Source/AudioInsights/Public/Views/SSignalFlowGraph.h
--- a/Engine/Plugins/AudioInsights/Source/AudioInsights/Public/Views/SSignalFlowGraph.h
+++ b/Engine/Plugins/AudioInsights/Source/AudioInsights/Public/Views/SSignalFlowGraph.h
@@ -8,6 +8,8 @@
 
 #define UE_API AUDIOINSIGHTS_API
 
+struct FSignalFlowNodeDetailFilterSettings;
+
 class SBorder;
 class SHorizontalBox;
 class SScrollBox;
@@ -30,7 +32,7 @@ namespace UE::Audio::Insights
 	public:
 		SLATE_BEGIN_ARGS(SSignalFlowGraphNode)
 			: _Entry(nullptr),
-			_ListSource(nullptr),
+			_NodeDetailFilterSettings(nullptr),
 			_DisplayName(),
 			_PositionInGraph(FVector2d::ZeroVector),
 			_XPos(0.0f),
@@ -43,7 +45,7 @@ namespace UE::Audio::Insights
 		}
 
 		SLATE_ARGUMENT(TSharedPtr<FSignalFlowDashboardEntry>, Entry)
-		SLATE_ARGUMENT(TArray<TSharedPtr<ISignalFlowNode>>*, ListSource)
+		SLATE_ARGUMENT(FSignalFlowNodeDetailFilterSettings*, NodeDetailFilterSettings)
 		SLATE_ARGUMENT(FText, DisplayName)
 		SLATE_ATTRIBUTE(FVector2d, PositionInGraph)
 		SLATE_ATTRIBUTE(float, XPos)
@@ -94,6 +96,8 @@ namespace UE::Audio::Insights
 		FVector2d CachedPositionInGraph = FVector2d::Zero();
 		FSignalFlowEntryKey NodeEntryKey;
 
+		FSignalFlowNodeDetailFilterSettings* NodeDetailFilterSettings = nullptr;
+
 		TSharedPtr<SBorder> SelectedNodeBorder;
 		TSharedPtr<SBorder> NodeBorder;
 	};
@@ -146,11 +150,13 @@ namespace UE::Audio::Insights
 	public:
 		SLATE_BEGIN_ARGS(SSignalFlowGraph)
 			: _ListSource(nullptr)
+			, _NodeDetailFilterSettings(nullptr)
 			, _ShowNodeDetails(false)
 		{
 		}
 
 		SLATE_ARGUMENT(TArray<TSharedPtr<ISignalFlowNode>>*, ListSource)
+		SLATE_ARGUMENT(FSignalFlowNodeDetailFilterSettings*, NodeDetailFilterSettings)
 
 		SLATE_ATTRIBUTE(TSharedPtr<FSignalFlowDashboardEntry>, SelectedItem)
 		SLATE_ATTRIBUTE(TSharedPtr<FSignalFlowDashboardEntry>, ActiveAudioDeviceEntry)
@@ -230,7 +236,9 @@ namespace UE::Audio::Insights
 		TMap<FSignalFlowEntryKey, ConnectedNodeKeys> DummyToRealConnectionKeys;
 		TMultiMap<FSignalFlowEntryKey, TSharedPtr<SSignalFlowGraphSendLabel>> SendLevelWidgets;
 
-		TArray<TSharedPtr<ISignalFlowNode>>* ListSource;
+		TArray<TSharedPtr<ISignalFlowNode>>* ListSource = nullptr;
+		FSignalFlowNodeDetailFilterSettings* NodeDetailFilterSettings = nullptr;
+
 		TAttribute<TSharedPtr<FSignalFlowDashboardEntry>> SelectedItem;
 		TAttribute<TSharedPtr<FSignalFlowDashboardEntry>> ActiveAudioDeviceEntry;
 		TAttribute<bool> ShowNodeDetails;
diff --git a/Engine/Plugins/AudioInsights/Source/AudioInsights/Public/Views/SignalFlowDashboardViewFactory.h b/Engine/Plugins/AudioInsights/Source/AudioInsights/Public/Views/SignalFlowDashboardViewFactory.h
--- a/Engine/Plugins/AudioInsights/Source/AudioInsights/Public/Views/SignalFlowDashboardViewFactory.h
+++ b/Engine/Plugins/AudioInsights/Source/AudioInsights/Public/Views/SignalFlowDashboardViewFactory.h
@@ -73,10 +73,12 @@ namespace UE::Audio::Insights
 		TSharedRef<SWidget> CreateToolbarButtonActions();
 
 		TSharedRef<SWidget> OnGetSettingsMenuContent();
+		void BuildNodeDetailFiltersMenuContent(FMenuBuilder& MenuBuilder);
 
 		FReply ScrollGraphViewToSelectedNode(const bool bResetZoom);
 		void ToggleGraphOrientation();
 		void ToggleShowNodeDetails();
+		void ToggleEnableAllNodeDetailFilters();
 
 		TSharedPtr<FSignalFlowDashboardEntry> GetFilteredAudioDeviceEntry() const;
 
@@ -148,6 +150,8 @@ namespace UE::Audio::Insights
 		TSharedPtr<FSignalFlowDashboardEntry> SelectedItem;
 		FText SearchBoxFilterText;
 
+		FSignalFlowNodeDetailFilterSettings NodeDetailFilters;
+
 		bool bGraphRequiresRefresh = false;
 		bool bAutoFocusAfterGraphRefresh = false;
 		bool bShowNodeDetails = false;
