[Audio Insights] Add option to display amplitude values in dB/Linear inside the Signal Flow graph

#Synced-CL 50395601

diff --git a/Engine/Plugins/AudioInsights/Source/AudioInsights/Private/Views/SSignalFlowGraph.cpp b/Engine/Plugins/AudioInsights/Source/AudioInsights/Private/Views/SSignalFlowGraph.cpp
--- a/Engine/Plugins/AudioInsights/Source/AudioInsights/Private/Views/SSignalFlowGraph.cpp
+++ b/Engine/Plugins/AudioInsights/Source/AudioInsights/Private/Views/SSignalFlowGraph.cpp
@@ -3,6 +3,7 @@
 #include "Views/SSignalFlowGraph.h"
 
 #include "AudioInsightsStyle.h"
+#include "AudioInsightsUtils.h"
 #include "Brushes/SlateRoundedBoxBrush.h"
 #include "Internationalization/Text.h"
 #include "Messages/SignalFlowTraceMessages.h"
@@ -235,7 +236,7 @@ namespace UE::Audio::Insights
 			switch (ParamType)
 			{
 				case ESignalFlowNodeDetailParam::Amplitude:
-					return FSlateStyle::Get().GetLinearVolumeFloatFormat();
+					return FSlateStyle::Get().GetLinearAmplitudeFloatFormat();
 
 				case ESignalFlowNodeDetailParam::Volume:
 					return FSlateStyle::Get().GetLinearVolumeFloatFormat();
@@ -282,6 +283,7 @@ namespace UE::Audio::Insights
 		ZoomScaleFactor = InArgs._ZoomScaleFactor;
 		IsSelected = InArgs._IsSelected;
 		ShowNodeDetails = InArgs._ShowNodeDetails;
+		DisplayAmpPeakInDb = InArgs._DisplayAmpPeakInDb;
 		GraphOrientation = InArgs._GraphOrientation;
 		OnNodeSelected = InArgs._OnNodeSelected;
 
@@ -565,6 +567,11 @@ namespace UE::Audio::Insights
 							return FText::GetEmpty();
 						}
 
+						if (ParamType == ESignalFlowNodeDetailParam::Amplitude && DisplayAmpPeakInDb.Get() == true)
+						{
+							return FText::Format(GetNodeDetailDisplayName(ParamType), AudioInsightsUtils::ConvertToDecibelsText(ParamValueOpt.GetValue()));
+						}
+
 						if (ParamType == ESignalFlowNodeDetailParam::Priority && SoundMessageUtils::IsPriorityValueSetToMax(ParamValueOpt.GetValue()))
 						{
 							return FText::Format(GetNodeDetailDisplayName(ParamType), LOCTEXT("SignalFlowGraph_NodeDetails_Max", "MAX"));
@@ -684,6 +691,7 @@ namespace UE::Audio::Insights
 		HorizontalScrollBox = InArgs._HorizontalScrollBox;
 		VerticalScrollBox = InArgs._VerticalScrollBox;
 		ShowNodeDetails = InArgs._ShowNodeDetails;
+		DisplayAmpPeakInDb = InArgs._DisplayAmpPeakInDb;
 		OnNodeSelected = InArgs._OnNodeSelected;
 
 		Clipping = EWidgetClipping::ClipToBounds;
@@ -1246,6 +1254,10 @@ namespace UE::Audio::Insights
 			{
 				return ShowNodeDetails.Get();
 			})
+			.DisplayAmpPeakInDb_Lambda([this]()
+			{
+				return DisplayAmpPeakInDb.Get();
+			})
 			.GraphOrientation_Lambda([this]()
 			{
 				return Orientation;
diff --git a/Engine/Plugins/AudioInsights/Source/AudioInsights/Private/Views/SignalFlowDashboardViewFactory.cpp b/Engine/Plugins/AudioInsights/Source/AudioInsights/Private/Views/SignalFlowDashboardViewFactory.cpp
--- a/Engine/Plugins/AudioInsights/Source/AudioInsights/Private/Views/SignalFlowDashboardViewFactory.cpp
+++ b/Engine/Plugins/AudioInsights/Source/AudioInsights/Private/Views/SignalFlowDashboardViewFactory.cpp
@@ -14,6 +14,7 @@
 #include "Widgets/Colors/SColorBlock.h"
 #include "Widgets/Input/SButton.h"
 #include "Widgets/Input/SComboButton.h"
+#include "Widgets/Input/SSegmentedControl.h"
 #include "Widgets/Layout/SSeparator.h"
 
 #if WITH_EDITOR
@@ -457,6 +458,8 @@ namespace UE::Audio::Insights
 
 		MenuBuilder.BeginSection("SignalFlowSettingActions", LOCTEXT("AudioDashboard_SignalFlow_Settings_HeaderText", "Settings"));
 		{
+			CreateAmpDisplayModeToggle(MenuBuilder);
+
 			MenuBuilder.AddMenuEntry(Commands.GetToggleHorizontalViewCommand());
 
 			MenuBuilder.AddMenuEntry(Commands.GetShowNodeDetailsCommand());
@@ -528,6 +531,60 @@ namespace UE::Audio::Insights
 		}
 	}
 
+	void FSignalFlowDashboardViewFactory::CreateAmpDisplayModeToggle(FMenuBuilder& OutMenuBuilder)
+	{
+		OutMenuBuilder.AddMenuEntry(
+			FUIAction(FExecuteAction::CreateLambda([this]()
+			{
+				bDisplayAmpPeakInDb = !bDisplayAmpPeakInDb;
+#if WITH_EDITOR
+				FSignalFlowSettings::OnRequestWriteSettings.Broadcast();
+#endif // WITH_EDITOR
+			})),
+
+			SNew(SHorizontalBox)
+
+			+ SHorizontalBox::Slot()
+			.AutoWidth()
+			[
+				SNew(STextBlock)
+				.Margin(FMargin(0.0, 4.0, 5.0, 0.0))
+				.Text(LOCTEXT("AudioDashboard_SignalFlow_Settings_AmpPeakDisplayMode_Text", "Amp (Peak) Display Mode"))
+			]
+
+			+ SHorizontalBox::Slot()
+			.Padding(4.0f, 0.0f, 0.0f, 0.0f)
+			.AutoWidth()
+			[
+				SNew(SSegmentedControl<bool>)
+				.OnValueChanged_Lambda([this](bool bInValue)
+				{
+					bDisplayAmpPeakInDb = bInValue;
+#if WITH_EDITOR
+					FSignalFlowSettings::OnRequestWriteSettings.Broadcast();
+#endif // WITH_EDITOR
+				})
+				.Value_Lambda([this]
+				{
+					return bDisplayAmpPeakInDb;
+				})
+					
+				// Decibels mode
+				+ SSegmentedControl<bool>::Slot(true)
+				.Text(LOCTEXT("AudioDashboard_SignalFlow_Settings_AmpDisplayMode_dB_Text", "dB"))
+				.ToolTip(LOCTEXT("AudioDashboard_SignalFlow_Settings_AmpDisplayMode_dB_ToolTipText", "Displays amplitude values in decibels."))
+				
+				// Linear mode
+				+ SSegmentedControl<bool>::Slot(false)
+				.Text(LOCTEXT("AudioDashboard_SignalFlow_Settings_AmpDisplayMode_Linear_Text", "Lin"))
+				.ToolTip(LOCTEXT("AudioDashboard_SignalFlow_Settings_AmpDisplayMode_Linear_TooltipText", "Displays amplitude values in linear scale."))
+			],
+			NAME_None /*InExtensionHook*/,
+			LOCTEXT("AudioDashboard_SignalFlow_Settings_AmpPeakDisplayMode_ToolTipText", "Switches the amplitude display mode between decibels or linear scale."),
+			EUserInterfaceActionType::CollapsedButton
+		);
+	}
+
 	FReply FSignalFlowDashboardViewFactory::ScrollGraphViewToSelectedNode(const bool bResetZoom)
 	{
 		if (GraphWidgetContents.IsValid())
@@ -755,6 +812,10 @@ namespace UE::Audio::Insights
 					{
 						return bShowNodeDetails;
 					})
+					.DisplayAmpPeakInDb_Lambda([this]()
+					{
+						return bDisplayAmpPeakInDb;
+					})
 				]
 			]
 		];
@@ -1702,6 +1763,7 @@ namespace UE::Audio::Insights
 
 		bShowNodeDetails = InSettings.bShowNodeDetails;
 		NodeDetailFilters = InSettings.NodeDetailFilters;
+		bDisplayAmpPeakInDb = InSettings.AmplitudeDisplayMode == EAudioAmplitudeDisplayMode::Decibels;
 	}
 
 	void FSignalFlowDashboardViewFactory::OnWriteEditorSettings(FSignalFlowSettings& OutSettings)
@@ -1714,6 +1776,7 @@ namespace UE::Audio::Insights
 
 		OutSettings.bShowNodeDetails = bShowNodeDetails;
 		OutSettings.NodeDetailFilters = NodeDetailFilters;
+		OutSettings.AmplitudeDisplayMode = bDisplayAmpPeakInDb ? EAudioAmplitudeDisplayMode::Decibels : EAudioAmplitudeDisplayMode::Linear;
 	}
 #endif // WITH_EDITOR
 } // namespace UE::Audio::Insights
diff --git a/Engine/Plugins/AudioInsights/Source/AudioInsights/Public/AudioInsightsStyle.h b/Engine/Plugins/AudioInsights/Source/AudioInsights/Public/AudioInsightsStyle.h
--- a/Engine/Plugins/AudioInsights/Source/AudioInsights/Public/AudioInsightsStyle.h
+++ b/Engine/Plugins/AudioInsights/Source/AudioInsights/Public/AudioInsightsStyle.h
@@ -36,6 +36,16 @@ namespace UE::Audio::Insights
 			return &FloatFormat;
 		};
 
+		const FNumberFormattingOptions* GetLinearAmplitudeFloatFormat() const
+		{
+			static FNumberFormattingOptions FloatFormat;
+			FloatFormat.MinimumIntegralDigits = 1;
+			FloatFormat.MaximumIntegralDigits = 3;
+			FloatFormat.MinimumFractionalDigits = 5;
+			FloatFormat.MaximumFractionalDigits = 5;
+			return &FloatFormat;
+		};
+
 		const FNumberFormattingOptions* GetLinearVolumeFloatFormat() const
 		{
 			static FNumberFormattingOptions FloatFormat;
diff --git a/Engine/Plugins/AudioInsights/Source/AudioInsights/Public/Settings/SignalFlowSettings.h b/Engine/Plugins/AudioInsights/Source/AudioInsights/Public/Settings/SignalFlowSettings.h
--- a/Engine/Plugins/AudioInsights/Source/AudioInsights/Public/Settings/SignalFlowSettings.h
+++ b/Engine/Plugins/AudioInsights/Source/AudioInsights/Public/Settings/SignalFlowSettings.h
@@ -5,6 +5,7 @@
 #include "Delegates/DelegateCombinations.h"
 #include "Engine/DeveloperSettings.h"
 #include "Messages/SignalFlowEntryKey.h"
+#include "Settings/SoundDashboardSettings.h"
 #include "Styling/SlateTypes.h"
 
 #include "SignalFlowSettings.generated.h"
@@ -65,6 +66,10 @@ struct FSignalFlowSettings
 {
 	GENERATED_BODY()
 
+	/** Toggles between decibel/linear scales when displaying amplitude (peak) values inside the signal flow graph. */
+	UPROPERTY(EditAnywhere, config, Category = SignalFlow, meta = (DisplayName = "Amp (Peak) Display Mode"))
+	EAudioAmplitudeDisplayMode AmplitudeDisplayMode = EAudioAmplitudeDisplayMode::Decibels;
+
 	/** Toggles between horizontally/vertically orienting the Signal Flow graph. */
 	UPROPERTY(EditAnywhere, config, Category = SignalFlow)
 	bool bHorizontalFlow = true;
diff --git a/Engine/Plugins/AudioInsights/Source/AudioInsights/Public/Views/SSignalFlowGraph.h b/Engine/Plugins/AudioInsights/Source/AudioInsights/Public/Views/SSignalFlowGraph.h
--- a/Engine/Plugins/AudioInsights/Source/AudioInsights/Public/Views/SSignalFlowGraph.h
+++ b/Engine/Plugins/AudioInsights/Source/AudioInsights/Public/Views/SSignalFlowGraph.h
@@ -40,6 +40,7 @@ namespace UE::Audio::Insights
 			_ZoomScaleFactor(1.0f),
 			_IsSelected(false),
 			_ShowNodeDetails(false),
+			_DisplayAmpPeakInDb(true),
 			_GraphOrientation(EOrientation::Orient_Vertical)
 		{
 		}
@@ -53,6 +54,7 @@ namespace UE::Audio::Insights
 		SLATE_ATTRIBUTE(float, ZoomScaleFactor)
 		SLATE_ATTRIBUTE(bool, IsSelected)
 		SLATE_ATTRIBUTE(bool, ShowNodeDetails)
+		SLATE_ATTRIBUTE(bool, DisplayAmpPeakInDb)
 		SLATE_ATTRIBUTE(EOrientation, GraphOrientation)
 		SLATE_EVENT(FOnSignalFlowNodeSelected, OnNodeSelected)
 
@@ -90,6 +92,7 @@ namespace UE::Audio::Insights
 		TAttribute<float> ZoomScaleFactor;
 		TAttribute<bool> IsSelected;
 		TAttribute<bool> ShowNodeDetails;
+		TAttribute<bool> DisplayAmpPeakInDb;
 		TAttribute<EOrientation> GraphOrientation;
 
 		FOnSignalFlowNodeSelected OnNodeSelected;
@@ -152,6 +155,7 @@ namespace UE::Audio::Insights
 			: _ListSource(nullptr)
 			, _NodeDetailFilterSettings(nullptr)
 			, _ShowNodeDetails(false)
+			, _DisplayAmpPeakInDb(true)
 		{
 		}
 
@@ -163,6 +167,7 @@ namespace UE::Audio::Insights
 		SLATE_ATTRIBUTE(TWeakPtr<SScrollBox>, HorizontalScrollBox)
 		SLATE_ATTRIBUTE(TWeakPtr<SScrollBox>, VerticalScrollBox)
 		SLATE_ATTRIBUTE(bool, ShowNodeDetails)
+		SLATE_ATTRIBUTE(bool, DisplayAmpPeakInDb)
 
 		SLATE_EVENT(FOnSignalFlowNodeSelected, OnNodeSelected)
 
@@ -242,6 +247,7 @@ namespace UE::Audio::Insights
 		TAttribute<TSharedPtr<FSignalFlowDashboardEntry>> SelectedItem;
 		TAttribute<TSharedPtr<FSignalFlowDashboardEntry>> ActiveAudioDeviceEntry;
 		TAttribute<bool> ShowNodeDetails;
+		TAttribute<bool> DisplayAmpPeakInDb;
 
 		TAttribute<TWeakPtr<SScrollBox>> HorizontalScrollBox;
 		TAttribute<TWeakPtr<SScrollBox>> VerticalScrollBox;
diff --git a/Engine/Plugins/AudioInsights/Source/AudioInsights/Public/Views/SignalFlowDashboardViewFactory.h b/Engine/Plugins/AudioInsights/Source/AudioInsights/Public/Views/SignalFlowDashboardViewFactory.h
--- a/Engine/Plugins/AudioInsights/Source/AudioInsights/Public/Views/SignalFlowDashboardViewFactory.h
+++ b/Engine/Plugins/AudioInsights/Source/AudioInsights/Public/Views/SignalFlowDashboardViewFactory.h
@@ -74,6 +74,7 @@ namespace UE::Audio::Insights
 
 		TSharedRef<SWidget> OnGetSettingsMenuContent();
 		void BuildNodeDetailFiltersMenuContent(FMenuBuilder& MenuBuilder);
+		void CreateAmpDisplayModeToggle(FMenuBuilder& OutMenuBuilder);
 
 		FReply ScrollGraphViewToSelectedNode(const bool bResetZoom);
 		void ToggleGraphOrientation();
@@ -155,6 +156,7 @@ namespace UE::Audio::Insights
 		bool bGraphRequiresRefresh = false;
 		bool bAutoFocusAfterGraphRefresh = false;
 		bool bShowNodeDetails = false;
+		bool bDisplayAmpPeakInDb = true;
 	};
 } // namespace UE::Audio::Insights
 
