[Audio Insights] Show signal strength between nodes in connecting wire animations (Signal Flow)
* Added new Animate Wires setting to enable/disable the effect
* Animations are hooked up to the cache and scrub systems
* Signal strength roughly based off metasound method
* Strength of color and signal strength thickness of wire is controllable via WireSplineAmplitudePowerFactor in SignalFlowSettings, allowing quieter sounds to be more visible in the wires.

#Synced-CL 50435572

diff --git a/Engine/Plugins/AudioInsights/Source/AudioInsights/Private/Messages/SignalFlowTraceMessages.cpp b/Engine/Plugins/AudioInsights/Source/AudioInsights/Private/Messages/SignalFlowTraceMessages.cpp
--- a/Engine/Plugins/AudioInsights/Source/AudioInsights/Private/Messages/SignalFlowTraceMessages.cpp
+++ b/Engine/Plugins/AudioInsights/Source/AudioInsights/Private/Messages/SignalFlowTraceMessages.cpp
@@ -10,6 +10,11 @@ namespace UE::Audio::Insights
 		const FName SubmixToSubmixSend = "SubmixToSubmixSend";
 	};
 
+	namespace FSignalFlowDashboardEntryPrivate
+	{
+		constexpr uint32 DataPointsCapacity = 128u;
+	}
+
 	FSignalFlowMessageBase::FSignalFlowMessageBase(const Trace::IAnalyzer::FOnEventContext& InContext)
 	{
 		const Trace::IAnalyzer::FEventData& EventData = InContext.EventData;
@@ -58,4 +63,25 @@ namespace UE::Audio::Insights
 	{
 		return sizeof(FSubmixToSubmixSendMessage);
 	}
+
+	FSignalFlowDashboardEntry::FSignalFlowDashboardEntry()
+		: AmplitudeDataRange(FSignalFlowDashboardEntryPrivate::DataPointsCapacity, 0.0f)
+	{
+	}
+
+	void FSignalFlowDashboardEntry::WriteNewAmplitudeValue(const float Value)
+	{
+		AmplitudeDataRange[AmplitudeWriteIndex] = Value;
+		AmplitudeWriteIndex = AmplitudeDataRange.GetNextIndex(AmplitudeWriteIndex);
+	}
+
+	void FSignalFlowDashboardEntry::ResetDataBuffers()
+	{
+		for (uint32 Index = 0u; Index < AmplitudeDataRange.Capacity(); ++Index)
+		{
+			AmplitudeDataRange[Index] = 0.0f;
+		}
+
+		AmplitudeWriteIndex = 0;
+	}
 } // namespace UE::Audio::Insights
diff --git a/Engine/Plugins/AudioInsights/Source/AudioInsights/Private/Providers/SignalFlowTraceProvider.cpp b/Engine/Plugins/AudioInsights/Source/AudioInsights/Private/Providers/SignalFlowTraceProvider.cpp
--- a/Engine/Plugins/AudioInsights/Source/AudioInsights/Private/Providers/SignalFlowTraceProvider.cpp
+++ b/Engine/Plugins/AudioInsights/Source/AudioInsights/Private/Providers/SignalFlowTraceProvider.cpp
@@ -513,6 +513,7 @@ namespace UE::Audio::Insights
 			if (OutEntry && OutEntry->IsValid())
 			{
 				(*OutEntry)->Amplitude = Msg.Envelope;
+				(*OutEntry)->WriteNewAmplitudeValue(Msg.Envelope);
 			}
 		});
 
@@ -838,13 +839,17 @@ namespace UE::Audio::Insights
 			ActiveEntriesArray.Add(*Itr);
 		}
 
+		// Start by gathering all parameters that are represented as single values
+		// Only amplitude envelopes require a range, and only sound sources have more parameters than just amplitude
 		ParallelFor(ActiveEntriesArray.Num(), [&CacheManager, TimeMarker, &ActiveEntriesArray, this](const int32 Index)
 		{
 			const FSignalFlowEntryKey& EntryKey = ActiveEntriesArray[Index];
 
-			if (EntryKey.EntryType != ESignalFlowEntryType::SoundSource
-				&& EntryKey.EntryType != ESignalFlowEntryType::AudioBus
-				&& EntryKey.EntryType != ESignalFlowEntryType::Submix)
+			const bool bIsEntryWithParam = EntryKey.EntryType == ESignalFlowEntryType::SoundSource
+										|| (!bAnimateWires && (EntryKey.EntryType == ESignalFlowEntryType::AudioBus 
+															|| EntryKey.EntryType == ESignalFlowEntryType::Submix));
+
+			if (!bIsEntryWithParam)
 			{
 				return;
 			}
@@ -855,19 +860,13 @@ namespace UE::Audio::Insights
 				return;
 			}
 
-			if (EntryKey.EntryType == ESignalFlowEntryType::SoundSource)
+			auto MatchingParamPredicate = [&EntryKey](const FSoundParameterMessage& Msg)
 			{
-				auto MatchingParamPredicate = [&EntryKey](const FSoundParameterMessage& Msg)
-				{
-					return Msg.WaveInstancePlayOrder == EntryKey.EntryID;
-				};
-
-				const FSoundEnvelopeMessage* EnvelopeMessage = CacheManager.FindClosestMessageByPredicate<FSoundEnvelopeMessage>(SoundMessageNames::EnvelopeParam, TimeMarker, MatchingParamPredicate);
-				if (EnvelopeMessage)
-				{
-					(*Entry)->Amplitude = EnvelopeMessage->Envelope;
-				}
+				return Msg.WaveInstancePlayOrder == EntryKey.EntryID;
+			};
 
+			if (EntryKey.EntryType == ESignalFlowEntryType::SoundSource)
+			{
 				const FSoundVolumeMessage* VolumeMessage = CacheManager.FindClosestMessageByPredicate<FSoundVolumeMessage>(SoundMessageNames::VolumeParam, TimeMarker, MatchingParamPredicate);
 				if (VolumeMessage)
 				{
@@ -916,23 +915,99 @@ namespace UE::Audio::Insights
 					(*Entry)->RelativeRenderCost = RelativeRenderCostMessage->RelativeRenderCost;
 				}
 			}
-			else if (EntryKey.EntryType == ESignalFlowEntryType::AudioBus)
+
+			// If we are not animating wires, grab the latest envelope value here
+			if (!bAnimateWires)
 			{
-				const FAudioBusEnvelopeValuesMessage* AudioBusEnvelopeMessage = CacheManager.FindClosestMessage<FAudioBusEnvelopeValuesMessage>(AudioBusMessageNames::EnvelopeValues, TimeMarker, EntryKey.EntryID);
-				if (AudioBusEnvelopeMessage)
+				if (EntryKey.EntryType == ESignalFlowEntryType::SoundSource)
 				{
-					ProcessMultiChannelEnvelopeMessage(Entry, AudioBusEnvelopeMessage->EnvelopeValues);
+					const FSoundEnvelopeMessage* EnvelopeMessage = CacheManager.FindClosestMessageByPredicate<FSoundEnvelopeMessage>(SoundMessageNames::EnvelopeParam, TimeMarker, MatchingParamPredicate);
+					if (EnvelopeMessage)
+					{
+						(*Entry)->Amplitude = EnvelopeMessage->Envelope;
+					}
 				}
-			}
-			else if (EntryKey.EntryType == ESignalFlowEntryType::Submix)
-			{
-				const FSubmixEnvelopeValuesMessage* SubmixEnvelopeMessage = CacheManager.FindClosestMessage<FSubmixEnvelopeValuesMessage>(SubmixMessageNames::EnvelopeValues, TimeMarker, EntryKey.EntryID);
-				if (SubmixEnvelopeMessage)
+				else if (EntryKey.EntryType == ESignalFlowEntryType::AudioBus)
+				{
+					const FAudioBusEnvelopeValuesMessage* AudioBusEnvelopeMessage = CacheManager.FindClosestMessage<FAudioBusEnvelopeValuesMessage>(AudioBusMessageNames::EnvelopeValues, TimeMarker, EntryKey.EntryID);
+					if (AudioBusEnvelopeMessage)
+					{
+						ProcessMultiChannelEnvelopeMessage(Entry, AudioBusEnvelopeMessage->EnvelopeValues);
+					}
+				}
+				else if (EntryKey.EntryType == ESignalFlowEntryType::Submix)
 				{
-					ProcessMultiChannelEnvelopeMessage(Entry, SubmixEnvelopeMessage->EnvelopeValues);
+					const FSubmixEnvelopeValuesMessage* SubmixEnvelopeMessage = CacheManager.FindClosestMessage<FSubmixEnvelopeValuesMessage>(SubmixMessageNames::EnvelopeValues, TimeMarker, EntryKey.EntryID);
+					if (SubmixEnvelopeMessage)
+					{
+						ProcessMultiChannelEnvelopeMessage(Entry, SubmixEnvelopeMessage->EnvelopeValues);
+					}
 				}
 			}
+			
 		});
+
+		if (bAnimateWires)
+		{
+			// Next, gather an amplitude range large enough to cover the wire animation
+			constexpr float AmplitudeLookbackTime = 2.0f;
+			const float StartTime = TimeMarker - AmplitudeLookbackTime;
+
+			constexpr int32 NumMessagesToProcess = 3; // we have 3 different types of amplitude message we want to retrieve a range for from the cache
+
+			ParallelFor(NumMessagesToProcess, [&CacheManager, TimeMarker, StartTime, this](const int32 Index)
+			{
+				switch (Index)
+				{
+					case 0: // FSoundEnvelopeMessage
+					{
+						CacheManager.IterateOverRange<FSoundEnvelopeMessage>(SoundMessageNames::EnvelopeParam, StartTime, TimeMarker, [this](const FSoundEnvelopeMessage& Message)
+						{
+							const FSignalFlowEntryKey EntryKey(Message.DeviceId, ESignalFlowEntryType::SoundSource, Message.WaveInstancePlayOrder);
+							TSharedPtr<FSignalFlowDashboardEntry>* Entry = FindDeviceEntry(Message.DeviceId, EntryKey);
+
+							if (Entry && Entry->IsValid())
+							{
+								(*Entry)->Amplitude = Message.Envelope;
+								(*Entry)->WriteNewAmplitudeValue(Message.Envelope);
+							}
+						});
+
+						break;
+					}
+
+					case 1: // FAudioBusEnvelopeValuesMessage
+					{
+						CacheManager.IterateOverRange<FAudioBusEnvelopeValuesMessage>(AudioBusMessageNames::EnvelopeValues, StartTime, TimeMarker, [this](const FAudioBusEnvelopeValuesMessage& Message)
+						{
+							const FSignalFlowEntryKey EntryKey(Message.DeviceId, ESignalFlowEntryType::AudioBus, Message.AudioBusId);
+							TSharedPtr<FSignalFlowDashboardEntry>* Entry = FindDeviceEntry(Message.DeviceId, EntryKey);
+
+							if (Entry && Entry->IsValid())
+							{
+								ProcessMultiChannelEnvelopeMessage(Entry, Message.EnvelopeValues);
+							}
+						});
+						break;
+					}
+
+					case 2: // FSubmixEnvelopeValuesMessage
+					{
+						CacheManager.IterateOverRange<FSubmixEnvelopeValuesMessage>(SubmixMessageNames::EnvelopeValues, StartTime, TimeMarker, [this](const FSubmixEnvelopeValuesMessage& Message)
+						{
+							const FSignalFlowEntryKey EntryKey(Message.DeviceId, ESignalFlowEntryType::Submix, Message.SubmixId);
+							TSharedPtr<FSignalFlowDashboardEntry>* Entry = FindDeviceEntry(Message.DeviceId, EntryKey);
+
+							if (Entry && Entry->IsValid())
+							{
+								ProcessMultiChannelEnvelopeMessage(Entry, Message.EnvelopeValues);
+							}
+						});
+						break;
+					}
+				}
+			});
+		}
 	}
 
 	void FSignalFlowTraceProvider::OnTimeControlMethodReset()
@@ -1121,6 +1196,7 @@ namespace UE::Audio::Insights
 		if (MaxPeakAmp > TNumericLimits<float>::Lowest())
 		{
 			(*OutEntry)->Amplitude = MaxPeakAmp;
+			(*OutEntry)->WriteNewAmplitudeValue(MaxPeakAmp);
 		}
 	}
 
diff --git a/Engine/Plugins/AudioInsights/Source/AudioInsights/Private/Providers/SignalFlowTraceProvider.h b/Engine/Plugins/AudioInsights/Source/AudioInsights/Private/Providers/SignalFlowTraceProvider.h
--- a/Engine/Plugins/AudioInsights/Source/AudioInsights/Private/Providers/SignalFlowTraceProvider.h
+++ b/Engine/Plugins/AudioInsights/Source/AudioInsights/Private/Providers/SignalFlowTraceProvider.h
@@ -18,6 +18,8 @@ namespace UE::Audio::Insights
 		virtual UE::Trace::IAnalyzer* ConstructAnalyzer(TraceServices::IAnalysisSession& InSession) override;
 		virtual void Reset() override;
 
+		void SetConnectionWiresAreAnimated(const bool bInAnimated) { bAnimateWires = bInAnimated; }
+
 		DECLARE_MULTICAST_DELEGATE(FOnRequestGraphRefresh);
 		FOnRequestGraphRefresh OnRequestGraphRefresh;
 
@@ -54,5 +56,6 @@ namespace UE::Audio::Insights
 		FSignalFlowMessages TraceMessages;
 
 		bool bForceGraphRefresh = false;
+		bool bAnimateWires = true;
 	};
 } // namespace UE::Audio::Insights
diff --git a/Engine/Plugins/AudioInsights/Source/AudioInsights/Private/SignalFlowEditorCommands.cpp b/Engine/Plugins/AudioInsights/Source/AudioInsights/Private/SignalFlowEditorCommands.cpp
--- a/Engine/Plugins/AudioInsights/Source/AudioInsights/Private/SignalFlowEditorCommands.cpp
+++ b/Engine/Plugins/AudioInsights/Source/AudioInsights/Private/SignalFlowEditorCommands.cpp
@@ -18,6 +18,7 @@ namespace UE::Audio::Insights
 	{
 		UI_COMMAND(ToggleHorizontalView, "Horizontal Flow", "Toggles between horizontally/vertically orienting the Signal Flow graph.", EUserInterfaceActionType::ToggleButton, FInputChord(EKeys::H, EModifierKey::Control));
 		UI_COMMAND(ShowNodeDetails, "Show Node Details", "Toggles visibility of node details.", EUserInterfaceActionType::ToggleButton, FInputChord(EKeys::D, EModifierKey::Control));
+		UI_COMMAND(ToggleAnimateWires, "Animate Wires", "Toggles whether connection wire animations driven by amplitude envelope values are enabled in the Signal Flow graph.", EUserInterfaceActionType::ToggleButton, FInputChord());
 	}
 
 } // namespace UE::Audio::Insights
diff --git a/Engine/Plugins/AudioInsights/Source/AudioInsights/Private/SignalFlowEditorCommands.h b/Engine/Plugins/AudioInsights/Source/AudioInsights/Private/SignalFlowEditorCommands.h
--- a/Engine/Plugins/AudioInsights/Source/AudioInsights/Private/SignalFlowEditorCommands.h
+++ b/Engine/Plugins/AudioInsights/Source/AudioInsights/Private/SignalFlowEditorCommands.h
@@ -17,9 +17,11 @@ namespace UE::Audio::Insights
 
 		TSharedPtr<const FUICommandInfo> GetToggleHorizontalViewCommand() const { return ToggleHorizontalView; }
 		TSharedPtr<const FUICommandInfo> GetShowNodeDetailsCommand() const { return ShowNodeDetails; }
+		TSharedPtr<const FUICommandInfo> GetToggleAnimateWiresCommand() const { return ToggleAnimateWires; }
 
 	private:
 		TSharedPtr<FUICommandInfo> ToggleHorizontalView;
 		TSharedPtr<FUICommandInfo> ShowNodeDetails;
+		TSharedPtr<FUICommandInfo> ToggleAnimateWires;
 	};
 } // namespace UE::Audio::Insights
diff --git a/Engine/Plugins/AudioInsights/Source/AudioInsights/Private/Views/SSignalFlowGraph.cpp b/Engine/Plugins/AudioInsights/Source/AudioInsights/Private/Views/SSignalFlowGraph.cpp
--- a/Engine/Plugins/AudioInsights/Source/AudioInsights/Private/Views/SSignalFlowGraph.cpp
+++ b/Engine/Plugins/AudioInsights/Source/AudioInsights/Private/Views/SSignalFlowGraph.cpp
@@ -29,6 +29,12 @@ namespace UE::Audio::Insights
 		constexpr float SendLabelRoundedBrushRadius = 8.0f;
 
 		constexpr float SplineRadius = 80.0f;
+		constexpr float SplineThicknessScalarMin = 1.0f;
+		constexpr float SplineThicknessDefault = 1.5f;
+		constexpr float SplineAnimatedWireThickness = 3.0f;
+		constexpr float SplineThicknessMin = 1.0f;
+		constexpr float SplineAnimPointsPerUnit = 0.2f;
+
 		constexpr float SendLabelPropAlongLine = 0.25f; // How far alone the spline we want the send level to appear (range 0 -> 1)
 
 		constexpr float SelectedNodeBorderThickness = 3.0f;
@@ -94,9 +100,9 @@ namespace UE::Audio::Insights
 			}
 		}
 
-		FSlateColor GetConnectionColor()
+		FLinearColor GetConnectionColor(const float AmplitudeValue, const float WireSplineAmplitudePowerFactor)
 		{
-			return FStyleColors::White25;
+			return FLinearColor::LerpUsingHSV(FStyleColors::Hover.GetSpecifiedColor(), FStyleColors::AccentWhite.GetSpecifiedColor(), FMath::Pow(AmplitudeValue, WireSplineAmplitudePowerFactor));
 		}
 
 		FVector2d GetSplineStartDirection(const EOrientation Orientation)
@@ -692,6 +698,7 @@ namespace UE::Audio::Insights
 		VerticalScrollBox = InArgs._VerticalScrollBox;
 		ShowNodeDetails = InArgs._ShowNodeDetails;
 		DisplayAmpPeakInDb = InArgs._DisplayAmpPeakInDb;
+		AnimateWires = InArgs._AnimateWires;
 		OnNodeSelected = InArgs._OnNodeSelected;
 
 		Clipping = EWidgetClipping::ClipToBounds;
@@ -752,6 +759,12 @@ namespace UE::Audio::Insights
 		RefreshGraph();
 	}
 
+	void SSignalFlowGraph::SetWireAnimationSettings(const float AmpPowerFactor, const float WireThicknessScalarMax)
+	{
+		WireSplineAmplitudePowerFactor = AmpPowerFactor;
+		WireSplineThicknessScalarMax = WireThicknessScalarMax;
+	}
+
 	void SSignalFlowGraph::ResetZoom()
 	{
 		bAnimateZoom = true;
@@ -939,11 +952,27 @@ namespace UE::Audio::Insights
 	{
 		using namespace SSignalFlowGraphStyle;
 
-		if (!ListSource)
+		if (!ListSource || !HorizontalScrollBox.IsSet() || !VerticalScrollBox.IsSet())
+		{
+			return LayerId;
+		}
+
+		TSharedPtr<SScrollBox> HorizontalScrollBoxPinned = HorizontalScrollBox.Get().Pin();
+		TSharedPtr<SScrollBox> VerticalScrollBoxPinned = VerticalScrollBox.Get().Pin();
+		if (!HorizontalScrollBoxPinned.IsValid() || !VerticalScrollBoxPinned.IsValid())
 		{
 			return LayerId;
 		}
 
+		const FGeometry& ViewportGeometry = VerticalScrollBoxPinned->GetPaintSpaceGeometry();
+
+		const FVector2D ViewportOffsetFraction(HorizontalScrollBoxPinned->GetViewOffsetFraction(), VerticalScrollBoxPinned->GetViewOffsetFraction());
+		const FVector2D ViewportPosInGraph = AllottedGeometry.GetLocalSize() * ViewportOffsetFraction;
+
+		FBox2D LocalBoundsBox(ForceInit);
+		LocalBoundsBox += FVector2D::Zero();
+		LocalBoundsBox += ViewportGeometry.GetLocalSize();
+
 		for (const TSharedPtr<ISignalFlowNode>& Node : *ListSource)
 		{
 			if (!Node->IsRealNode())
@@ -959,6 +988,7 @@ namespace UE::Audio::Insights
 			}
 
 			const FVector2D ConnectionStartPos = (*InputNode)->GetLocalOutputPos();
+			const FVector2D ConnectionStartPosInView = ConnectionStartPos - ViewportPosInGraph;
 
 			for (const FSignalFlowEntryKey& NodeOutputKey : Node->FilteredOutputs)
 			{
@@ -966,28 +996,41 @@ namespace UE::Audio::Insights
 
 				if (OutputNode.IsValid())
 				{
+					// Check if the node connection wire is inside the viewport
 					const FVector2D ConnectionEndPos = OutputNode->GetLocalInputPos();
-					if (Orientation == EOrientation::Orient_Vertical && ConnectionEndPos.Y <= ConnectionStartPos.Y)
+					const FVector2D ConnectionEndPosInView = ConnectionEndPos - ViewportPosInGraph;
+					if (!FMath::LineBoxIntersection2D(LocalBoundsBox, ConnectionStartPosInView, ConnectionEndPosInView))
 					{
 						continue;
 					}
-					else if (Orientation == EOrientation::Orient_Horizontal && ConnectionEndPos.X <= ConnectionStartPos.X)
+
+					const TSharedPtr<SSignalFlowGraphNode> RealInputNode = FindRealNode((*InputNode)->GetNodeEntryKey(), ENodeConnectionType::Input);
+					const TSharedPtr<FSignalFlowDashboardEntry> Entry = RealInputNode.IsValid() ? RealInputNode->GetEntry() : nullptr;
+
+					if (AnimateWires.Get() && Entry.IsValid() && Entry->Amplitude.IsSet())
 					{
-						continue;
-					}
+						FSignalFlowOutputData* OutputData = Entry->Outputs.Find(OutputNode->GetNodeEntryKey());
+						const float SendLevel = OutputData && OutputData->SendLevel.IsSet() ? OutputData->SendLevel.GetValue() : 1.0f;
 
-					FSlateDrawElement::MakeSpline(
-						OutDrawElements,
-						++LayerId,
-						AllottedGeometry.ToPaintGeometry(),
-						ConnectionStartPos,
-						GetSplineStartDirection(Orientation),
-						ConnectionEndPos,
-						GetSplineEndDirection(Orientation),
-						FMath::Max(2.0f * ZoomScaleFactor, 1.0f),
-						ESlateDrawEffect::NoPixelSnapping,
-						GetConnectionColor().GetSpecifiedColor()
-					);
+						LayerId = AnimateWire(AllottedGeometry, OutDrawElements, LayerId, Entry, ConnectionStartPos, ConnectionEndPos, SendLevel);
+					}
+					else
+					{
+						const float ScaledSplineThickness = FMath::Max(SplineThicknessDefault * ZoomScaleFactor, SplineThicknessMin);
+
+						FSlateDrawElement::MakeSpline(
+							OutDrawElements,
+							++LayerId,
+							AllottedGeometry.ToPaintGeometry(),
+							ConnectionStartPos,
+							GetSplineStartDirection(Orientation),
+							ConnectionEndPos,
+							GetSplineEndDirection(Orientation),
+							ScaledSplineThickness,
+							ESlateDrawEffect::NoPixelSnapping,
+							GetConnectionColor(0.0f, WireSplineAmplitudePowerFactor)
+						);
+					}
 				}
 			}
 		}
@@ -995,6 +1038,142 @@ namespace UE::Audio::Insights
 		return LayerId;
 	}
 
+	int32 SSignalFlowGraph::AnimateWire(const FGeometry& AllottedGeometry,
+										FSlateWindowElementList& OutDrawElements,
+										int32 LayerId,
+										const TSharedPtr<FSignalFlowDashboardEntry>& Entry,
+										const FVector2D& ConnectionStartPos,
+										const FVector2D& ConnectionEndPos,
+										const float SendLevel) const
+	{
+		using namespace SSignalFlowGraphStyle;
+
+		if (!Entry.IsValid())
+		{
+			return LayerId;
+		}
+
+		// Use the straight line distance between the two points as an approximation of how long the spline length will be
+		const float DistanceBetweenConnections = FVector2D::Distance(ConnectionStartPos, ConnectionEndPos);
+		const int32 NumPoints = DistanceBetweenConnections * SplineAnimPointsPerUnit;
+		if (NumPoints <= 0)
+		{
+			return LayerId;
+		}
+
+		const uint32 NumAmpReadings = Entry->AmplitudeDataRange.Capacity();
+		const float AmpIndexDelta = static_cast<float>(NumAmpReadings) / static_cast<float>(NumPoints);
+
+		TArray<FVector2D> TopPoints;
+		TopPoints.Reserve(NumPoints * 2);
+
+		TArray<FVector2D> BottomPoints;
+		BottomPoints.Reserve(NumPoints * 2);
+
+		const int32 AmpWriteIndex = Entry->AmplitudeWriteIndex;
+
+		int32 ReadIndex = Entry->AmplitudeDataRange.GetPreviousIndex(AmpWriteIndex);
+		uint32 AmpReadingCount = 0u;
+
+		const FVector2D StartDirection = GetSplineStartDirection(Orientation);
+		const FVector2D EndDirection = GetSplineEndDirection(Orientation);
+		FVector2D LastPointPos = ConnectionStartPos;
+
+		float InterpolationAmount = 0.0f;
+		for (int32 PointIndex = 0; PointIndex < NumPoints; ++PointIndex)
+		{
+			if (AmpReadingCount >= NumAmpReadings)
+			{
+				break;
+			}
+
+			const float CurrentAmp = Entry->AmplitudeDataRange[ReadIndex] * SendLevel;
+			const float Prop = static_cast<float>(PointIndex) / static_cast<float>(NumPoints);
+
+			CollectAmplitudeConnectionPoints(CurrentAmp,
+											 ConnectionStartPos,
+											 StartDirection,
+											 ConnectionEndPos,
+											 EndDirection,
+											 Prop,
+											 TopPoints,
+											 BottomPoints,
+											 LastPointPos);
+
+			InterpolationAmount += AmpIndexDelta;
+			while (InterpolationAmount >= 1.0f)
+			{
+				InterpolationAmount -= 1.0f;
+				ReadIndex = Entry->AmplitudeDataRange.GetPreviousIndex(ReadIndex);
+				++AmpReadingCount;
+			}
+		}
+
+		const FLinearColor ConnectionColor(GetConnectionColor(Entry->Amplitude.GetValue(), WireSplineAmplitudePowerFactor));
+		FSlateDrawElement::MakeLines(OutDrawElements,
+									 ++LayerId,
+									 AllottedGeometry.ToPaintGeometry(),
+									 TopPoints,
+									 ESlateDrawEffect::None,
+									 ConnectionColor,
+									 true, /* bAntialias */
+									 SplineAnimatedWireThickness);
+
+		FSlateDrawElement::MakeLines(OutDrawElements,
+									 ++LayerId,
+									 AllottedGeometry.ToPaintGeometry(),
+									 BottomPoints,
+									 ESlateDrawEffect::None,
+									 ConnectionColor,
+									 true, /* bAntialias */
+									 SplineAnimatedWireThickness);
+
+		return LayerId;
+	}
+
+	void SSignalFlowGraph::CollectAmplitudeConnectionPoints(const float Amplitude,
+															const FVector2D& ConnectionStartPos,
+															const FVector2D& StartDirection,
+															const FVector2D& ConnectionEndPos,
+															const FVector2D& EndDirection,
+															const float Prop,
+															TArray<FVector2D>& OutTopPoints,
+															TArray<FVector2D>& OutBottomPoints,
+															FVector2D& OutLastPoint) const
+	{
+		using namespace SSignalFlowGraphStyle;
+
+		const FVector2D SplinePoint = FMath::CubicInterp(ConnectionStartPos,
+														 StartDirection,
+														 ConnectionEndPos,
+														 EndDirection,
+														 Prop);
+
+		const float SplineThickness = FMath::Lerp(SplineThicknessDefault * SplineThicknessScalarMin, SplineThicknessDefault * WireSplineThicknessScalarMax, FMath::Pow(Amplitude, WireSplineAmplitudePowerFactor));
+		const float PointHeight = SplineThickness * 0.5f;
+
+		const FVector2D LastTangent(SplinePoint.X - OutLastPoint.X, SplinePoint.Y - OutLastPoint.Y);
+		const FVector2D Normal = FVector2D(LastTangent.Y, -1.0f * LastTangent.X).GetSafeNormal();
+
+		const FVector2D Top = { SplinePoint.X - (Normal.X * PointHeight), SplinePoint.Y - (Normal.Y * PointHeight) };
+		const FVector2D Bottom = { SplinePoint.X + (Normal.X * PointHeight), SplinePoint.Y + (Normal.Y * PointHeight) };
+
+		if (Top.Equals(Bottom))
+		{
+			OutTopPoints.Add(Top);
+			OutBottomPoints.Add(Bottom);
+		}
+		else
+		{
+			OutTopPoints.Add(Top);
+			OutTopPoints.Add(Bottom);
+			OutBottomPoints.Add(Bottom);
+			OutBottomPoints.Add(Top);
+		}
+
+		OutLastPoint = SplinePoint;
+	}
+
 	FVector2D SSignalFlowGraph::ComputeDesiredSize(float InLayoutScaleMultiplier) const
 	{
 		using namespace SSignalFlowGraphStyle;
diff --git a/Engine/Plugins/AudioInsights/Source/AudioInsights/Private/Views/SignalFlowDashboardViewFactory.cpp b/Engine/Plugins/AudioInsights/Source/AudioInsights/Private/Views/SignalFlowDashboardViewFactory.cpp
--- a/Engine/Plugins/AudioInsights/Source/AudioInsights/Private/Views/SignalFlowDashboardViewFactory.cpp
+++ b/Engine/Plugins/AudioInsights/Source/AudioInsights/Private/Views/SignalFlowDashboardViewFactory.cpp
@@ -279,6 +279,12 @@ namespace UE::Audio::Insights
 			FExecuteAction::CreateLambda([this]() { ToggleShowNodeDetails(); }),
 			FCanExecuteAction(),
 			FIsActionChecked::CreateLambda([this]() { return bShowNodeDetails; }));
+
+		CommandList->MapAction(
+			Commands.GetToggleAnimateWiresCommand(),
+			FExecuteAction::CreateLambda([this]() { ToggleAnimateWires(); }),
+			FCanExecuteAction(),
+			FIsActionChecked::CreateLambda([this]() { return bAnimateWires; }));
 	}
 
 	TSharedRef<SWidget> FSignalFlowDashboardViewFactory::CreateToolbar()
@@ -459,6 +465,7 @@ namespace UE::Audio::Insights
 		MenuBuilder.BeginSection("SignalFlowSettingActions", LOCTEXT("AudioDashboard_SignalFlow_Settings_HeaderText", "Settings"));
 		{
 			CreateAmpDisplayModeToggle(MenuBuilder);
+			MenuBuilder.AddMenuEntry(Commands.GetToggleAnimateWiresCommand());
 
 			MenuBuilder.AddMenuEntry(Commands.GetToggleHorizontalViewCommand());
 
@@ -646,6 +653,15 @@ namespace UE::Audio::Insights
 	{
 		NodeDetailFilters.ToggleEnableAllNodeDetailFilters();
 
+#if WITH_EDITOR
+		FSignalFlowSettings::OnRequestWriteSettings.Broadcast();
+#endif // WITH_EDITOR	
+	}
+
+	void FSignalFlowDashboardViewFactory::ToggleAnimateWires()
+	{
+		bAnimateWires = !bAnimateWires;
+
 #if WITH_EDITOR
 		FSignalFlowSettings::OnRequestWriteSettings.Broadcast();
 #endif // WITH_EDITOR	
@@ -816,6 +832,10 @@ namespace UE::Audio::Insights
 					{
 						return bDisplayAmpPeakInDb;
 					})
+					.AnimateWires_Lambda([this]()
+					{
+						return bAnimateWires;
+					})
 				]
 			]
 		];
@@ -1759,11 +1779,20 @@ namespace UE::Audio::Insights
 			{
 				GraphWidgetContents->ToggleOrientation();
 			}
+
+			GraphWidgetContents->SetWireAnimationSettings(InSettings.WireSplineAmplitudePowerFactor, InSettings.WireSplineMaxThicknessScalar);
 		}
 
 		bShowNodeDetails = InSettings.bShowNodeDetails;
 		NodeDetailFilters = InSettings.NodeDetailFilters;
 		bDisplayAmpPeakInDb = InSettings.AmplitudeDisplayMode == EAudioAmplitudeDisplayMode::Decibels;
+		bAnimateWires = InSettings.bAnimateWires;
+
+		const TSharedPtr<FSignalFlowTraceProvider> Provider = FindProvider<FSignalFlowTraceProvider>();
+		if (Provider.IsValid())
+		{
+			Provider->SetConnectionWiresAreAnimated(bAnimateWires);
+		}
 	}
 
 	void FSignalFlowDashboardViewFactory::OnWriteEditorSettings(FSignalFlowSettings& OutSettings)
@@ -1777,6 +1806,13 @@ namespace UE::Audio::Insights
 		OutSettings.bShowNodeDetails = bShowNodeDetails;
 		OutSettings.NodeDetailFilters = NodeDetailFilters;
 		OutSettings.AmplitudeDisplayMode = bDisplayAmpPeakInDb ? EAudioAmplitudeDisplayMode::Decibels : EAudioAmplitudeDisplayMode::Linear;
+		OutSettings.bAnimateWires = bAnimateWires;
+
+		const TSharedPtr<FSignalFlowTraceProvider> Provider = FindProvider<FSignalFlowTraceProvider>();
+		if (Provider.IsValid())
+		{
+			Provider->SetConnectionWiresAreAnimated(bAnimateWires);
+		}
 	}
 #endif // WITH_EDITOR
 } // namespace UE::Audio::Insights
diff --git a/Engine/Plugins/AudioInsights/Source/AudioInsights/Public/Messages/SignalFlowTraceMessages.h b/Engine/Plugins/AudioInsights/Source/AudioInsights/Public/Messages/SignalFlowTraceMessages.h
--- a/Engine/Plugins/AudioInsights/Source/AudioInsights/Public/Messages/SignalFlowTraceMessages.h
+++ b/Engine/Plugins/AudioInsights/Source/AudioInsights/Public/Messages/SignalFlowTraceMessages.h
@@ -123,7 +123,7 @@ namespace UE::Audio::Insights
 
 	struct FSignalFlowDashboardEntry : public IObjectDashboardEntry
 	{
-		FSignalFlowDashboardEntry() = default;
+		FSignalFlowDashboardEntry();
 		virtual ~FSignalFlowDashboardEntry() = default;
 
 		virtual FText GetDisplayName() const override 
@@ -138,6 +138,9 @@ namespace UE::Audio::Insights
 
 		FSignalFlowEntryKey GetSignalFlowEntryKey() const { return { DeviceId, EntryType, SignalFlowNodeID }; }
 
+		void WriteNewAmplitudeValue(const float Value);
+		void ResetDataBuffers();
+
 		::Audio::FDeviceId DeviceId = INDEX_NONE;
 		ESignalFlowEntryType EntryType = ESignalFlowEntryType::SoundSource;
 		uint32 SignalFlowNodeID = INDEX_NONE;
@@ -149,6 +152,9 @@ namespace UE::Audio::Insights
 		TSet<FSignalFlowEntryKey> Inputs;
 		TSortedMap<FSignalFlowEntryKey, FSignalFlowOutputData> Outputs;
 
+		TCircularBuffer<float> AmplitudeDataRange;
+		int32 AmplitudeWriteIndex = 0;
+
 		TOptional<float> Amplitude;
 		TOptional<float> Volume;
 		TOptional<float> Pitch;
diff --git a/Engine/Plugins/AudioInsights/Source/AudioInsights/Public/Settings/SignalFlowSettings.h b/Engine/Plugins/AudioInsights/Source/AudioInsights/Public/Settings/SignalFlowSettings.h
--- a/Engine/Plugins/AudioInsights/Source/AudioInsights/Public/Settings/SignalFlowSettings.h
+++ b/Engine/Plugins/AudioInsights/Source/AudioInsights/Public/Settings/SignalFlowSettings.h
@@ -82,6 +82,18 @@ struct FSignalFlowSettings
 	UPROPERTY(EditAnywhere, config, Category = SignalFlow)
 	FSignalFlowNodeDetailFilterSettings NodeDetailFilters;
 
+	/** Toggles whether connection wire animations driven by amplitude envelope values are enabled in the Signal Flow graph. */
+	UPROPERTY(EditAnywhere, config, Category = SignalFlow)
+	bool bAnimateWires = true;
+
+	/** Power factor applied to amplitude values when drawing animated signal flow connection wires. Controls how much additional weight is applied to quieter nodes. */
+	UPROPERTY(EditAnywhere, config, Category = SignalFlow, meta = (UIMin = "0.1", UIMax = "1.0", ClampMin = "0.1", ClampMax = "1.0"))
+	float WireSplineAmplitudePowerFactor = 0.4f;
+
+	/** The maximum thickness multiplier applied to animated signal flow connection wires when a node is at high amplitudes. */
+	UPROPERTY(EditAnywhere, config, Category = SignalFlow, meta = (UIMin = "1.0", UIMax = "40.0", ClampMin = "1.0", ClampMax = "40.0"))
+	float WireSplineMaxThicknessScalar = 14.0f;
+
 #if WITH_EDITOR
 	DECLARE_MULTICAST_DELEGATE_OneParam(FOnReadSignalFlowSettings, const FSignalFlowSettings&);
 	static AUDIOINSIGHTS_API FOnReadSignalFlowSettings OnReadSettings;
diff --git a/Engine/Plugins/AudioInsights/Source/AudioInsights/Public/Views/SSignalFlowGraph.h b/Engine/Plugins/AudioInsights/Source/AudioInsights/Public/Views/SSignalFlowGraph.h
--- a/Engine/Plugins/AudioInsights/Source/AudioInsights/Public/Views/SSignalFlowGraph.h
+++ b/Engine/Plugins/AudioInsights/Source/AudioInsights/Public/Views/SSignalFlowGraph.h
@@ -156,6 +156,7 @@ namespace UE::Audio::Insights
 			, _NodeDetailFilterSettings(nullptr)
 			, _ShowNodeDetails(false)
 			, _DisplayAmpPeakInDb(true)
+			, _AnimateWires(true)
 		{
 		}
 
@@ -168,6 +169,7 @@ namespace UE::Audio::Insights
 		SLATE_ATTRIBUTE(TWeakPtr<SScrollBox>, VerticalScrollBox)
 		SLATE_ATTRIBUTE(bool, ShowNodeDetails)
 		SLATE_ATTRIBUTE(bool, DisplayAmpPeakInDb)
+		SLATE_ATTRIBUTE(bool, AnimateWires)
 
 		SLATE_EVENT(FOnSignalFlowNodeSelected, OnNodeSelected)
 
@@ -182,6 +184,8 @@ namespace UE::Audio::Insights
 		UE_API void ToggleOrientation();
 		EOrientation GetGraphOrientation() const { return Orientation; }
 
+		UE_API void SetWireAnimationSettings(const float AmpPowerFactor, const float WireThicknessScalarMax);
+
 		UE_API void ResetZoom();
 
 		UE_API TSharedPtr<SSignalFlowGraphNode> FindNodeWidget(const FSignalFlowEntryKey& NodeEntryKey) const;
@@ -201,6 +205,24 @@ namespace UE::Audio::Insights
 	private:
 		int32 OnPaintConnections(const FGeometry& AllottedGeometry, FSlateWindowElementList& OutDrawElements, int32 LayerId) const;
 
+		int32 AnimateWire(const FGeometry& AllottedGeometry,
+							FSlateWindowElementList& OutDrawElements,
+							int32 LayerId,
+							const TSharedPtr<FSignalFlowDashboardEntry>& Entry,
+							const FVector2D& ConnectionStartPos,
+							const FVector2D& ConnectionEndPos,
+							const float SendLevel) const;
+
+		void CollectAmplitudeConnectionPoints(const float Amplitude,
+												const FVector2D& ConnectionStartPos,
+												const FVector2D& StartDirection,
+												const FVector2D& ConnectionEndPos,
+												const FVector2D& EndDirection,
+												const float Prop,
+												TArray<FVector2D>& OutTopPoints,
+												TArray<FVector2D>& OutBottomPoints,
+												FVector2D& OutLastPoint) const;
+
 		void CreateGraphNodes(TArray<FSignalFlowEntryKey>& OutRemainingNodeKeys);
 		void CreateSendLabels();
 		void CleanUpOldWidgets(const TArray<FSignalFlowEntryKey>& NodeKeysToRemove);
@@ -248,6 +270,7 @@ namespace UE::Audio::Insights
 		TAttribute<TSharedPtr<FSignalFlowDashboardEntry>> ActiveAudioDeviceEntry;
 		TAttribute<bool> ShowNodeDetails;
 		TAttribute<bool> DisplayAmpPeakInDb;
+		TAttribute<bool> AnimateWires;
 
 		TAttribute<TWeakPtr<SScrollBox>> HorizontalScrollBox;
 		TAttribute<TWeakPtr<SScrollBox>> VerticalScrollBox;
@@ -256,6 +279,8 @@ namespace UE::Audio::Insights
 		EOrientation Orientation = EOrientation::Orient_Vertical;
 
 		float ZoomScaleFactor = 1.0f;
+		float WireSplineAmplitudePowerFactor = 0.4f;
+		float WireSplineThicknessScalarMax = 14.0f;
 
 		bool bNeedsRefresh = false;
 		bool bAnimateZoom = false;
diff --git a/Engine/Plugins/AudioInsights/Source/AudioInsights/Public/Views/SignalFlowDashboardViewFactory.h b/Engine/Plugins/AudioInsights/Source/AudioInsights/Public/Views/SignalFlowDashboardViewFactory.h
--- a/Engine/Plugins/AudioInsights/Source/AudioInsights/Public/Views/SignalFlowDashboardViewFactory.h
+++ b/Engine/Plugins/AudioInsights/Source/AudioInsights/Public/Views/SignalFlowDashboardViewFactory.h
@@ -80,6 +80,7 @@ namespace UE::Audio::Insights
 		void ToggleGraphOrientation();
 		void ToggleShowNodeDetails();
 		void ToggleEnableAllNodeDetailFilters();
+		void ToggleAnimateWires();
 
 		TSharedPtr<FSignalFlowDashboardEntry> GetFilteredAudioDeviceEntry() const;
 
@@ -157,6 +158,7 @@ namespace UE::Audio::Insights
 		bool bAutoFocusAfterGraphRefresh = false;
 		bool bShowNodeDetails = false;
 		bool bDisplayAmpPeakInDb = true;
+		bool bAnimateWires = true;
 	};
 } // namespace UE::Audio::Insights
 
