[Audio Insights] Signal flow scroll boxes are now always visible

#Synced-CL 50481844

diff --git a/Engine/Plugins/AudioInsights/Source/AudioInsights/Private/Views/SignalFlowDashboardViewFactory.cpp b/Engine/Plugins/AudioInsights/Source/AudioInsights/Private/Views/SignalFlowDashboardViewFactory.cpp
--- a/Engine/Plugins/AudioInsights/Source/AudioInsights/Private/Views/SignalFlowDashboardViewFactory.cpp
+++ b/Engine/Plugins/AudioInsights/Source/AudioInsights/Private/Views/SignalFlowDashboardViewFactory.cpp
@@ -737,6 +737,7 @@ namespace UE::Audio::Insights
 	{
 		GraphVerticalScrollBox = SAssignNew(GraphVerticalScrollBox, SScrollBox)
 		.Orientation(EOrientation::Orient_Vertical)
+		.ScrollBarAlwaysVisible(true)
 
 		+ SScrollBox::Slot()
 		.FillSize(1.0f)
@@ -745,6 +746,7 @@ namespace UE::Audio::Insights
 		[
 			SAssignNew(GraphHorizontalScrollBox, SScrollBox)
 			.Orientation(EOrientation::Orient_Horizontal)
+			.ScrollBarAlwaysVisible(true)
 
 			+ SScrollBox::Slot()
 			.FillSize(1.0f)
