[Audio Insights] Allow node details to become expandable/collapsible in the node heading

#Synced-CL 50533306

diff --git a/Engine/Plugins/AudioInsights/Source/AudioInsights/Private/Views/SSignalFlowGraph.cpp b/Engine/Plugins/AudioInsights/Source/AudioInsights/Private/Views/SSignalFlowGraph.cpp
--- a/Engine/Plugins/AudioInsights/Source/AudioInsights/Private/Views/SSignalFlowGraph.cpp
+++ b/Engine/Plugins/AudioInsights/Source/AudioInsights/Private/Views/SSignalFlowGraph.cpp
@@ -11,6 +11,7 @@
 #include "Styling/StyleColors.h"
 #include "Views/SignalFlowNodes.h"
 #include "Widgets/Colors/SColorBlock.h"
+#include "Widgets/Input/SCheckBox.h"
 #include "Widgets/Layout/SBorder.h"
 #include "Widgets/Layout/SScaleBox.h"
 #include "Widgets/Layout/SScrollBox.h"
@@ -379,7 +380,19 @@ namespace UE::Audio::Insights
 	{
 		using namespace SSignalFlowGraphStyle;
 
-		const int32 NumParamsVisible = NodeDetailFilterSettings ? NodeDetailFilterSettings->GetNumVisibleParams() : 0;
+		int32 NumParamsVisible = 0;
+		if (NodeDetailFilterSettings)
+		{
+			if (bNodeDetailsAreExpanded)
+			{
+				NumParamsVisible = NodeDetailFilterSettings->GetNumVisibleParams();
+			}
+			// We always show amplitude if the node is collapsed
+			else if (Entry.IsValid() && Entry->Amplitude.IsSet())
+			{
+				NumParamsVisible = 1;
+			}
+		}
 
 		const FVector2D Size = SCompoundWidget::ComputeDesiredSize(InLayoutScaleMultiplier);
 		const FVector2D ClampedSize = { FMath::Min(Size.X, NodeMaxWidth), FMath::Min(Size.Y, GetNodeMaxHeight(NumParamsVisible, ShowNodeDetails.Get())) };
@@ -446,7 +459,7 @@ namespace UE::Audio::Insights
 		return CachedPositionInGraph;
 	}
 
-	TSharedRef<SHorizontalBox> SSignalFlowGraphNode::CreateNodeName(const SSignalFlowGraphNode::FArguments& InArgs) const
+	TSharedRef<SHorizontalBox> SSignalFlowGraphNode::CreateNodeName(const SSignalFlowGraphNode::FArguments& InArgs)
 	{
 		using namespace SSignalFlowGraphStyle;
 
@@ -469,6 +482,28 @@ namespace UE::Audio::Insights
 				.Justification(ETextJustify::Left)
 				.OverflowPolicy(ETextOverflowPolicy::Ellipsis)
 				.Clipping(EWidgetClipping::ClipToBounds)
+			]
+		
+			+ SHorizontalBox::Slot()
+			.AutoWidth()
+			.MinSize(12.0f)
+			.MaxSize(12.0f)
+			.Padding(4.0f, 0.0f, 0.0f, 0.0f)
+			[
+				SNew(SCheckBox)
+				.Style(&FSlateStyle::Get().GetWidgetStyle<FCheckBoxStyle>("AudioInsights.ExpandArrowToggle"))
+				.Visibility_Lambda([this]()
+				{ 
+					return NodeDetailsExpanderArrowIsEnabled() ? EVisibility::Visible : EVisibility::Collapsed;
+				})
+				.IsChecked_Lambda([this]()
+				{
+					return bNodeDetailsAreExpanded ? ECheckBoxState::Checked : ECheckBoxState::Unchecked;
+				})
+				.OnCheckStateChanged_Lambda([this](const ECheckBoxState CheckboxState)
+				{
+					bNodeDetailsAreExpanded = CheckboxState == ECheckBoxState::Checked;
+				})
 			];
 	}
 
@@ -549,7 +584,9 @@ namespace UE::Audio::Insights
 				.Visibility_Lambda([this, ParamType]()
 				{
 					const bool bParamIsFiltered = NodeDetailFilterSettings && NodeDetailFilterSettings->GetParameterIsVisible(ParamType);
-					return bParamIsFiltered && Entry.IsValid() && GetNodeDetailParamValue(Entry, ParamType).IsSet() ? EVisibility::Visible : EVisibility::Collapsed;
+					const bool bParamIsHiddenByExpander = !bNodeDetailsAreExpanded && ParamType != ESignalFlowNodeDetailParam::Amplitude;
+
+					return bParamIsFiltered && !bParamIsHiddenByExpander && Entry.IsValid() && GetNodeDetailParamValue(Entry, ParamType).IsSet() ? EVisibility::Visible : EVisibility::Collapsed;
 				})
 
 				+ SHorizontalBox::Slot()
@@ -590,6 +627,46 @@ namespace UE::Audio::Insights
 				];
 	}
 
+	bool SSignalFlowGraphNode::NodeDetailsExpanderArrowIsEnabled() const
+	{
+		using namespace SSignalFlowGraphStyle;
+
+		// The expander arrow should only be visible if a parameter other than amplitude is visible on the node
+		// This requires show node details to be switched on, the parameter filter to be enabled, and data to be available for that parameter
+		if (!ShowNodeDetails.Get() || !Entry.IsValid() || NodeDetailFilterSettings == nullptr)
+		{
+			return false;
+		}
+
+		const int32 NumVisibleParams = NodeDetailFilterSettings->GetNumVisibleParams();
+		if (NumVisibleParams <= 0)
+		{
+			return false;
+		}
+
+		if (NodeDetailFilterSettings->GetNumVisibleParams() == 1 && NodeDetailFilterSettings->GetParameterIsVisible(ESignalFlowNodeDetailParam::Amplitude))
+		{
+			return false;
+		}
+		
+		const int32 NumFilters = static_cast<uint8>(ESignalFlowNodeDetailParam::MAX);
+		for (uint8 FilterID = 0u; FilterID < NumFilters; ++FilterID)
+		{
+			const ESignalFlowNodeDetailParam Param = static_cast<ESignalFlowNodeDetailParam>(FilterID);
+			if (Param == ESignalFlowNodeDetailParam::Amplitude)
+			{
+				continue;
+			}
+
+			if (NodeDetailFilterSettings->GetParameterIsVisible(Param) && GetNodeDetailParamValue(Entry, Param).IsSet())
+			{
+				return true;
+			}
+		}
+
+		return false;
+	}
+
 	void SSignalFlowGraphSendLabel::Construct(const SSignalFlowGraphSendLabel::FArguments& InArgs)
 	{
 		using namespace SSignalFlowGraphStyle;
diff --git a/Engine/Plugins/AudioInsights/Source/AudioInsights/Public/AudioInsightsStyle.h b/Engine/Plugins/AudioInsights/Source/AudioInsights/Public/AudioInsightsStyle.h
--- a/Engine/Plugins/AudioInsights/Source/AudioInsights/Public/AudioInsightsStyle.h
+++ b/Engine/Plugins/AudioInsights/Source/AudioInsights/Public/AudioInsightsStyle.h
@@ -194,6 +194,7 @@ namespace UE::Audio::Insights
 			Set("AudioInsights.Analyzers.BackgroundColor", FLinearColor(0.0075f, 0.0075f, 0.0075f, 1.0f));
 
 			// Icons
+			const FVector2D Icon12(12.0f, 12.0f);
 			const FVector2D Icon16(16.0f, 16.0f);
 			const FVector2D Icon20(20.0f, 20.0f);
 			const FVector2D Icon24(24.0f, 24.0f);
@@ -206,6 +207,8 @@ namespace UE::Audio::Insights
 			Set("AudioInsights.Icon.CheckBoxIndeterminate", new IMAGE_BRUSH_SVG(TEXT("Icons/checkboxindeterminate"), Icon16));
 			Set("AudioInsights.Icon.Dashboard", new IMAGE_BRUSH_SVG(TEXT("Icons/audio_dashboard"), Icon16));
 			Set("AudioInsights.Icon.Event", new IMAGE_BRUSH_SVG(TEXT("Icons/audio_event"), Icon16));
+			Set("AudioInsights.Icon.ExpandArrowDown", new IMAGE_BRUSH_SVG(TEXT("Icons/expand_arrow_down"), Icon12));
+			Set("AudioInsights.Icon.ExpandArrowUp", new IMAGE_BRUSH_SVG(TEXT("Icons/expand_arrow_up"), Icon12));
 			Set("AudioInsights.Icon.Log", new IMAGE_BRUSH_SVG(TEXT("Icons/audio_log"), Icon16));
 			Set("AudioInsights.Icon.OnlyTraceAudioChannels", new IMAGE_BRUSH_SVG(TEXT("Icons/only_trace_audio_channels"), Icon16));
 			Set("AudioInsights.Icon.Open", new IMAGE_BRUSH_SVG(TEXT("Icons/open"), Icon20));
@@ -294,6 +297,16 @@ namespace UE::Audio::Insights
 				.SetUncheckedHoveredImage(RoundedWhiteAlphaBrush)
 			);
 
+			Set("AudioInsights.ExpandArrowToggle", FCheckBoxStyle(FAppStyle::Get().GetWidgetStyle<FCheckBoxStyle>("TransparentCheckBox"))
+				.SetCheckBoxType(ESlateCheckBoxType::ToggleButton)
+				.SetCheckedImage(IMAGE_BRUSH_SVG(TEXT("Icons/expand_arrow_up"), Icon12, FStyleColors::Foreground))
+				.SetCheckedHoveredImage(IMAGE_BRUSH_SVG(TEXT("Icons/expand_arrow_up"), Icon12, FStyleColors::ForegroundHover))
+				.SetCheckedPressedImage(IMAGE_BRUSH_SVG(TEXT("Icons/expand_arrow_up"), Icon12, FStyleColors::ForegroundHover))
+				.SetUncheckedImage(IMAGE_BRUSH_SVG(TEXT("Icons/expand_arrow_down"), Icon12, FStyleColors::Foreground))
+				.SetUncheckedHoveredImage(IMAGE_BRUSH_SVG(TEXT("Icons/expand_arrow_down"), Icon12, FStyleColors::ForegroundHover))
+				.SetUncheckedPressedImage(IMAGE_BRUSH_SVG(TEXT("Icons/expand_arrow_down"), Icon12, FStyleColors::ForegroundHover))
+			);
+
 			FSlateStyleRegistry::RegisterSlateStyle(*this);
 		}
 
diff --git a/Engine/Plugins/AudioInsights/Source/AudioInsights/Public/Views/SSignalFlowGraph.h b/Engine/Plugins/AudioInsights/Source/AudioInsights/Public/Views/SSignalFlowGraph.h
--- a/Engine/Plugins/AudioInsights/Source/AudioInsights/Public/Views/SSignalFlowGraph.h
+++ b/Engine/Plugins/AudioInsights/Source/AudioInsights/Public/Views/SSignalFlowGraph.h
@@ -81,10 +81,12 @@ namespace UE::Audio::Insights
 		void SetYPos(const float Pos) { YPos = Pos; }
 
 	private:
-		TSharedRef<SHorizontalBox> CreateNodeName(const SSignalFlowGraphNode::FArguments& InArgs) const;
+		TSharedRef<SHorizontalBox> CreateNodeName(const SSignalFlowGraphNode::FArguments& InArgs);
 		TSharedRef<SVerticalBox> CreateNodeDetails() const;
 		TSharedRef<SHorizontalBox> CreateNodeDetailsRow(const ESignalFlowNodeDetailParam ParamType) const;
 
+		bool NodeDetailsExpanderArrowIsEnabled() const;
+
 		TSharedPtr<FSignalFlowDashboardEntry> Entry;
 		TAttribute<FVector2d> PositionInGraph;
 		TAttribute<float> XPos;
@@ -103,6 +105,8 @@ namespace UE::Audio::Insights
 
 		TSharedPtr<SBorder> SelectedNodeBorder;
 		TSharedPtr<SBorder> NodeBorder;
+
+		bool bNodeDetailsAreExpanded = true;
 	};
 
 	class SSignalFlowGraphSendLabel : public SCompoundWidget
