[Audio Insights] Fix "solo filtered" in sound dashboard cutting off the start of the sound.
* Bonus - fix solo/mute filtered not working at all when the dashboard is viewed as a flat list/active sound list

#Synced-CL 50533306

diff --git a/Engine/Plugins/AudioInsights/Source/AudioInsights/Private/Views/SoundDashboardViewFactory.cpp b/Engine/Plugins/AudioInsights/Source/AudioInsights/Private/Views/SoundDashboardViewFactory.cpp
--- a/Engine/Plugins/AudioInsights/Source/AudioInsights/Private/Views/SoundDashboardViewFactory.cpp
+++ b/Engine/Plugins/AudioInsights/Source/AudioInsights/Private/Views/SoundDashboardViewFactory.cpp
@@ -353,11 +353,16 @@ namespace UE::Audio::Insights
 
 			if (!bIsSoundCueType)
 			{
+				bool bIsMuted = false;
+				bool bIsSoloed = false;
+				FString Reason;
+				InAudioDebugger.QuerySoloMuteSoundWave(SoundEntry.GetDisplayNameStr(), bIsSoloed, bIsMuted, Reason);
+
 				switch (InMuteSoloMode)
 				{
 					case EMuteSoloMode::Mute:
 					{
-						if (InAudioDebugger.IsMuteSoundWave(SoundEntry.GetDisplayFName()))
+						if (bIsMuted)
 						{
 							return true;
 						}
@@ -367,7 +372,7 @@ namespace UE::Audio::Insights
 
 					case EMuteSoloMode::Solo:
 					{
-						if (InAudioDebugger.IsSoloSoundWave(SoundEntry.GetDisplayFName()))
+						if (bIsSoloed)
 						{
 							return true;
 						}
@@ -789,14 +794,14 @@ namespace UE::Audio::Insights
 			{
 				return FilteredList.IsValid() && bNotMuteSolo && FilteredList->GetNumItemsSelected() > 0 ?
 					LOCTEXT("SoundDashboard_MuteButtonTextSelect", "Mute Selected") :
-					LOCTEXT("SoundDashboard_MuteButtonTextFilter", "Mute Filtered");
+					LOCTEXT("SoundDashboard_MuteButtonTextFilter", "Mute by Text");
 			};
 		TFunction<FSlateColor()> SoloColorFunc = [this]() { return bIsSoloFilteredMode ? FSlateColor(FColor::Green) : FSlateColor::UseForeground(); };
 		TFunction<FText()> SoloTextFunc = [bNotMuteSolo = bIsNotFilteredMuteOrSolo, &FilteredList = FilteredEntriesListView]()
 			{
 				return FilteredList.IsValid() && bNotMuteSolo && FilteredList->GetNumItemsSelected() > 0 ?
 					LOCTEXT("SoundDashboard_SoloButtonTextSelect", "Solo Selected") :
-					LOCTEXT("SoundDashboard_SoloButtonTextFilter", "Solo Filtered");
+					LOCTEXT("SoundDashboard_SoloButtonTextFilter", "Solo by Text");
 			};
 
 		return SNew(SHorizontalBox)
@@ -992,6 +997,15 @@ namespace UE::Audio::Insights
 		return SoundsFilterBarButton;
 	}
 
+	void FSoundDashboardViewFactory::SetSearchBoxFilterText(const FText& NewText)
+	{
+		FTraceTreeDashboardViewFactory::SetSearchBoxFilterText(NewText);
+
+#if ENABLE_AUDIO_DEBUG
+		MuteSoloFilteredEntries();
+#endif // ENABLE_AUDIO_DEBUG
+	}
+
 	TSharedRef<SWidget> FSoundDashboardViewFactory::MakeWidget(TSharedRef<SDockTab> OwnerTab, const FSpawnTabArgs& SpawnTabArgs)
 	{
 #if WITH_EDITOR
@@ -1310,13 +1324,6 @@ namespace UE::Audio::Insights
 
 		const FSoundDashboardEntry& SoundDashboardEntry = CastEntry(*Item);
 
-#if WITH_EDITOR
-		if (bIsMuteFilteredMode || bIsSoloFilteredMode)
-		{
-			MuteSoloFilteredEntries();
-		}
-#endif //WITH_EDITOR
-
 		return SNew(SRowWidget, OwnerTable, Item, AsShared())
 			.RenderOpacity_Lambda([Item]
 			{
@@ -1545,14 +1552,6 @@ namespace UE::Audio::Insights
 				SetFilteredVisibility(*CategoryEntry, FilterString);
 			}
 		}
-
-#if WITH_EDITOR
-		if (bIsMuteFilteredMode || bIsSoloFilteredMode)
-		{
-			ClearMutesAndSolos();
-			MuteSoloFilteredEntries();
-		}
-#endif // WITH_EDITOR
 	}
 
 	TSharedPtr<SWidget> FSoundDashboardViewFactory::OnConstructContextMenu()
@@ -2803,43 +2802,27 @@ namespace UE::Audio::Insights
 			return;
 		}
 
+		ClearMutesAndSolos();
+
 		if (!bIsMuteFilteredMode && !bIsSoloFilteredMode)
 		{
-			ClearMutesAndSolos();
 			return;
 		}
 
-		::Audio::FAudioDebugger& AudioDebugger = AudioDeviceManager->GetDebugger();
-
-		const EMuteSoloMode Mode = bIsSoloFilteredMode ? EMuteSoloMode::Solo : EMuteSoloMode::Mute;
-		bool bSomethingIsVisible = false;
-
-		for (const TSharedPtr<IDashboardDataTreeViewEntry>& Entry : DataViewEntries)
+		const FText& FilterText = GetSearchFilterText();
+		if (!FilterText.IsEmpty())
 		{
-			const FSoundDashboardEntry& CategoryEntry = FSoundDashboardViewFactoryPrivate::CastEntry(*Entry);
-			if (CategoryEntry.bIsCategory)
+			const EMuteSoloMode Mode = bIsSoloFilteredMode ? EMuteSoloMode::Solo : EMuteSoloMode::Mute;
+			::Audio::FAudioDebugger& AudioDebugger = AudioDeviceManager->GetDebugger();
+
+			if (Mode == EMuteSoloMode::Solo)
 			{
-				for (const TSharedPtr<IDashboardDataTreeViewEntry>& EntryChild : Entry->Children)
-				{
-					const FSoundDashboardEntry& SoundEntry = FSoundDashboardViewFactoryPrivate::CastEntry(*EntryChild);
-					if (SoundEntry.bIsVisible)
-					{
-						bSomethingIsVisible = true;
-						if (!FSoundDashboardViewFactoryPrivate::IsMuteSolo(AudioDebugger, *EntryChild, true /*bInCheckChildren*/, Mode))
-						{
-							SetMuteSolo(*EntryChild, Mode, true /*bInOnOff*/);
-						}
-					}
-				}
+				AudioDebugger.SetSoloSoundWave(FName(FilterText.ToString()), true);
+			}
+			else
+			{
+				AudioDebugger.SetMuteSoundWave(FName(FilterText.ToString()), true);
 			}
-		}
-
-		// If you have filtered out everything but are soloing to listen things that do match the filter,
-		// we mute the rest of the mix at this moment. 
-		const bool bShouldMuteAll = !GetSearchFilterText().IsEmpty() && Mode == EMuteSoloMode::Solo && !bSomethingIsVisible;
-		if (bShouldMuteAll)
-		{
-			AudioDebugger.SetSoloSoundCue(NAME_None, true /*bInOnOff*/);
 		}
 #endif //ENABLE_AUDIO_DEBUG
 	}
diff --git a/Engine/Plugins/AudioInsights/Source/AudioInsights/Public/Views/SoundDashboardViewFactory.h b/Engine/Plugins/AudioInsights/Source/AudioInsights/Public/Views/SoundDashboardViewFactory.h
--- a/Engine/Plugins/AudioInsights/Source/AudioInsights/Public/Views/SoundDashboardViewFactory.h
+++ b/Engine/Plugins/AudioInsights/Source/AudioInsights/Public/Views/SoundDashboardViewFactory.h
@@ -129,6 +129,8 @@ namespace UE::Audio::Insights
 		UE_API virtual TSharedPtr<SWidget> GetFilterBarWidget() override;
 		UE_API virtual TSharedPtr<SWidget> GetFilterBarButtonWidget() override;
 
+		UE_API virtual void SetSearchBoxFilterText(const FText& NewText) override;
+
 		UE_API bool IsRootItem(const TSharedRef<IDashboardDataTreeViewEntry>& InEntry) const;
 
 		UE_API TSharedRef<SWidget> GenerateWidgetForRootColumn(const TSharedRef<FTraceTreeDashboardViewFactory::SRowWidget>& InRowWidget, const TSharedRef<IDashboardDataTreeViewEntry>& InRowData, const FName& InColumn, const FText& InValueText);
diff --git a/Engine/Plugins/AudioInsights/Source/AudioInsights/Public/Views/TreeDashboardViewFactory.h b/Engine/Plugins/AudioInsights/Source/AudioInsights/Public/Views/TreeDashboardViewFactory.h
--- a/Engine/Plugins/AudioInsights/Source/AudioInsights/Public/Views/TreeDashboardViewFactory.h
+++ b/Engine/Plugins/AudioInsights/Source/AudioInsights/Public/Views/TreeDashboardViewFactory.h
@@ -90,6 +90,7 @@ namespace UE::Audio::Insights
 		UE_API virtual FReply OnDataRowKeyInput(const FGeometry& Geometry, const FKeyEvent& KeyEvent) const;
 
 		UE_API const FText& GetSearchFilterText() const;
+		UE_API virtual void SetSearchBoxFilterText(const FText& NewText);
 
 		UE_API virtual void RefreshFilteredEntriesListView() override;
 
@@ -185,7 +186,6 @@ namespace UE::Audio::Insights
 #endif // WITH_EDITOR
 
 		UE_API TSharedRef<SHeaderRow> MakeHeaderRowWidget();
-		UE_API void SetSearchBoxFilterText(const FText& NewText);
 
 		UE_API EColumnSortMode::Type GetColumnSortMode(const FName InColumnId) const;
 		UE_API void RequestSort();
