[Audio Insights] Fix build - Sound Dashboard View factory

#Synced-CL 50533306

diff --git a/Engine/Plugins/AudioInsights/Source/AudioInsights/Private/Views/SoundDashboardViewFactory.cpp b/Engine/Plugins/AudioInsights/Source/AudioInsights/Private/Views/SoundDashboardViewFactory.cpp
--- a/Engine/Plugins/AudioInsights/Source/AudioInsights/Private/Views/SoundDashboardViewFactory.cpp
+++ b/Engine/Plugins/AudioInsights/Source/AudioInsights/Private/Views/SoundDashboardViewFactory.cpp
@@ -1001,9 +1001,11 @@ namespace UE::Audio::Insights
 	{
 		FTraceTreeDashboardViewFactory::SetSearchBoxFilterText(NewText);
 
+#if WITH_EDITOR
 #if ENABLE_AUDIO_DEBUG
 		MuteSoloFilteredEntries();
 #endif // ENABLE_AUDIO_DEBUG
+#endif // WITH_EDITOR
 	}
 
 	TSharedRef<SWidget> FSoundDashboardViewFactory::MakeWidget(TSharedRef<SDockTab> OwnerTab, const FSpawnTabArgs& SpawnTabArgs)
