[Audio Insights] Show node full name in tooltip when hovering over a node's name

#Synced-CL 50703972

#Synced-CL 50703972

diff --git a/Engine/Plugins/AudioInsights/Source/AudioInsights/Private/Views/SSignalFlowGraph.cpp b/Engine/Plugins/AudioInsights/Source/AudioInsights/Private/Views/SSignalFlowGraph.cpp
--- a/Engine/Plugins/AudioInsights/Source/AudioInsights/Private/Views/SSignalFlowGraph.cpp
+++ b/Engine/Plugins/AudioInsights/Source/AudioInsights/Private/Views/SSignalFlowGraph.cpp
@@ -471,6 +471,7 @@ namespace UE::Audio::Insights
 			[
 				SNew(SImage)
 				.Image(FSlateStyle::Get().CreateIcon(InArgs._Entry->IconName).GetIcon())
+				.ToolTipText(InArgs._DisplayName)
 				.ColorAndOpacity(GetNodeAccentColor(NodeEntryKey.EntryType))
 			]
 
@@ -479,6 +480,7 @@ namespace UE::Audio::Insights
 			[
 				SNew(STextBlock)
 				.Text(InArgs._DisplayName)
+				.ToolTipText(InArgs._DisplayName)
 				.Justification(ETextJustify::Left)
 				.OverflowPolicy(ETextOverflowPolicy::Ellipsis)
 				.Clipping(EWidgetClipping::ClipToBounds)
