[Audio Insights] Browse to asset and edit asset commands in Signal Flow
* Created new helper class for building context menus with the Browse and Edit options
* Integrated new helper into both the graph when r-clicking a node, and the dashboard for r-clicking an item in one of the list views
* Integrated hotkeys in the dashboard only, so the hotkeys will work with the selected item in the dashboard and graph.

#Synced-CL 50703972

#Synced-CL 50703972

diff --git a/Engine/Plugins/AudioInsights/Source/AudioInsights/Private/AssetEditorCommands.cpp b/Engine/Plugins/AudioInsights/Source/AudioInsights/Private/AssetEditorCommands.cpp
new file mode 100644
--- /dev/null
+++ b/Engine/Plugins/AudioInsights/Source/AudioInsights/Private/AssetEditorCommands.cpp
@@ -0,0 +1,25 @@
+// Copyright Epic Games, Inc. All Rights Reserved.
+#include "AssetEditorCommands.h"
+
+#include "AudioInsightsStyle.h"
+#include "Framework/Commands/InputChord.h"
+#include "Framework/Commands/UICommandInfo.h"
+
+#define LOCTEXT_NAMESPACE "AudioInsights"
+
+namespace UE::Audio::Insights
+{
+	FAssetEditorCommands::FAssetEditorCommands()
+		: TCommands<FAssetEditorCommands>("AssetEditorCommands", LOCTEXT("AssetEditorCommands_ContextDescText", "Asset Editor Commands"), NAME_None, UE::Audio::Insights::FSlateStyle::GetStyleName())
+	{
+	}
+
+	void FAssetEditorCommands::RegisterCommands()
+	{
+		UI_COMMAND(Browse, "Browse To Asset", "Browses to the selected node asset in the content browser.", EUserInterfaceActionType::Button, FInputChord(EKeys::B, EModifierKey::Control));
+		UI_COMMAND(Edit, "Edit", "Opens the selected node for edit.", EUserInterfaceActionType::Button, FInputChord(EKeys::E, EModifierKey::Control));
+	}
+
+} // namespace UE::Audio::Insights
+
+#undef LOCTEXT_NAMESPACE
diff --git a/Engine/Plugins/AudioInsights/Source/AudioInsights/Private/AssetEditorCommands.h b/Engine/Plugins/AudioInsights/Source/AudioInsights/Private/AssetEditorCommands.h
new file mode 100644
--- /dev/null
+++ b/Engine/Plugins/AudioInsights/Source/AudioInsights/Private/AssetEditorCommands.h
@@ -0,0 +1,25 @@
+// Copyright Epic Games, Inc. All Rights Reserved.
+#pragma once
+
+#include "Framework/Commands/Commands.h"
+#include "Templates/SharedPointer.h"
+
+class FUICommandInfo;
+
+namespace UE::Audio::Insights
+{
+	class FAssetEditorCommands : public TCommands<FAssetEditorCommands>
+	{
+	public:
+		FAssetEditorCommands();
+
+		virtual void RegisterCommands() override;
+
+		TSharedPtr<const FUICommandInfo> GetBrowseCommand() const { return Browse; }
+		TSharedPtr<const FUICommandInfo> GetEditCommand() const { return Edit; }
+
+	private:
+		TSharedPtr<FUICommandInfo> Browse;
+		TSharedPtr<FUICommandInfo> Edit;
+	};
+} // namespace UE::Audio::Insights
diff --git a/Engine/Plugins/AudioInsights/Source/AudioInsights/Private/AudioInsightsModule.cpp b/Engine/Plugins/AudioInsights/Source/AudioInsights/Private/AudioInsightsModule.cpp
--- a/Engine/Plugins/AudioInsights/Source/AudioInsights/Private/AudioInsightsModule.cpp
+++ b/Engine/Plugins/AudioInsights/Source/AudioInsights/Private/AudioInsightsModule.cpp
@@ -13,6 +13,10 @@
 #include "TraceServices/ModuleService.h"
 #include "UObject/NameTypes.h"
 
+#if WITH_EDITOR
+#include "AssetEditorCommands.h"
+#endif // WITH_EDITOR
+
 #if !WITH_EDITOR
 #include "AudioInsightsComponent.h"
 #include "Views/AudioBusDashboardViewFactory.h"
@@ -53,6 +57,10 @@ namespace UE::Audio::Insights
 			DashboardFactory = MakeShared<FDashboardFactory>();
 			FDashboardAssetCommands::Register();
 
+#if WITH_EDITOR
+			FAssetEditorCommands::Register();
+#endif // WITH_EDITOR
+
 			CacheManager = MakeUnique<FAudioInsightsCacheManager>();
 
 #if !WITH_EDITOR
@@ -104,6 +112,10 @@ namespace UE::Audio::Insights
 
 			FDashboardAssetCommands::Unregister();
 
+#if WITH_EDITOR
+			FAssetEditorCommands::Unregister();
+#endif // WITH_EDITOR
+
 			DashboardFactory.Reset();
 
 			IModularFeatures::Get().UnregisterModularFeature(TraceServices::ModuleFeatureName, TraceModule.Get());
diff --git a/Engine/Plugins/AudioInsights/Source/AudioInsights/Private/Views/AssetEditorContextMenuHelper.cpp b/Engine/Plugins/AudioInsights/Source/AudioInsights/Private/Views/AssetEditorContextMenuHelper.cpp
new file mode 100644
--- /dev/null
+++ b/Engine/Plugins/AudioInsights/Source/AudioInsights/Private/Views/AssetEditorContextMenuHelper.cpp
@@ -0,0 +1,128 @@
+// Copyright Epic Games, Inc. All Rights Reserved.
+#include "Views/AssetEditorContextMenuHelper.h"
+
+#if WITH_EDITOR
+#include "AssetEditorCommands.h"
+#include "AudioInsightsStyle.h"
+#include "Editor.h"
+#include "Framework/Commands/UICommandList.h"
+#include "Framework/MultiBox/MultiBoxBuilder.h"
+#include "Input/Events.h"
+#include "Subsystems/AssetEditorSubsystem.h"
+#include "Views/TableDashboardViewFactory.h"
+#include "Widgets/SWidget.h"
+
+#define LOCTEXT_NAMESPACE "AudioInsights"
+
+namespace UE::Audio::Insights
+{
+	FAssetEditorContextMenuHelper::FAssetEditorContextMenuHelper()
+	{
+		BindCommands();
+	}
+
+	TSharedPtr<SWidget> FAssetEditorContextMenuHelper::ContructContextMenuOptions()
+	{
+		if (GetEditableAsset() == nullptr)
+		{
+			return SNullWidget::NullWidget;
+		}
+
+		const FAssetEditorCommands& Commands = FAssetEditorCommands::Get();
+		FMenuBuilder MenuBuilder(true /*bInShouldCloseWindowAfterMenuSelection*/, CommandList);
+
+		MenuBuilder.BeginSection("AssetEditorActions", LOCTEXT("AudioDashboard_AssetEditorActions_RightClick_HeaderText", "Asset options"));
+		{
+			MenuBuilder.AddMenuEntry(Commands.GetBrowseCommand(), NAME_None, TAttribute<FText>(), TAttribute<FText>(), FSlateStyle::Get().CreateIcon("AudioInsights.Icon.SoundDashboard.Browse"));
+			MenuBuilder.AddMenuEntry(Commands.GetEditCommand(), NAME_None, TAttribute<FText>(), TAttribute<FText>(), FSlateStyle::Get().CreateIcon("AudioInsights.Icon.SoundDashboard.Edit"));
+		}
+
+		MenuBuilder.EndSection();
+
+		return MenuBuilder.MakeWidget();
+	}
+
+	bool FAssetEditorContextMenuHelper::ProcessCommandBindings(const FKeyEvent& InKeyEvent)
+	{
+		return CommandList.IsValid() && CommandList->ProcessCommandBindings(InKeyEvent);
+	}
+
+	void FAssetEditorContextMenuHelper::SetAssetEntry(const TSharedPtr<IObjectDashboardEntry>& InEntry)
+	{
+		if (InEntry.IsValid())
+		{
+			Entry = InEntry.ToWeakPtr();
+		}
+	}
+
+	void FAssetEditorContextMenuHelper::ResetAssetEntry()
+	{
+		Entry.Reset();
+	}
+
+	void FAssetEditorContextMenuHelper::BindCommands()
+	{
+		CommandList = MakeShared<FUICommandList>();
+
+		const FAssetEditorCommands& Commands = FAssetEditorCommands::Get();
+
+		CommandList->MapAction(
+			Commands.GetBrowseCommand(),
+			FExecuteAction::CreateLambda([this]() { BrowseToAsset(); }),
+			FCanExecuteAction::CreateLambda([this]() { return GetEditableAsset() != nullptr; }));
+
+		CommandList->MapAction(
+			Commands.GetEditCommand(),
+			FExecuteAction::CreateLambda([this]() { OpenAsset(); }),
+			FCanExecuteAction::CreateLambda([this]() { return GetEditableAsset() != nullptr; }));
+	}
+
+	bool FAssetEditorContextMenuHelper::BrowseToAsset() const
+	{
+		if (GEditor)
+		{
+			if (const TObjectPtr<UObject> Asset = GetEditableAsset())
+			{
+				GEditor->SyncBrowserToObject(Asset);
+				return true;
+			}
+		}
+
+		return false;
+	}
+
+	bool FAssetEditorContextMenuHelper::OpenAsset() const
+	{
+		if (GEditor)
+		{
+			if (const TObjectPtr<UObject> Asset = GetEditableAsset())
+			{
+				if (const TObjectPtr<UAssetEditorSubsystem> AssetSubsystem = GEditor->GetEditorSubsystem<UAssetEditorSubsystem>())
+				{
+					return AssetSubsystem->OpenEditorForAsset(Asset);
+				}
+			}
+		}
+
+		return false;
+	}
+
+	TObjectPtr<UObject> FAssetEditorContextMenuHelper::GetEditableAsset() const
+	{
+		const TSharedPtr<IObjectDashboardEntry> EntryPinned = Entry.Pin();
+		if (EntryPinned.IsValid())
+		{
+			const TObjectPtr<UObject> ObjectPtr = EntryPinned->GetObject();
+			if (ObjectPtr && ObjectPtr->IsAsset())
+			{
+				return ObjectPtr;
+			}
+		}
+
+		return nullptr;
+	}
+} // namespace UE::Audio::Insights
+
+#undef LOCTEXT_NAMESPACE
+
+#endif // WITH_EDITOR
\ No newline at end of file
diff --git a/Engine/Plugins/AudioInsights/Source/AudioInsights/Private/Views/SSignalFlowGraph.cpp b/Engine/Plugins/AudioInsights/Source/AudioInsights/Private/Views/SSignalFlowGraph.cpp
--- a/Engine/Plugins/AudioInsights/Source/AudioInsights/Private/Views/SSignalFlowGraph.cpp
+++ b/Engine/Plugins/AudioInsights/Source/AudioInsights/Private/Views/SSignalFlowGraph.cpp
@@ -18,6 +18,13 @@
 #include "Widgets/SBoxPanel.h"
 #include "Widgets/Text/STextBlock.h"
 
+#if WITH_EDITOR
+#include "Framework/Application/IMenu.h"
+#include "Framework/Application/MenuStack.h"
+#include "Framework/Application/SlateApplication.h"
+#include "Layout/WidgetPath.h"
+#endif // WITH_EDITOR
+
 #define LOCTEXT_NAMESPACE "AudioInsights"
 
 namespace UE::Audio::Insights
@@ -294,6 +301,11 @@ namespace UE::Audio::Insights
 		GraphOrientation = InArgs._GraphOrientation;
 		OnNodeSelected = InArgs._OnNodeSelected;
 
+#if WITH_EDITOR
+		OnOpenContextMenu = InArgs._OnOpenContextMenu;
+		GraphIsDraggingToMove = InArgs._GraphIsDraggingToMove;
+#endif // WITH_EDITOR
+
 		if (!Entry.IsValid())
 		{
 			return;
@@ -376,6 +388,26 @@ namespace UE::Audio::Insights
 		}
 	}
 
+#if WITH_EDITOR
+	FReply SSignalFlowGraphNode::OnMouseButtonUp(const FGeometry& MyGeometry, const FPointerEvent& MouseEvent)
+	{
+		if (MouseEvent.GetEffectingButton() == EKeys::RightMouseButton)
+		{
+			if (Entry.IsValid() && !GraphIsDraggingToMove.Get())
+			{
+				const TObjectPtr<UObject> ObjectPtr = Entry->GetObject();
+				if (ObjectPtr && ObjectPtr->IsAsset())
+				{
+					OnOpenContextMenu.ExecuteIfBound(MouseEvent, Entry);
+					return FReply::Handled();
+				}
+			}
+		}
+
+		return SCompoundWidget::OnMouseButtonUp(MyGeometry, MouseEvent);
+	}
+#endif // WITH_EDITOR
+
 	FVector2D SSignalFlowGraphNode::ComputeDesiredSize(float InLayoutScaleMultiplier) const
 	{
 		using namespace SSignalFlowGraphStyle;
@@ -1524,6 +1556,13 @@ namespace UE::Audio::Insights
 			{
 				OnNodeSelected.ExecuteIfBound(SelectedEntry);
 			})
+#if WITH_EDITOR
+			.OnOpenContextMenu(this, &SSignalFlowGraph::OnRequestedContextMenu)
+			.GraphIsDraggingToMove_Lambda([this]()
+			{
+				return bIsDraggingToMove;
+			})
+#endif // WITH_EDITOR
 			.ZoomScaleFactor_Lambda([this]()
 			{
 				return ZoomScaleFactor;
@@ -1730,6 +1769,42 @@ namespace UE::Audio::Insights
 
 		return FMath::IsNearlyEqual(TargetAlongOrientation, ScrollOffset, ScrollAnimInRangeThreshold) ? EAnimationResult::Complete : EAnimationResult::InProgress;
 	}
+
+#if WITH_EDITOR
+	void SSignalFlowGraph::OnRequestedContextMenu(const FPointerEvent& MouseEvent, const TSharedPtr<FSignalFlowDashboardEntry> Entry)
+	{
+		if (!Entry.IsValid())
+		{
+			return;
+		}
+
+		AssetContextMenuHelper.SetAssetEntry(Entry);
+
+		const TSharedPtr<SWidget> MenuContent = AssetContextMenuHelper.ContructContextMenuOptions();
+		if (MenuContent.IsValid())
+		{
+			const FVector2D SummonLocation = MouseEvent.GetScreenSpacePosition();
+			const FWidgetPath WidgetPath = MouseEvent.GetEventPath() != nullptr ? *MouseEvent.GetEventPath() : FWidgetPath();
+			GraphContextMenu = FSlateApplication::Get().PushMenu(AsShared(), WidgetPath, MenuContent.ToSharedRef(), SummonLocation, FPopupTransitionEffect(FPopupTransitionEffect::ContextMenu));
+
+			if (GraphContextMenu.IsValid())
+			{
+				GraphContextMenu->GetOnMenuDismissed().AddSPLambda(this, [this](TSharedRef<IMenu> InMenu)
+				{
+					GraphContextMenu.Reset();
+				});
+			}
+			else
+			{
+				AssetContextMenuHelper.ResetAssetEntry();
+			}
+		}
+		else
+		{
+			AssetContextMenuHelper.ResetAssetEntry();
+		}
+	}
+#endif // WITH_EDITOR
 } // namespace UE::Audio::Insights
 
 #undef LOCTEXT_NAMESPACE
diff --git a/Engine/Plugins/AudioInsights/Source/AudioInsights/Private/Views/SignalFlowDashboardViewFactory.cpp b/Engine/Plugins/AudioInsights/Source/AudioInsights/Private/Views/SignalFlowDashboardViewFactory.cpp
--- a/Engine/Plugins/AudioInsights/Source/AudioInsights/Private/Views/SignalFlowDashboardViewFactory.cpp
+++ b/Engine/Plugins/AudioInsights/Source/AudioInsights/Private/Views/SignalFlowDashboardViewFactory.cpp
@@ -205,6 +205,12 @@ namespace UE::Audio::Insights
 			{
 				return FReply::Handled();
 			}
+#if WITH_EDITOR
+			else if (AssetContextMenuHelper.ProcessCommandBindings(InKeyEvent))
+			{
+				return FReply::Handled();
+			}
+#endif // WITH_EDITOR
 			else if (InKeyEvent.GetKey() == EKeys::F)
 			{
 				return ScrollGraphViewToSelectedNode(true /*bResetZoom*/);
@@ -880,6 +886,9 @@ namespace UE::Audio::Insights
 					   .Style(GetRowStyle());
 			})
 			.OnSelectionChanged(this, &FSignalFlowDashboardViewFactory::OnSelectionChanged)
+#if WITH_EDITOR
+			.OnContextMenuOpening(this, &FSignalFlowDashboardViewFactory::OnConstructContextMenu)
+#endif // WITH_EDITOR
 			.HeaderRow
 			(
 				HeaderRowWidget
@@ -906,6 +915,9 @@ namespace UE::Audio::Insights
 	{
 		if (!ActiveSourcesListView.IsValid() || !ActiveBusesListView.IsValid() || !ActiveSubmixesListView.IsValid())
 		{
+#if WITH_EDITOR
+			AssetContextMenuHelper.ResetAssetEntry();
+#endif // WITH_EDITOR
 			SelectedItem.Reset();
 			return;
 		}
@@ -939,6 +951,10 @@ namespace UE::Audio::Insights
 			SelectedItem = InSelectedItem;
 			if (SelectedItem.IsValid())
 			{
+#if WITH_EDITOR
+				AssetContextMenuHelper.SetAssetEntry(SelectedItem);
+#endif // WITH_EDITOR
+
 				if (TSharedPtr<ISignalFlowNode>* SelectedNode = FilteredEntries.Find(SelectedItem->GetSignalFlowEntryKey()))
 				{
 					if (SelectedNode->IsValid())
@@ -950,6 +966,10 @@ namespace UE::Audio::Insights
 		}
 		else
 		{
+#if WITH_EDITOR
+			AssetContextMenuHelper.ResetAssetEntry();
+#endif // WITH_EDITOR
+
 			SelectedItem.Reset();
 			ActiveSourcesListView->ClearSelection();
 			ActiveBusesListView->ClearSelection();
@@ -1816,6 +1836,11 @@ namespace UE::Audio::Insights
 			Provider->SetConnectionWiresAreAnimated(bAnimateWires);
 		}
 	}
+
+	TSharedPtr<SWidget> FSignalFlowDashboardViewFactory::OnConstructContextMenu()
+	{
+		return AssetContextMenuHelper.ContructContextMenuOptions();
+	}
 #endif // WITH_EDITOR
 } // namespace UE::Audio::Insights
 
diff --git a/Engine/Plugins/AudioInsights/Source/AudioInsights/Public/Views/AssetEditorContextMenuHelper.h b/Engine/Plugins/AudioInsights/Source/AudioInsights/Public/Views/AssetEditorContextMenuHelper.h
new file mode 100644
--- /dev/null
+++ b/Engine/Plugins/AudioInsights/Source/AudioInsights/Public/Views/AssetEditorContextMenuHelper.h
@@ -0,0 +1,46 @@
+// Copyright Epic Games, Inc. All Rights Reserved.
+#pragma once
+
+#include "Templates/SharedPointer.h"
+#include "UObject/ObjectPtr.h"
+
+#define UE_API AUDIOINSIGHTS_API
+
+#if WITH_EDITOR
+struct FKeyEvent;
+class FUICommandList;
+class SWidget;
+class UObject;
+
+namespace UE::Audio::Insights
+{
+	class IObjectDashboardEntry;
+
+	// Helper for adding editor-only commands to the context menu of an Audio Insights dashboard
+	class FAssetEditorContextMenuHelper
+	{
+	public:
+		UE_API FAssetEditorContextMenuHelper();
+
+		UE_API TSharedPtr<SWidget> ContructContextMenuOptions();
+		UE_API bool ProcessCommandBindings(const FKeyEvent& InKeyEvent);
+
+		UE_API void SetAssetEntry(const TSharedPtr<IObjectDashboardEntry>& InEntry);
+		UE_API void ResetAssetEntry();
+
+	private:
+		void BindCommands();
+
+		bool BrowseToAsset() const;
+		bool OpenAsset() const;
+
+		TObjectPtr<UObject> GetEditableAsset() const;
+
+		TSharedPtr<FUICommandList> CommandList;
+		TWeakPtr<IObjectDashboardEntry> Entry;
+	};
+} // namespace UE::Audio::Insights
+
+#endif // WITH_EDITOR
+
+#undef UE_API
\ No newline at end of file
diff --git a/Engine/Plugins/AudioInsights/Source/AudioInsights/Public/Views/SSignalFlowGraph.h b/Engine/Plugins/AudioInsights/Source/AudioInsights/Public/Views/SSignalFlowGraph.h
--- a/Engine/Plugins/AudioInsights/Source/AudioInsights/Public/Views/SSignalFlowGraph.h
+++ b/Engine/Plugins/AudioInsights/Source/AudioInsights/Public/Views/SSignalFlowGraph.h
@@ -6,6 +6,10 @@
 #include "Widgets/SCanvas.h"
 #include "Widgets/SCompoundWidget.h"
 
+#if WITH_EDITOR
+#include "Views/AssetEditorContextMenuHelper.h"
+#endif // WITH_EDITOR
+
 #define UE_API AUDIOINSIGHTS_API
 
 struct FSignalFlowNodeDetailFilterSettings;
@@ -15,6 +19,10 @@ class SHorizontalBox;
 class SScrollBox;
 class SVerticalBox;
 
+#if WITH_EDITOR
+class IMenu;
+#endif // WITH_EDITOR
+
 namespace UE::Audio::Insights
 {
 	struct FSignalFlowDashboardEntry;
@@ -22,6 +30,10 @@ namespace UE::Audio::Insights
 
 	DECLARE_DELEGATE_OneParam(FOnSignalFlowNodeSelected, TSharedPtr<FSignalFlowDashboardEntry>);
 
+#if WITH_EDITOR
+	DECLARE_DELEGATE_TwoParams(FOnSignalFlowNodeContextMenuOpened, const FPointerEvent&, const TSharedPtr<FSignalFlowDashboardEntry>);
+#endif // WITH_EDITOR
+
 	namespace SSignalFlowGraphStyle
 	{
 		UE_API FSlateColor GetNodeAccentColor(const ESignalFlowEntryType NodeType);
@@ -41,6 +53,9 @@ namespace UE::Audio::Insights
 			_IsSelected(false),
 			_ShowNodeDetails(false),
 			_DisplayAmpPeakInDb(true),
+#if WITH_EDITOR
+			_GraphIsDraggingToMove(false),
+#endif // WITH_EDITOR
 			_GraphOrientation(EOrientation::Orient_Vertical)
 		{
 		}
@@ -55,14 +70,28 @@ namespace UE::Audio::Insights
 		SLATE_ATTRIBUTE(bool, IsSelected)
 		SLATE_ATTRIBUTE(bool, ShowNodeDetails)
 		SLATE_ATTRIBUTE(bool, DisplayAmpPeakInDb)
+
+#if WITH_EDITOR
+		SLATE_ATTRIBUTE(bool, GraphIsDraggingToMove)
+#endif // WITH_EDITOR
+
 		SLATE_ATTRIBUTE(EOrientation, GraphOrientation)
 		SLATE_EVENT(FOnSignalFlowNodeSelected, OnNodeSelected)
 
+#if WITH_EDITOR
+		SLATE_EVENT(FOnSignalFlowNodeContextMenuOpened, OnOpenContextMenu)
+#endif // WITH_EDITOR
+
 		SLATE_END_ARGS()
 
 		UE_API void Construct(const FArguments& InArgs);
 
 		UE_API virtual FReply OnMouseButtonDown(const FGeometry& MyGeometry, const FPointerEvent& MouseEvent) override;
+
+#if WITH_EDITOR
+		UE_API virtual FReply OnMouseButtonUp(const FGeometry& MyGeometry, const FPointerEvent& MouseEvent) override;
+#endif // WITH_EDITOR
+
 		UE_API virtual FVector2D ComputeDesiredSize(float InLayoutScaleMultiplier) const override;
 
 		UE_API FVector2D GetLocalInputPos() const;
@@ -98,6 +127,12 @@ namespace UE::Audio::Insights
 		TAttribute<EOrientation> GraphOrientation;
 
 		FOnSignalFlowNodeSelected OnNodeSelected;
+
+#if WITH_EDITOR
+		FOnSignalFlowNodeContextMenuOpened OnOpenContextMenu;
+		TAttribute<bool> GraphIsDraggingToMove;
+#endif // WITH_EDITOR
+
 		FVector2d CachedPositionInGraph = FVector2d::Zero();
 		FSignalFlowEntryKey NodeEntryKey;
 
@@ -254,6 +289,10 @@ namespace UE::Audio::Insights
 
 		EAnimationResult AnimateScrollTo(const TSharedRef<SScrollBox>& ScrollBox, const FVector2D& ScrollToPos, const float DeltaTime);
 
+#if WITH_EDITOR
+		void OnRequestedContextMenu(const FPointerEvent& MouseEvent, const TSharedPtr<FSignalFlowDashboardEntry> Entry);
+#endif // WITH_EDITOR
+
 		TSharedPtr<SCanvas> Canvas;
 
 		TMap<FSignalFlowEntryKey, TSharedPtr<SSignalFlowGraphNode>> NodeWidgets;
@@ -280,6 +319,12 @@ namespace UE::Audio::Insights
 		TAttribute<TWeakPtr<SScrollBox>> VerticalScrollBox;
 
 		FOnSignalFlowNodeSelected OnNodeSelected;
+
+#if WITH_EDITOR
+		TSharedPtr<IMenu> GraphContextMenu;
+		FAssetEditorContextMenuHelper AssetContextMenuHelper;
+#endif // WITH_EDITOR
+
 		EOrientation Orientation = EOrientation::Orient_Vertical;
 
 		float ZoomScaleFactor = 1.0f;
diff --git a/Engine/Plugins/AudioInsights/Source/AudioInsights/Public/Views/SignalFlowDashboardViewFactory.h b/Engine/Plugins/AudioInsights/Source/AudioInsights/Public/Views/SignalFlowDashboardViewFactory.h
--- a/Engine/Plugins/AudioInsights/Source/AudioInsights/Public/Views/SignalFlowDashboardViewFactory.h
+++ b/Engine/Plugins/AudioInsights/Source/AudioInsights/Public/Views/SignalFlowDashboardViewFactory.h
@@ -9,6 +9,10 @@
 #include "Widgets/Layout/SScrollBox.h"
 #include "Widgets/Views/SListView.h"
 
+#if WITH_EDITOR
+#include "Views/AssetEditorContextMenuHelper.h"
+#endif // WITH_EDITOR
+
 #define UE_API AUDIOINSIGHTS_API
 
 namespace UE::Audio::Insights
@@ -126,6 +130,7 @@ namespace UE::Audio::Insights
 #if WITH_EDITOR
 		void OnReadEditorSettings(const FSignalFlowSettings& InSettings);
 		void OnWriteEditorSettings(FSignalFlowSettings& OutSettings);
+		TSharedPtr<SWidget> OnConstructContextMenu();
 #endif // WITH_EDITOR
 
 		TSharedPtr<FSignalFlowTraceProvider> SignalFlowProvider;
@@ -159,6 +164,10 @@ namespace UE::Audio::Insights
 		bool bShowNodeDetails = false;
 		bool bDisplayAmpPeakInDb = true;
 		bool bAnimateWires = true;
+
+#if WITH_EDITOR
+		FAssetEditorContextMenuHelper AssetContextMenuHelper;
+#endif // WITH_EDITOR
 	};
 } // namespace UE::Audio::Insights
 
