[Audio Insights] Try to maintain the order of submix and bus nodes where possible
* Keep track of a node's previous node order ID
* Instead of falling back on timestamp for submixes and buses, put more weight on their previous order to try and maintain some stability in the graph.
* So we can do this, do not reset the filtered entries map each refresh - instead add new entries and remove no longer referenced ones as we go. This allows us to retain the state of members inside the filtered entries map (such as previous node order.

#Synced-CL 50815132

diff --git a/Engine/Plugins/AudioInsights/Source/AudioInsights/Private/Views/SignalFlowDashboardViewFactory.cpp b/Engine/Plugins/AudioInsights/Source/AudioInsights/Private/Views/SignalFlowDashboardViewFactory.cpp
--- a/Engine/Plugins/AudioInsights/Source/AudioInsights/Private/Views/SignalFlowDashboardViewFactory.cpp
+++ b/Engine/Plugins/AudioInsights/Source/AudioInsights/Private/Views/SignalFlowDashboardViewFactory.cpp
@@ -949,20 +949,12 @@ namespace UE::Audio::Insights
 			}
 
 			SelectedItem = InSelectedItem;
+#if WITH_EDITOR
 			if (SelectedItem.IsValid())
 			{
-#if WITH_EDITOR
 				AssetContextMenuHelper.SetAssetEntry(SelectedItem);
-#endif // WITH_EDITOR
-
-				if (TSharedPtr<ISignalFlowNode>* SelectedNode = FilteredEntries.Find(SelectedItem->GetSignalFlowEntryKey()))
-				{
-					if (SelectedNode->IsValid())
-					{
-						(*SelectedNode)->SetTreeDepth((*SelectedNode)->GetTreeDepth().LayerID, 0);
-					}
-				}
 			}
+#endif // WITH_EDITOR
 		}
 		else
 		{
@@ -1100,10 +1092,16 @@ namespace UE::Audio::Insights
 
 		bGraphRequiresRefresh = true;
 
-		FilteredEntries.Reset();
+		TSet<FSignalFlowEntryKey> ExistingEntries;
+		ExistingEntries.Reserve(FilteredEntries.Num());
+		for (const auto& [EntryKey, Entry] : FilteredEntries)
+		{
+			ExistingEntries.Add(EntryKey);
+		}
 
-		AddFilteredInputsRecursive(*DeviceData, FilterFromEntry, nullptr, 0);
-		AddFilteredOutputs(*DeviceData, FilterFromEntry, 1);
+		AddFilteredInputsRecursive(*DeviceData, FilterFromEntry, nullptr, 0, ExistingEntries);
+		AddFilteredOutputs(*DeviceData, FilterFromEntry, 1, ExistingEntries);
+		CleanUpFilteredEntries(ExistingEntries);
 
 		// Gather stats (such as root entry and tree depth structure)
 		TOptional<FSignalFlowEntryKey> RootEntryKey;
@@ -1163,7 +1161,7 @@ namespace UE::Audio::Insights
 		}
 	}
 
-	void FSignalFlowDashboardViewFactory::AddFilteredInputsRecursive(const TSortedMap<FSignalFlowEntryKey, TSharedPtr<FSignalFlowDashboardEntry>>& DeviceData, TSharedPtr<FSignalFlowDashboardEntry> Entry, TSharedPtr<FSignalFlowDashboardEntry> ChildEntry, int32 TreeDepth)
+	void FSignalFlowDashboardViewFactory::AddFilteredInputsRecursive(const TSortedMap<FSignalFlowEntryKey, TSharedPtr<FSignalFlowDashboardEntry>>& DeviceData, TSharedPtr<FSignalFlowDashboardEntry> Entry, TSharedPtr<FSignalFlowDashboardEntry> ChildEntry, int32 TreeDepth, TSet<FSignalFlowEntryKey>& OutExistingEntryKeys)
 	{
 		if (!Entry.IsValid())
 		{
@@ -1172,23 +1170,37 @@ namespace UE::Audio::Insights
 
 		TSharedPtr<ISignalFlowNode> EntryNode = nullptr;
 
-		TSharedPtr<ISignalFlowNode>* FoundNode = FilteredEntries.Find(Entry->GetSignalFlowEntryKey());
+		const FSignalFlowEntryKey EntryKey = Entry->GetSignalFlowEntryKey();
+		TSharedPtr<ISignalFlowNode>* FoundNode = FilteredEntries.Find(EntryKey);
 		if (FoundNode && FoundNode->IsValid())
 		{
-			const ESignalFlowEntryType EntryType = (*FoundNode)->GetEntryKey().EntryType;
-			const bool bCanBeLayered = EntryType == ESignalFlowEntryType::AudioBus || EntryType == ESignalFlowEntryType::Submix;
+			if (OutExistingEntryKeys.Contains(EntryKey))
+			{
+				// First time we have tried to recreate this node but it was here already
+				// Remove from tracking container and make a note of the new tree depth
+				OutExistingEntryKeys.Remove(EntryKey);
+				(*FoundNode)->GetTreeDepth().TreeDepth = TreeDepth;
+			}
+			else
+			{
+				// We have created and assigned assigned this node a tree depth - move it higher in the tree
+				const ESignalFlowEntryType EntryType = EntryKey.EntryType;
+				const bool bCanBeLayered = EntryType == ESignalFlowEntryType::AudioBus || EntryType == ESignalFlowEntryType::Submix;
 
-			TreeDepthPair& NodeTreeDepth = (*FoundNode)->GetTreeDepth();
+				TreeDepthPair& NodeTreeDepth = (*FoundNode)->GetTreeDepth();
 
-			if (bCanBeLayered && NodeTreeDepth.TreeDepth >= TreeDepth)
-			{
-				NodeTreeDepth.TreeDepth = TreeDepth - 1;
+				if (bCanBeLayered && NodeTreeDepth.TreeDepth >= TreeDepth)
+				{
+					NodeTreeDepth.TreeDepth = TreeDepth - 1;
+				}
 			}
 
 			EntryNode = *FoundNode;
 		}
 		else
 		{
+			ensure(!OutExistingEntryKeys.Contains(EntryKey));
+
 			EntryNode = MakeShared<FSignalFlowEntryNode>(Entry, TreeDepth, Entry->EntryType, Entry->Timestamp);
 
 			FilteredEntries.Add(Entry->GetSignalFlowEntryKey(), EntryNode);
@@ -1211,19 +1223,19 @@ namespace UE::Audio::Insights
 				}
 
 				bool bResetTreeDepth = (*InputEntry)->EntryType != Entry->EntryType;
-				AddFilteredInputsRecursive(DeviceData, *InputEntry, Entry, bResetTreeDepth ? 0 : TreeDepth - 1);
+				AddFilteredInputsRecursive(DeviceData, *InputEntry, Entry, bResetTreeDepth ? 0 : TreeDepth - 1, OutExistingEntryKeys);
 			}
 		}
 	}
 
-	void FSignalFlowDashboardViewFactory::AddFilteredOutputs(const TSortedMap<FSignalFlowEntryKey, TSharedPtr<FSignalFlowDashboardEntry>>& DeviceData, TSharedPtr<FSignalFlowDashboardEntry> Entry, int32 TreeDepth)
+	void FSignalFlowDashboardViewFactory::AddFilteredOutputs(const TSortedMap<FSignalFlowEntryKey, TSharedPtr<FSignalFlowDashboardEntry>>& DeviceData, TSharedPtr<FSignalFlowDashboardEntry> Entry, int32 TreeDepth, TSet<FSignalFlowEntryKey>& OutExistingEntryKeys)
 	{
 		for (const auto& [OutputKey, SendInfo] : Entry->Outputs)
 		{
 			const TSharedPtr<FSignalFlowDashboardEntry>* OutputEntry = DeviceData.Find(OutputKey);
 			if (OutputEntry && OutputEntry->IsValid())
 			{
-				AddFilteredOutputsRecursive(DeviceData, *OutputEntry, Entry, TreeDepth);
+				AddFilteredOutputsRecursive(DeviceData, *OutputEntry, Entry, TreeDepth, OutExistingEntryKeys);
 
 				TSharedPtr<ISignalFlowNode>* EntryNode = FilteredEntries.Find(Entry->GetSignalFlowEntryKey());
 				TSharedPtr<ISignalFlowNode>* OutputNode = FilteredEntries.Find(OutputKey);
@@ -1235,7 +1247,7 @@ namespace UE::Audio::Insights
 		}
 	}
 
-	void FSignalFlowDashboardViewFactory::AddFilteredOutputsRecursive(const TSortedMap<FSignalFlowEntryKey, TSharedPtr<FSignalFlowDashboardEntry>>& DeviceData, TSharedPtr<FSignalFlowDashboardEntry> Entry, TSharedPtr<FSignalFlowDashboardEntry> ParentEntry, int32 TreeDepth)
+	void FSignalFlowDashboardViewFactory::AddFilteredOutputsRecursive(const TSortedMap<FSignalFlowEntryKey, TSharedPtr<FSignalFlowDashboardEntry>>& DeviceData, TSharedPtr<FSignalFlowDashboardEntry> Entry, TSharedPtr<FSignalFlowDashboardEntry> ParentEntry, int32 TreeDepth, TSet<FSignalFlowEntryKey>& OutExistingEntryKeys)
 	{
 		if (!Entry.IsValid())
 		{
@@ -1249,16 +1261,29 @@ namespace UE::Audio::Insights
 
 		if (FoundNode && FoundNode->IsValid())
 		{
-			TreeDepthPair& NodeTreeDepth = (*FoundNode)->GetTreeDepth();
-			if (NodeTreeDepth.TreeDepth <= TreeDepth)
+			if (OutExistingEntryKeys.Contains(EntryKey))
+			{
+				// First time we have tried to recreate this node but it was here already
+				// Remove from tracking container and make a note of the new tree depth
+				OutExistingEntryKeys.Remove(EntryKey);
+				(*FoundNode)->GetTreeDepth().TreeDepth = TreeDepth;
+			}
+			else
 			{
-				NodeTreeDepth.TreeDepth = TreeDepth + 1;
+				// We have created and assigned assigned this node a tree depth - move it lower in the tree
+				TreeDepthPair& NodeTreeDepth = (*FoundNode)->GetTreeDepth();
+				if (NodeTreeDepth.TreeDepth <= TreeDepth)
+				{
+					NodeTreeDepth.TreeDepth = TreeDepth + 1;
+				}
 			}
 
 			Node = (*FoundNode);
 		}
 		else
 		{
+			ensure(!OutExistingEntryKeys.Contains(EntryKey));
+
 			Node = MakeShared<FSignalFlowEntryNode>(Entry, TreeDepth, Entry->EntryType, Entry->Timestamp);
 
 			FilteredEntries.Add(EntryKey, Node);
@@ -1281,8 +1306,32 @@ namespace UE::Audio::Insights
 
 				bool bResetTreeDepth = (*OutputEntry)->EntryType != Entry->EntryType;
 
-				AddFilteredOutputsRecursive(DeviceData, *OutputEntry, Entry, bResetTreeDepth ? 0 : TreeDepth + 1);
+				AddFilteredOutputsRecursive(DeviceData, *OutputEntry, Entry, bResetTreeDepth ? 0 : TreeDepth + 1, OutExistingEntryKeys);
+			}
+		}
+	}
+
+	void FSignalFlowDashboardViewFactory::CleanUpFilteredEntries(const TSet<FSignalFlowEntryKey>& EntriesToRemove)
+	{
+		if (EntriesToRemove.IsEmpty())
+		{
+			return;
+		}
+
+		for (const auto& [EntryKey, Entry] : FilteredEntries)
+		{
+			if (EntriesToRemove.Contains(EntryKey))
+			{
+				continue;
 			}
+
+			Entry->FilteredInputs.RemoveAll([&EntriesToRemove](const FSignalFlowEntryKey& Input) { return EntriesToRemove.Contains(Input); });
+			Entry->FilteredOutputs.RemoveAll([&EntriesToRemove](const FSignalFlowEntryKey& Output) { return EntriesToRemove.Contains(Output); });
+		}
+
+		for (const FSignalFlowEntryKey& EntryKey : EntriesToRemove)
+		{
+			FilteredEntries.Remove(EntryKey);
 		}
 	}
 
@@ -1487,9 +1536,17 @@ namespace UE::Audio::Insights
 				}
 
 				// Try timestamp of this node
-				if (!FMath::IsNearlyEqual((*FirstNode)->GetTimestamp(), (*SecondNode)->GetTimestamp()))
+				if (NodeIsOwnerOrSound(**FirstNode) && NodeIsOwnerOrSound(**SecondNode))
 				{
-					return (*FirstNode)->GetTimestamp() < (*SecondNode)->GetTimestamp();
+					if (!FMath::IsNearlyEqual((*FirstNode)->GetTimestamp(), (*SecondNode)->GetTimestamp()))
+					{
+						return (*FirstNode)->GetTimestamp() < (*SecondNode)->GetTimestamp();
+					}
+				}
+
+				if ((*FirstNode)->GetPreviousNodeOrderID() != (*SecondNode)->GetPreviousNodeOrderID())
+				{
+					return (*FirstNode)->GetPreviousNodeOrderID() < (*SecondNode)->GetPreviousNodeOrderID();
 				}
 			}
 
@@ -1527,10 +1584,8 @@ namespace UE::Audio::Insights
 
 			if (Node && Node->IsValid())
 			{
-				const ESignalFlowEntryType NodeEntryType = (*Node)->GetEntryKey().EntryType;
-
 				// Only count entries that are transient when finding the minimum timestamp of a branch
-				if (NodeEntryType == ESignalFlowEntryType::SoundSource || NodeEntryType == ESignalFlowEntryType::OwnerObject)
+				if (NodeIsOwnerOrSound(**Node))
 				{
 					OutMinTimestamp = FMath::Min((*Node)->GetTimestamp(), OutMinTimestamp);
 				}
@@ -1540,6 +1595,11 @@ namespace UE::Audio::Insights
 		}
 	}
 
+	bool FSignalFlowDashboardViewFactory::NodeIsOwnerOrSound(const ISignalFlowNode& Node) const
+	{
+		return Node.GetEntryKey().EntryType == ESignalFlowEntryType::SoundSource || Node.GetEntryKey().EntryType == ESignalFlowEntryType::OwnerObject;
+	}
+
 	void FSignalFlowDashboardViewFactory::FixUnassignedNodeOrders()
 	{
 		// Run through again and find any nodes that did not get an order assigned
diff --git a/Engine/Plugins/AudioInsights/Source/AudioInsights/Private/Views/SignalFlowNodes.cpp b/Engine/Plugins/AudioInsights/Source/AudioInsights/Private/Views/SignalFlowNodes.cpp
--- a/Engine/Plugins/AudioInsights/Source/AudioInsights/Private/Views/SignalFlowNodes.cpp
+++ b/Engine/Plugins/AudioInsights/Source/AudioInsights/Private/Views/SignalFlowNodes.cpp
@@ -10,4 +10,10 @@ namespace UE::Audio::Insights
 		ensure(Entry.IsValid());
 		return Entry.IsValid() ? Entry->GetSignalFlowEntryKey() : FSignalFlowEntryKey();
 	}
+
+	void ISignalFlowNode::ResetOrderID()
+	{
+		PreviousNodeOrderID = NodeOrderID;
+		NodeOrderID = INVALID_NODE_ORDER_ID;
+	}
 } // namespace UE::Audio::Insights
\ No newline at end of file
diff --git a/Engine/Plugins/AudioInsights/Source/AudioInsights/Public/Views/SignalFlowDashboardViewFactory.h b/Engine/Plugins/AudioInsights/Source/AudioInsights/Public/Views/SignalFlowDashboardViewFactory.h
--- a/Engine/Plugins/AudioInsights/Source/AudioInsights/Public/Views/SignalFlowDashboardViewFactory.h
+++ b/Engine/Plugins/AudioInsights/Source/AudioInsights/Public/Views/SignalFlowDashboardViewFactory.h
@@ -104,9 +104,10 @@ namespace UE::Audio::Insights
 		void CreateFilteredNodeGraph(const FSignalFlowTraceProvider& Provider, TSharedPtr<FSignalFlowDashboardEntry> FilterFromEntry);
 		
 		void AddAllNodesToGraph();
-		void AddFilteredInputsRecursive(const TSortedMap<FSignalFlowEntryKey, TSharedPtr<FSignalFlowDashboardEntry>>& DeviceData, TSharedPtr<FSignalFlowDashboardEntry> Entry, TSharedPtr<FSignalFlowDashboardEntry> ChildEntry, int32 TreeDepth);
-		void AddFilteredOutputs(const TSortedMap<FSignalFlowEntryKey, TSharedPtr<FSignalFlowDashboardEntry>>& DeviceData, TSharedPtr<FSignalFlowDashboardEntry> Entry, int32 TreeDepth);
-		void AddFilteredOutputsRecursive(const TSortedMap<FSignalFlowEntryKey, TSharedPtr<FSignalFlowDashboardEntry>>& DeviceData, TSharedPtr<FSignalFlowDashboardEntry> Entry, TSharedPtr<FSignalFlowDashboardEntry> ParentEntry, int32 TreeDepth);
+		void AddFilteredInputsRecursive(const TSortedMap<FSignalFlowEntryKey, TSharedPtr<FSignalFlowDashboardEntry>>& DeviceData, TSharedPtr<FSignalFlowDashboardEntry> Entry, TSharedPtr<FSignalFlowDashboardEntry> ChildEntry, int32 TreeDepth, TSet<FSignalFlowEntryKey>& OutExistingEntryKeys);
+		void AddFilteredOutputs(const TSortedMap<FSignalFlowEntryKey, TSharedPtr<FSignalFlowDashboardEntry>>& DeviceData, TSharedPtr<FSignalFlowDashboardEntry> Entry, int32 TreeDepth, TSet<FSignalFlowEntryKey>& OutExistingEntryKeys);
+		void AddFilteredOutputsRecursive(const TSortedMap<FSignalFlowEntryKey, TSharedPtr<FSignalFlowDashboardEntry>>& DeviceData, TSharedPtr<FSignalFlowDashboardEntry> Entry, TSharedPtr<FSignalFlowDashboardEntry> ParentEntry, int32 TreeDepth, TSet<FSignalFlowEntryKey>& OutExistingEntryKeys);
+		void CleanUpFilteredEntries(const TSet<FSignalFlowEntryKey>& EntriesToRemove);
 
 		void CalculateFilteredGraphDepthStructure(TOptional<FSignalFlowEntryKey>& OutRootEntryKey, TArray<TreeDepthPair>& OutSortedTreeDepths) const;
 
@@ -117,6 +118,7 @@ namespace UE::Audio::Insights
 		void SortNodeInputs(TSharedPtr<ISignalFlowNode> OutNode);
 		TreeDepthPair GetOriginTreeDepthForNode(const TSharedPtr<ISignalFlowNode>& Node) const;
 		void GetMinTimestampRecursive(const TArray<FSignalFlowEntryKey>& NodeKeys, double& OutMinTimestamp) const;
+		bool NodeIsOwnerOrSound(const ISignalFlowNode& Node) const;
 
 		void FixUnassignedNodeOrders();
 		void SortNodeOrder();
diff --git a/Engine/Plugins/AudioInsights/Source/AudioInsights/Public/Views/SignalFlowNodes.h b/Engine/Plugins/AudioInsights/Source/AudioInsights/Public/Views/SignalFlowNodes.h
--- a/Engine/Plugins/AudioInsights/Source/AudioInsights/Public/Views/SignalFlowNodes.h
+++ b/Engine/Plugins/AudioInsights/Source/AudioInsights/Public/Views/SignalFlowNodes.h
@@ -63,8 +63,9 @@ namespace UE::Audio::Insights
 		TreeDepthPair& GetTreeDepth() { return TreeDepthPair; }
 
 		void SetNodeOrderID(const int32 InNodeOrderID) { NodeOrderID = InNodeOrderID; }
-		void ResetOrderID() { NodeOrderID = INVALID_NODE_ORDER_ID; }
+		void ResetOrderID();
 		int32 GetNodeOrderID() const { return NodeOrderID; }
+		int32 GetPreviousNodeOrderID() const { return PreviousNodeOrderID; }
 		bool NodeOrderIsValid() const { return NodeOrderID != INVALID_NODE_ORDER_ID; }
 
 		double GetTimestamp() const { return Timestamp; }
@@ -79,6 +80,7 @@ namespace UE::Audio::Insights
 
 		TreeDepthPair TreeDepthPair;
 		int32 NodeOrderID = INVALID_NODE_ORDER_ID;
+		int32 PreviousNodeOrderID = INVALID_NODE_ORDER_ID;
 
 		const double Timestamp = 0.0;
 	};
