[Audio Insights] "Pause graph" on selecting a node toggle option
* Option to freeze audio insights at the current timestamp when selecting a node in the signal flow graph
* Once the graph is paused, the timestamp will remain until processing is resumed again
* Added hotkey ESC and temp button to resume processing
* Unselecting all nodes will resume processing if the timestamp is controlled by the signal flow graph
* Hotkey Ctrl + P to toggle pausing on select.
* Reset the graph when time is reset. This avoids zombie nodes being left in the graph without updating.

#Synced-CL 50815132

diff --git a/Engine/Plugins/AudioInsights/Source/AudioInsights/Private/AudioInsightsTimingViewExtender.h b/Engine/Plugins/AudioInsights/Source/AudioInsights/Private/AudioInsightsTimingViewExtender.h
--- a/Engine/Plugins/AudioInsights/Source/AudioInsights/Private/AudioInsightsTimingViewExtender.h
+++ b/Engine/Plugins/AudioInsights/Source/AudioInsights/Private/AudioInsightsTimingViewExtender.h
@@ -30,6 +30,7 @@ namespace UE::Audio::Insights
 	{
 		EventLog = 0,
 		PlotsWidget,
+		SignalFlow,
 		External
 	};
 
diff --git a/Engine/Plugins/AudioInsights/Source/AudioInsights/Private/SignalFlowEditorCommands.cpp b/Engine/Plugins/AudioInsights/Source/AudioInsights/Private/SignalFlowEditorCommands.cpp
--- a/Engine/Plugins/AudioInsights/Source/AudioInsights/Private/SignalFlowEditorCommands.cpp
+++ b/Engine/Plugins/AudioInsights/Source/AudioInsights/Private/SignalFlowEditorCommands.cpp
@@ -18,7 +18,11 @@ namespace UE::Audio::Insights
 	{
 		UI_COMMAND(ToggleHorizontalView, "Horizontal Flow", "Toggles between horizontally/vertically orienting the Signal Flow graph.", EUserInterfaceActionType::ToggleButton, FInputChord(EKeys::H, EModifierKey::Control));
 		UI_COMMAND(ShowNodeDetails, "Show Node Details", "Toggles visibility of node details.", EUserInterfaceActionType::ToggleButton, FInputChord(EKeys::D, EModifierKey::Control));
+		UI_COMMAND(PauseOnSelect, "Pause on Select", "Toggles whether selecting a node should pause processing trace data, freezing audio insights at the latest timestamp.", EUserInterfaceActionType::ToggleButton, FInputChord(EKeys::P, EModifierKey::Control));
 		UI_COMMAND(ToggleAnimateWires, "Animate Wires", "Toggles whether connection wire animations driven by amplitude envelope values are enabled in the Signal Flow graph.", EUserInterfaceActionType::ToggleButton, FInputChord());
+	
+		UI_COMMAND(CenterView, "Center View", "Focus the graph view on the selected node.", EUserInterfaceActionType::Button, FInputChord(EKeys::F, EModifierKey::Control));
+		UI_COMMAND(ResetTimestampPause, "Reset Inspect Timestamp Signal Flow", "Resumes signal flow graph to showing the latest trace data if previously paused.", EUserInterfaceActionType::Button, FInputChord(EKeys::Escape));
 	}
 
 } // namespace UE::Audio::Insights
diff --git a/Engine/Plugins/AudioInsights/Source/AudioInsights/Private/SignalFlowEditorCommands.h b/Engine/Plugins/AudioInsights/Source/AudioInsights/Private/SignalFlowEditorCommands.h
--- a/Engine/Plugins/AudioInsights/Source/AudioInsights/Private/SignalFlowEditorCommands.h
+++ b/Engine/Plugins/AudioInsights/Source/AudioInsights/Private/SignalFlowEditorCommands.h
@@ -15,13 +15,25 @@ namespace UE::Audio::Insights
 
 		virtual void RegisterCommands() override;
 
+		// Toggles
 		TSharedPtr<const FUICommandInfo> GetToggleHorizontalViewCommand() const { return ToggleHorizontalView; }
 		TSharedPtr<const FUICommandInfo> GetShowNodeDetailsCommand() const { return ShowNodeDetails; }
+		TSharedPtr<const FUICommandInfo> GetPauseOnSelectCommand() const { return PauseOnSelect; }
 		TSharedPtr<const FUICommandInfo> GetToggleAnimateWiresCommand() const { return ToggleAnimateWires; }
 
+		// Actions
+		TSharedPtr<const FUICommandInfo> GetCenterViewCommand() const { return CenterView; }
+		TSharedPtr<const FUICommandInfo> GetResetTimestampPauseCommand() const { return ResetTimestampPause; }
+
 	private:
+		// Toggles
 		TSharedPtr<FUICommandInfo> ToggleHorizontalView;
 		TSharedPtr<FUICommandInfo> ShowNodeDetails;
+		TSharedPtr<FUICommandInfo> PauseOnSelect;
 		TSharedPtr<FUICommandInfo> ToggleAnimateWires;
+
+		// Actions
+		TSharedPtr<FUICommandInfo> CenterView;
+		TSharedPtr<FUICommandInfo> ResetTimestampPause;
 	};
 } // namespace UE::Audio::Insights
diff --git a/Engine/Plugins/AudioInsights/Source/AudioInsights/Private/Views/SSignalFlowGraph.cpp b/Engine/Plugins/AudioInsights/Source/AudioInsights/Private/Views/SSignalFlowGraph.cpp
--- a/Engine/Plugins/AudioInsights/Source/AudioInsights/Private/Views/SSignalFlowGraph.cpp
+++ b/Engine/Plugins/AudioInsights/Source/AudioInsights/Private/Views/SSignalFlowGraph.cpp
@@ -863,6 +863,18 @@ namespace UE::Audio::Insights
 		CleanUpOldWidgets(RemainingNodeKeys);
 	}
 
+	void SSignalFlowGraph::ResetGraph()
+	{
+		if (Canvas.IsValid())
+		{
+			Canvas->ClearChildren();
+		}
+
+		NodeWidgets.Reset();
+		DummyToRealConnectionKeys.Reset();
+		SendLevelWidgets.Reset();
+	}
+
 	void SSignalFlowGraph::ToggleOrientation()
 	{
 		Orientation = (Orientation == EOrientation::Orient_Vertical) ? EOrientation::Orient_Horizontal : EOrientation::Orient_Vertical;
diff --git a/Engine/Plugins/AudioInsights/Source/AudioInsights/Private/Views/SignalFlowDashboardViewFactory.cpp b/Engine/Plugins/AudioInsights/Source/AudioInsights/Private/Views/SignalFlowDashboardViewFactory.cpp
--- a/Engine/Plugins/AudioInsights/Source/AudioInsights/Private/Views/SignalFlowDashboardViewFactory.cpp
+++ b/Engine/Plugins/AudioInsights/Source/AudioInsights/Private/Views/SignalFlowDashboardViewFactory.cpp
@@ -4,6 +4,7 @@
 #include "AudioInsightsModule.h"
 #include "AudioInsightsStyle.h"
 #include "AudioInsightsTimingViewExtender.h"
+#include "Cache/AudioInsightsCacheManager.h"
 #include "Framework/MultiBox/MultiBoxBuilder.h"
 #include "IAudioInsightsModule.h"
 #include "Internationalization/Text.h"
@@ -80,6 +81,12 @@ namespace UE::Audio::Insights
 
 			return FText::GetEmpty();
 		}
+
+		const FSlateRoundedBoxBrush& GetRoundedBoxBrush()
+		{
+			static const FSlateRoundedBoxBrush RoundedBoxBrush(FLinearColor::White, 8.0f);
+			return RoundedBoxBrush;
+		}
 	}
 
 	FReply SSignalFlowDashboard::OnKeyDown(const FGeometry& MyGeometry, const FKeyEvent& InKeyEvent)
@@ -199,7 +206,7 @@ namespace UE::Audio::Insights
 
 	FReply FSignalFlowDashboardViewFactory::OnKeyDown(const FGeometry& MyGeometry, const FKeyEvent& InKeyEvent)
 	{
-		if (GraphWidgetContents.IsValid() && InKeyEvent.IsControlDown())
+		if (GraphWidgetContents.IsValid())
 		{
 			if (CommandList && CommandList->ProcessCommandBindings(InKeyEvent))
 			{
@@ -211,10 +218,6 @@ namespace UE::Audio::Insights
 				return FReply::Handled();
 			}
 #endif // WITH_EDITOR
-			else if (InKeyEvent.GetKey() == EKeys::F)
-			{
-				return ScrollGraphViewToSelectedNode(true /*bResetZoom*/);
-			}
 		}
 
 		return FReply::Unhandled();
@@ -222,6 +225,11 @@ namespace UE::Audio::Insights
 
 	void FSignalFlowDashboardViewFactory::Reset()
 	{
+		if (GraphWidgetContents.IsValid())
+		{
+			GraphWidgetContents->ResetGraph();
+		}
+
 		SetNewSelection(nullptr);
 	}
 
@@ -286,11 +294,25 @@ namespace UE::Audio::Insights
 			FCanExecuteAction(),
 			FIsActionChecked::CreateLambda([this]() { return bShowNodeDetails; }));
 
+		CommandList->MapAction(
+			Commands.GetPauseOnSelectCommand(),
+			FExecuteAction::CreateLambda([this]() { TogglePauseGraphEnabled(); }),
+			FCanExecuteAction(),
+			FIsActionChecked::CreateLambda([this]() { return bPauseGraphOnSelect; }));
+
 		CommandList->MapAction(
 			Commands.GetToggleAnimateWiresCommand(),
 			FExecuteAction::CreateLambda([this]() { ToggleAnimateWires(); }),
 			FCanExecuteAction(),
 			FIsActionChecked::CreateLambda([this]() { return bAnimateWires; }));
+
+		CommandList->MapAction(
+			Commands.GetCenterViewCommand(),
+			FExecuteAction::CreateLambda([this]() { ScrollGraphViewToSelectedNode(true /*bResetZoom*/); }));
+
+		CommandList->MapAction(
+			Commands.GetResetTimestampPauseCommand(),
+			FExecuteAction::CreateLambda([this]() { ResetPauseTimestamp(); }));
 	}
 
 	TSharedRef<SWidget> FSignalFlowDashboardViewFactory::CreateToolbar()
@@ -384,6 +406,11 @@ namespace UE::Audio::Insights
 			return bShowNodeDetails ? FSlateColor(FColor::Green) : FSlateColor::UseForeground();
 		};
 
+		TFunction<FSlateColor()> PauseOnSelectColorFunc = [this]()
+		{ 
+			return bPauseGraphOnSelect ? FSlateColor(FColor::Green) : FSlateColor::UseForeground();
+		};
+
 		return SNew(SHorizontalBox)
 
 		+ SHorizontalBox::Slot()
@@ -428,6 +455,28 @@ namespace UE::Audio::Insights
 												[]() { return LOCTEXT("AudioDashboard_SignalFlow_NodeDetails", "Node Details"); },
 												ToolbarButtonHeight)
 			]
+		]
+		
+		+ SHorizontalBox::Slot()
+		.AutoWidth()
+		.Padding(ToolbarButtonPaddingX, ToolbarButtonPaddingY)
+		[
+			SNew(SButton)
+			.ButtonStyle(&FAppStyle::Get().GetWidgetStyle<FButtonStyle>(ButtonStyleName))
+			.ToolTipText(LOCTEXT("AudioDashboard_SignalFlow_PauseTimestampTooltip", "Toggles pausing the graph at the current timestamp when selecting a node. (Ctrl + p)"))
+			.OnClicked_Lambda([this]()
+			{
+				TogglePauseGraphEnabled();
+
+				return FReply::Handled();
+			})
+			[
+				Style.CreateToggleButtonContent("AudioInsights.Icon.Pause",
+												ButtonTextStyleName,
+												PauseOnSelectColorFunc,
+												[]() { return LOCTEXT("AudioDashboard_SignalFlow_PauseTimestamp", "Pause on Select"); },
+												ToolbarButtonHeight)
+			]
 		];
 	}
 	
@@ -474,6 +523,7 @@ namespace UE::Audio::Insights
 			MenuBuilder.AddMenuEntry(Commands.GetToggleAnimateWiresCommand());
 
 			MenuBuilder.AddMenuEntry(Commands.GetToggleHorizontalViewCommand());
+			MenuBuilder.AddMenuEntry(Commands.GetPauseOnSelectCommand());
 
 			MenuBuilder.AddMenuEntry(Commands.GetShowNodeDetailsCommand());
 			MenuBuilder.AddSubMenu
@@ -650,6 +700,15 @@ namespace UE::Audio::Insights
 		bShowNodeDetails = !bShowNodeDetails;
 		OnRequestGraphRefresh();
 
+#if WITH_EDITOR
+		FSignalFlowSettings::OnRequestWriteSettings.Broadcast();
+#endif // WITH_EDITOR	
+	}
+
+	void FSignalFlowDashboardViewFactory::TogglePauseGraphEnabled()
+	{
+		bPauseGraphOnSelect = !bPauseGraphOnSelect;
+
 #if WITH_EDITOR
 		FSignalFlowSettings::OnRequestWriteSettings.Broadcast();
 #endif // WITH_EDITOR	
@@ -673,6 +732,44 @@ namespace UE::Audio::Insights
 #endif // WITH_EDITOR	
 	}
 
+	bool FSignalFlowDashboardViewFactory::CanPauseTimestamp() const
+	{
+		if (!bPauseGraphOnSelect)
+		{
+			return false;
+		}
+
+		FAudioInsightsModule& AudioInsightsModule = FAudioInsightsModule::GetChecked();
+
+#if !WITH_EDITOR
+		const TSharedPtr<const FAudioInsightsComponent> AudioInsightsComponent = AudioInsightsModule.GetAudioInsightsComponent();
+		if (AudioInsightsComponent.IsValid() && !AudioInsightsComponent->GetIsLiveSession())
+		{
+			// Don't update the caching and processing status if the cache isn't limited (i.e. loading in a .urace file)
+			return false;
+		}
+#endif // !WITH_EDITOR
+
+		return AudioInsightsModule.GetTimingViewExtender().GetMessageCacheAndProcessingStatus() == ECacheAndProcess::Latest;
+	}
+
+	bool FSignalFlowDashboardViewFactory::IsTimestampPaused() const
+	{
+		FAudioInsightsTimingViewExtender& TimingViewExtender = FAudioInsightsModule::GetChecked().GetTimingViewExtender();
+		return TimingViewExtender.GetMessageCacheAndProcessingStatus() != ECacheAndProcess::Latest;
+	}
+
+	void FSignalFlowDashboardViewFactory::ResetPauseTimestamp()
+	{
+		FAudioInsightsTimingViewExtender& TimingViewExtender = FAudioInsightsModule::GetChecked().GetTimingViewExtender();
+		if (TimingViewExtender.GetMessageCacheAndProcessingStatus() == ECacheAndProcess::Latest)
+		{
+			return;
+		}
+
+		TimingViewExtender.ResetMessageProcessType();
+	}
+
 	TSharedPtr<FSignalFlowDashboardEntry> FSignalFlowDashboardViewFactory::GetFilteredAudioDeviceEntry() const
 	{
 		const TSharedPtr<const FSignalFlowTraceProvider> Provider = FindProvider<const FSignalFlowTraceProvider>();
@@ -741,17 +838,14 @@ namespace UE::Audio::Insights
 
 	TSharedRef<SWidget> FSignalFlowDashboardViewFactory::CreateGraphSection()
 	{
-		GraphVerticalScrollBox = SAssignNew(GraphVerticalScrollBox, SScrollBox)
-		.Orientation(EOrientation::Orient_Vertical)
-		.ScrollBarAlwaysVisible(true)
+		using namespace FSignalFlowDashboardViewFactoryPrivate;
+
+		TSharedRef<SOverlay> GraphWidget = SNew(SOverlay)
 
-		+ SScrollBox::Slot()
-		.FillSize(1.0f)
-		.HAlign(HAlign_Fill)
-		.VAlign(VAlign_Fill)
+		+ SOverlay::Slot()
 		[
-			SAssignNew(GraphHorizontalScrollBox, SScrollBox)
-			.Orientation(EOrientation::Orient_Horizontal)
+			SAssignNew(GraphVerticalScrollBox, SScrollBox)
+			.Orientation(EOrientation::Orient_Vertical)
 			.ScrollBarAlwaysVisible(true)
 
 			+ SScrollBox::Slot()
@@ -759,93 +853,122 @@ namespace UE::Audio::Insights
 			.HAlign(HAlign_Fill)
 			.VAlign(VAlign_Fill)
 			[
-				SNew(SOverlay)
-
-				// Graph background
-				+ SOverlay::Slot()
+				SAssignNew(GraphHorizontalScrollBox, SScrollBox)
+				.Orientation(EOrientation::Orient_Horizontal)
+				.ScrollBarAlwaysVisible(true)
+
+				+ SScrollBox::Slot()
+				.FillSize(1.0f)
+				.HAlign(HAlign_Fill)
+				.VAlign(VAlign_Fill)
 				[
-					SNew(SColorBlock)
-					.Color(FStyleColors::Background.GetSpecifiedColor())
-				]
-
-				// Graph
-				+ SOverlay::Slot()
-				.HAlign(HAlign_Center)
-				.VAlign(VAlign_Center)
-				[
-					SAssignNew(GraphWidgetContents, SSignalFlowGraph)
-					.ListSource(&FilteredGraphNodes)
-					.NodeDetailFilterSettings(&NodeDetailFilters)
-					.SelectedItem_Lambda([this]()
-					{
-						return SelectedItem;
-					})
-					.ActiveAudioDeviceEntry_Lambda([this]()
-					{
-						return GetFilteredAudioDeviceEntry();
-					})
-					.OnNodeSelected_Lambda([this](TSharedPtr<FSignalFlowDashboardEntry> SelectedNodeEntry)
-					{
-						if (SelectedNodeEntry.IsValid())
+					SNew(SOverlay)
+
+					// Graph background
+					+ SOverlay::Slot()
+					[
+						SNew(SColorBlock)
+						.Color(FStyleColors::Background.GetSpecifiedColor())
+					]
+
+					// Graph
+					+ SOverlay::Slot()
+					.HAlign(HAlign_Center)
+					.VAlign(VAlign_Center)
+					[
+						SAssignNew(GraphWidgetContents, SSignalFlowGraph)
+						.ListSource(&FilteredGraphNodes)
+						.NodeDetailFilterSettings(&NodeDetailFilters)
+						.SelectedItem_Lambda([this]()
+						{
+							return SelectedItem;
+						})
+						.ActiveAudioDeviceEntry_Lambda([this]()
+						{
+							return GetFilteredAudioDeviceEntry();
+						})
+						.OnNodeSelected_Lambda([this](TSharedPtr<FSignalFlowDashboardEntry> SelectedNodeEntry)
 						{
-							switch (SelectedNodeEntry->EntryType)
+							if (SelectedNodeEntry.IsValid())
 							{
-								case ESignalFlowEntryType::SoundSource:
-									if (ActiveSourcesListView.IsValid())
-									{
-										ActiveSourcesListView->SetSelection(SelectedNodeEntry);
-									}
-									break;
-
-								case ESignalFlowEntryType::AudioBus:
-									if (ActiveBusesListView.IsValid())
-									{
-										ActiveBusesListView->SetSelection(SelectedNodeEntry);
-									}
-									break;
-
-								case ESignalFlowEntryType::Submix:
-									if (ActiveSubmixesListView.IsValid())
-									{
-										ActiveSubmixesListView->SetSelection(SelectedNodeEntry);
-									}
-									break;
-
-								default:
-									if (ActiveSourcesListView.IsValid())
-									{
-										ActiveSourcesListView->ClearSelection();
-									}
-									
-									if (ActiveBusesListView.IsValid())
-									{
-										ActiveBusesListView->ClearSelection();
-									}
-									
-									if (ActiveSubmixesListView.IsValid())
-									{
-										ActiveSubmixesListView->ClearSelection();
-									}
-									break;
+								switch (SelectedNodeEntry->EntryType)
+								{
+									case ESignalFlowEntryType::SoundSource:
+										if (ActiveSourcesListView.IsValid())
+										{
+											ActiveSourcesListView->SetSelection(SelectedNodeEntry);
+										}
+										break;
+
+									case ESignalFlowEntryType::AudioBus:
+										if (ActiveBusesListView.IsValid())
+										{
+											ActiveBusesListView->SetSelection(SelectedNodeEntry);
+										}
+										break;
+
+									case ESignalFlowEntryType::Submix:
+										if (ActiveSubmixesListView.IsValid())
+										{
+											ActiveSubmixesListView->SetSelection(SelectedNodeEntry);
+										}
+										break;
+
+									default:
+										if (ActiveSourcesListView.IsValid())
+										{
+											ActiveSourcesListView->ClearSelection();
+										}
+										
+										if (ActiveBusesListView.IsValid())
+										{
+											ActiveBusesListView->ClearSelection();
+										}
+										
+										if (ActiveSubmixesListView.IsValid())
+										{
+											ActiveSubmixesListView->ClearSelection();
+										}
+										break;
+								}
 							}
-						}
 
-						OnSelectionChanged(SelectedNodeEntry, ESelectInfo::OnMouseClick);
-					})
-					.ShowNodeDetails_Lambda([this]()
-					{
-						return bShowNodeDetails;
-					})
-					.DisplayAmpPeakInDb_Lambda([this]()
-					{
-						return bDisplayAmpPeakInDb;
-					})
-					.AnimateWires_Lambda([this]()
-					{
-						return bAnimateWires;
-					})
+							OnSelectionChanged(SelectedNodeEntry, ESelectInfo::OnMouseClick);
+						})
+						.ShowNodeDetails_Lambda([this]()
+						{
+							return bShowNodeDetails;
+						})
+						.DisplayAmpPeakInDb_Lambda([this]()
+						{
+							return bDisplayAmpPeakInDb;
+						})
+						.AnimateWires_Lambda([this]()
+						{
+							return bAnimateWires;
+						})
+					]
 				]
 			]
+		]
+		
+		+ SOverlay::Slot()
+		.HAlign(HAlign_Center)
+		.VAlign(VAlign_Bottom)
+		.Padding(24.0f)
+		[
+			SNew(SBorder)
+			.BorderImage(&GetRoundedBoxBrush())
+			.BorderBackgroundColor(FStyleColors::AccentRed.GetSpecifiedColor())
+			.Visibility_Lambda([this]() { return IsTimestampPaused() ? EVisibility::Visible : EVisibility::Collapsed; })
+			.Padding(8.0f)
+			.Content()
+			[
+				SNew(STextBlock)
+				.Text(LOCTEXT("AudioDashboard_SignalFlow_GraphPausedWarning", "Processing is paused! Press Esc to resume."))
+				.Font(FCoreStyle::Get().GetFontStyle("BoldFont"))
+				.ColorAndOpacity(FStyleColors::White.GetSpecifiedColor())
+			]
 		];
 
 		if (GraphWidgetContents.IsValid())
@@ -854,7 +977,7 @@ namespace UE::Audio::Insights
 			GraphWidgetContents->SetVerticalScrollBox(GraphVerticalScrollBox);
 		}
 
-		return GraphVerticalScrollBox.ToSharedRef();
+		return GraphWidget;
 	}
 
 	TSharedRef<SListView<TSharedPtr<FSignalFlowDashboardEntry>>> FSignalFlowDashboardViewFactory::CreateSelectableList(const FText& HeadingText, const TArray<TSharedPtr<FSignalFlowDashboardEntry>>& ListSource)
@@ -949,12 +1072,19 @@ namespace UE::Audio::Insights
 			}
 
 			SelectedItem = InSelectedItem;
-#if WITH_EDITOR
 			if (SelectedItem.IsValid())
 			{
+#if WITH_EDITOR
 				AssetContextMenuHelper.SetAssetEntry(SelectedItem);
-			}
 #endif // WITH_EDITOR
+
+				if (CanPauseTimestamp())
+				{
+					FAudioInsightsModule& AudioInsightsModule = FAudioInsightsModule::GetChecked();
+					const double PauseTime = AudioInsightsModule.GetCacheManager().GetCacheEndTimeStamp();
+					AudioInsightsModule.GetTimingViewExtender().SetTimeMarker(PauseTime, ESystemControllingTimeMarker::SignalFlow);
+				}
+			}
 		}
 		else
 		{
@@ -1865,6 +1995,7 @@ namespace UE::Audio::Insights
 			GraphWidgetContents->SetWireAnimationSettings(InSettings.WireSplineAmplitudePowerFactor, InSettings.WireSplineMaxThicknessScalar);
 		}
 
+		bPauseGraphOnSelect = InSettings.bPauseGraphOnSelect;
 		bShowNodeDetails = InSettings.bShowNodeDetails;
 		NodeDetailFilters = InSettings.NodeDetailFilters;
 		bDisplayAmpPeakInDb = InSettings.AmplitudeDisplayMode == EAudioAmplitudeDisplayMode::Decibels;
@@ -1886,6 +2017,7 @@ namespace UE::Audio::Insights
 		}
 
 		OutSettings.bShowNodeDetails = bShowNodeDetails;
+		OutSettings.bPauseGraphOnSelect = bPauseGraphOnSelect;
 		OutSettings.NodeDetailFilters = NodeDetailFilters;
 		OutSettings.AmplitudeDisplayMode = bDisplayAmpPeakInDb ? EAudioAmplitudeDisplayMode::Decibels : EAudioAmplitudeDisplayMode::Linear;
 		OutSettings.bAnimateWires = bAnimateWires;
diff --git a/Engine/Plugins/AudioInsights/Source/AudioInsights/Public/Settings/SignalFlowSettings.h b/Engine/Plugins/AudioInsights/Source/AudioInsights/Public/Settings/SignalFlowSettings.h
--- a/Engine/Plugins/AudioInsights/Source/AudioInsights/Public/Settings/SignalFlowSettings.h
+++ b/Engine/Plugins/AudioInsights/Source/AudioInsights/Public/Settings/SignalFlowSettings.h
@@ -74,6 +74,10 @@ struct FSignalFlowSettings
 	UPROPERTY(EditAnywhere, config, Category = SignalFlow)
 	bool bHorizontalFlow = true;
 
+	/** Toggles pausing the graph at the current timestamp when selecting a node. */
+	UPROPERTY(EditAnywhere, config, Category = SignalFlow)
+	bool bPauseGraphOnSelect = false;
+
 	/** Toggles visibility of node details. */
 	UPROPERTY(EditAnywhere, config, Category = SignalFlow)
 	bool bShowNodeDetails = false;
diff --git a/Engine/Plugins/AudioInsights/Source/AudioInsights/Public/Views/SSignalFlowGraph.h b/Engine/Plugins/AudioInsights/Source/AudioInsights/Public/Views/SSignalFlowGraph.h
--- a/Engine/Plugins/AudioInsights/Source/AudioInsights/Public/Views/SSignalFlowGraph.h
+++ b/Engine/Plugins/AudioInsights/Source/AudioInsights/Public/Views/SSignalFlowGraph.h
@@ -219,6 +219,7 @@ namespace UE::Audio::Insights
 		void SetVerticalScrollBox(const TWeakPtr<SScrollBox> ScrollBox) { VerticalScrollBox = ScrollBox; }
 
 		UE_API void RefreshGraph();
+		UE_API void ResetGraph();
 
 		UE_API void ToggleOrientation();
 		EOrientation GetGraphOrientation() const { return Orientation; }
diff --git a/Engine/Plugins/AudioInsights/Source/AudioInsights/Public/Views/SignalFlowDashboardViewFactory.h b/Engine/Plugins/AudioInsights/Source/AudioInsights/Public/Views/SignalFlowDashboardViewFactory.h
--- a/Engine/Plugins/AudioInsights/Source/AudioInsights/Public/Views/SignalFlowDashboardViewFactory.h
+++ b/Engine/Plugins/AudioInsights/Source/AudioInsights/Public/Views/SignalFlowDashboardViewFactory.h
@@ -83,9 +83,14 @@ namespace UE::Audio::Insights
 		FReply ScrollGraphViewToSelectedNode(const bool bResetZoom);
 		void ToggleGraphOrientation();
 		void ToggleShowNodeDetails();
+		void TogglePauseGraphEnabled();
 		void ToggleEnableAllNodeDetailFilters();
 		void ToggleAnimateWires();
 
+		bool CanPauseTimestamp() const;
+		bool IsTimestampPaused() const;
+		void ResetPauseTimestamp();
+
 		TSharedPtr<FSignalFlowDashboardEntry> GetFilteredAudioDeviceEntry() const;
 
 		TSharedRef<SWidget> CreateSelectionPanel();
@@ -164,6 +169,7 @@ namespace UE::Audio::Insights
 		bool bGraphRequiresRefresh = false;
 		bool bAutoFocusAfterGraphRefresh = false;
 		bool bShowNodeDetails = false;
+		bool bPauseGraphOnSelect = false;
 		bool bDisplayAmpPeakInDb = true;
 		bool bAnimateWires = true;
 
