[Audio Insights] Maintain the user's viewing position inside the Signal flow graph when new nodes are created

#Synced-CL 50815132

diff --git a/Engine/Plugins/AudioInsights/Source/AudioInsights/Private/Views/SSignalFlowGraph.cpp b/Engine/Plugins/AudioInsights/Source/AudioInsights/Private/Views/SSignalFlowGraph.cpp
--- a/Engine/Plugins/AudioInsights/Source/AudioInsights/Private/Views/SSignalFlowGraph.cpp
+++ b/Engine/Plugins/AudioInsights/Source/AudioInsights/Private/Views/SSignalFlowGraph.cpp
@@ -858,9 +858,16 @@ namespace UE::Audio::Insights
 
 		DummyToRealConnectionKeys.Reset();
 
-		CreateGraphNodes(RemainingNodeKeys);
+		const FVector2D OriginalSize = GetDesiredSize();
+		const FGraphLimits NewGraphWidthLimits = GetGraphWidthLimits();
+		TArray<TreeDepthPair> NewTreeDepthOrder;
+
+		CreateGraphNodes(RemainingNodeKeys, NewGraphWidthLimits, NewTreeDepthOrder);
 		CreateSendLabels();
 		CleanUpOldWidgets(RemainingNodeKeys);
+		MaintainOldViewportPosition(OriginalSize, NewGraphWidthLimits, NewTreeDepthOrder);
+
+		CachedGraphWidthLimits = NewGraphWidthLimits;
 	}
 
 	void SSignalFlowGraph::ResetGraph()
@@ -1331,36 +1338,54 @@ namespace UE::Audio::Insights
 		return FVector2D{ GraphWidth, GraphHeight } * InLayoutScaleMultiplier * ZoomScaleFactor;
 	}
 
-	void SSignalFlowGraph::CreateGraphNodes(TArray<FSignalFlowEntryKey>& OutRemainingNodeKeys)
+	SSignalFlowGraph::FGraphLimits SSignalFlowGraph::GetGraphWidthLimits() const
 	{
-		TreeDepthPair CurrentTreeDepth{ ESignalFlowEntryType::InvalidLow, TNumericLimits<int32>::Lowest() };
+		FGraphLimits GraphWidthLimits;
 
-		float MinXPos = TNumericLimits<float>::Max();
-		float MaxXPos = TNumericLimits<float>::Lowest();
-
-		// Gather stats on the width of the graph
-		for (const TSharedPtr<ISignalFlowNode>& Node : *ListSource)
+		if (ListSource)
 		{
-			const float XPos = Node->XPos;
+			// Gather stats on the width of the graph
+			for (const TSharedPtr<ISignalFlowNode>& Node : *ListSource)
+			{
+				if (!Node.IsValid())
+				{
+					continue;
+				}
+
+				const float XPos = Node->XPos;
 
-			MinXPos = FMath::Min(MinXPos, XPos);
-			MaxXPos = FMath::Max(MaxXPos, XPos);
+				GraphWidthLimits.MinPos = FMath::Min(GraphWidthLimits.MinPos, XPos);
+				GraphWidthLimits.MaxPos = FMath::Max(GraphWidthLimits.MaxPos, XPos);
+			}
 		}
 
-		// DepthTracker is used to assign a YPosition to each node
-		// Each time we cross a layer, add one to the depth tracker
-		int32 DepthTracker = 0;
+		return GraphWidthLimits;
+	}
+
+	void SSignalFlowGraph::CreateGraphNodes(TArray<FSignalFlowEntryKey>& OutRemainingNodeKeys, const FGraphLimits& GraphWidthLimits, TArray<TreeDepthPair>& OutTreeDepthStructure)
+	{
+		TreeDepthPair CurrentTreeDepth{ ESignalFlowEntryType::InvalidLow, TNumericLimits<int32>::Lowest() };
+		OutTreeDepthStructure.Reset();
 
 		// Create the nodes first
 		for (const TSharedPtr<ISignalFlowNode>& Node : *ListSource)
 		{
+			if (!Node.IsValid())
+			{
+				continue;
+			}
+
 			if (CurrentTreeDepth != Node->GetTreeDepth())
 			{
 				CurrentTreeDepth = Node->GetTreeDepth();
-				DepthTracker++;
+
+				// Each tree depth layer is used to assign a YPosition to each node
+				// Each time we cross a layer, add one to the depth tracker
+				// OutTreeDepthStructure tracks the unique tree depth layers inside the graph
+				OutTreeDepthStructure.Add(CurrentTreeDepth);
 			}
 
-			CreateGraphNode(Node, MinXPos, MaxXPos, DepthTracker, OutRemainingNodeKeys);
+			CreateGraphNode(Node, GraphWidthLimits.MinPos, GraphWidthLimits.MaxPos, OutTreeDepthStructure.Num(), OutRemainingNodeKeys);
 		}
 	}
 
@@ -1433,6 +1458,104 @@ namespace UE::Audio::Insights
 		}
 	}
 
+	void SSignalFlowGraph::MaintainOldViewportPosition(const FVector2D& OriginalSize, const SSignalFlowGraph::FGraphLimits& NewGraphWidthLimits, const TArray<TreeDepthPair>& NewTreeDepthOrder)
+	{
+		using namespace SSignalFlowGraphStyle;
+
+		// When the graph changes, we want to try and maintain the current viewport position as much as possible.
+		// XPos is assigned when generating the order of nodes, and can be positive or negative
+		// YPos is assigned via treedepth, and can only be positive
+
+		// We maintain the position with the following:
+		// New Position = Old Postion + Total Size Increase - Size Increase Beyond Old Postion
+		// (i.e. Old Postion + Size Increase Before Old Postion)
+
+		const TSharedPtr<SScrollBox> HorizontalScrollBoxPinned = HorizontalScrollBox.Get().Pin();
+		const TSharedPtr<SScrollBox> VerticalScrollBoxPinned = VerticalScrollBox.Get().Pin();
+
+		if (HorizontalScrollBoxPinned.IsValid() && VerticalScrollBoxPinned.IsValid() && CachedGraphWidthLimits.IsValid())
+		{
+			constexpr float LayoutScaleMultiplier = 1.0f;
+
+			// Get the total increase/decrease in size of the graph
+			const FVector2D NewDesiredSize = ComputeDesiredSize(LayoutScaleMultiplier);
+			const FVector2D TotalSizeDiff = NewDesiredSize - OriginalSize;
+			
+			// Get the old position (top left corner) of the viewport
+			const FVector2D OldViewportPosInGraph(HorizontalScrollBoxPinned->GetScrollOffset(), VerticalScrollBoxPinned->GetScrollOffset());
+			
+			const float IncreaseAtRight = CalculateGraphSizeIncreaseAtRight(NewGraphWidthLimits);
+
+			const float IncreaseAtBottom = CalculateGraphSizeIncreaseAtBottom
+			(
+				NewTreeDepthOrder, 
+				Orientation == EOrientation::Orient_Vertical ? OriginalSize.Y : OriginalSize.X,
+				Orientation == EOrientation::Orient_Vertical ? OldViewportPosInGraph.Y : OldViewportPosInGraph.X
+			);
+
+			const float XOffset = Orientation == EOrientation::Orient_Vertical ? TotalSizeDiff.X - IncreaseAtRight : TotalSizeDiff.X - IncreaseAtBottom;
+			const float YOffset = Orientation == EOrientation::Orient_Vertical ? TotalSizeDiff.Y - IncreaseAtBottom : TotalSizeDiff.Y - IncreaseAtRight;
+
+			const FVector2D NewDesiredPos(OldViewportPosInGraph.X + XOffset, OldViewportPosInGraph.Y + YOffset);
+
+			const FVector2D CurrentScrollOffset(HorizontalScrollBoxPinned->GetScrollOffsetOfEnd(), VerticalScrollBoxPinned->GetScrollOffsetOfEnd());
+			const FVector2D MaxScrollOffset = NewDesiredSize - (OriginalSize - CurrentScrollOffset);
+
+			HorizontalScrollBoxPinned->SetScrollOffset(FMath::Clamp(NewDesiredPos.X, 0.0f, FMath::Max(MaxScrollOffset.X, 0.0f)));
+			VerticalScrollBoxPinned->SetScrollOffset(FMath::Clamp(NewDesiredPos.Y, 0.0f, FMath::Max(MaxScrollOffset.Y, 0.0f)));
+		}
+
+		CachedTreeDepthPairs = NewTreeDepthOrder;
+	}
+
+	float SSignalFlowGraph::CalculateGraphSizeIncreaseAtRight(const FGraphLimits& NewGraphWidthLimits) const
+	{
+		using namespace SSignalFlowGraphStyle;
+
+		const float NodeGraphSlotWidth = GetNodeGraphSlotWidth(GetGraphOrientation());
+
+		// Work out in pixels how much the graph has increased in size beyond the original viewport position
+		const float OldMaxXPosScaled = CachedGraphWidthLimits.IsValid() ? NodeGraphSlotWidth * CachedGraphWidthLimits.MaxPos : 0.0f;
+		const float NewMaxXPosScaled = NewGraphWidthLimits.IsValid() ? NodeGraphSlotWidth * NewGraphWidthLimits.MaxPos : 0.0f;
+
+		return NewMaxXPosScaled - OldMaxXPosScaled;
+	}
+
+	float SSignalFlowGraph::CalculateGraphSizeIncreaseAtBottom(const TArray<TreeDepthPair>& NewTreeDepthOrder, const float OriginalGraphHeight, const float OriginalViewportPosY) const
+	{
+		using namespace SSignalFlowGraphStyle;
+		if (CachedTreeDepthPairs.IsEmpty())
+		{
+			return 0.0f;
+		}
+
+		// The original YPos of the viewport = viewport pos in original as percentage X the original number of tree depths
+		int32 OriginalYPos = FMath::IsNearlyZero(OriginalGraphHeight) ? 0 : static_cast<int32>(FMath::Floor((OriginalViewportPosY / OriginalGraphHeight) * CachedTreeDepthPairs.Num()));
+		OriginalYPos = FMath::Clamp(OriginalYPos, 0, CachedTreeDepthPairs.Num() - 1);
+
+		const TreeDepthPair OldViewportCenterTreeDepth = CachedTreeDepthPairs[OriginalYPos];
+
+		// Locate where the original tree depth is in the new structure.
+		int32 NewYPos = 0;
+		for (const TreeDepthPair& NewTreeDepth : NewTreeDepthOrder)
+		{
+			if (OldViewportCenterTreeDepth >= NewTreeDepth)
+			{
+				break;
+			}
+
+			NewYPos++;
+		}
+
+		// Calculate how much the graph has increased below the original viewport (positive increase)
+		const int32 YPosDiffAtBottom = FMath::Max(NewYPos - OriginalYPos, 0);
+
+		const int32 NumParamsVisible = NodeDetailFilterSettings ? NodeDetailFilterSettings->GetNumVisibleParams() : 0;
+		const float NodeGraphSlotHeight = GetNodeGraphSlotHeight(GetGraphOrientation(), NumParamsVisible, ShowNodeDetails.Get());
+
+		return YPosDiffAtBottom * NodeGraphSlotHeight;
+	}
+
 	void SSignalFlowGraph::CreateGraphNode(const TSharedPtr<ISignalFlowNode>& Node, const float MinXPos, const float MaxXPos, const int32 DepthTracker, TArray<FSignalFlowEntryKey>& OutRemainingNodeKeys)
 	{
 		using namespace SSignalFlowGraphStyle;
@@ -1713,7 +1836,7 @@ namespace UE::Audio::Insights
 		else
 		{
 			// If this is a dummy node, try and find the real input/output node we are connected to instead
-			const ConnectedNodeKeys* FoundRealConnectionKeys = DummyToRealConnectionKeys.Find(EntryKey);
+			const FConnectedNodeKeys* FoundRealConnectionKeys = DummyToRealConnectionKeys.Find(EntryKey);
 			if (FoundRealConnectionKeys)
 			{
 				const FSignalFlowEntryKey& LookupKey = ConnectionType == ENodeConnectionType::Input ? FoundRealConnectionKeys->RealInputKey 
diff --git a/Engine/Plugins/AudioInsights/Source/AudioInsights/Public/Views/SSignalFlowGraph.h b/Engine/Plugins/AudioInsights/Source/AudioInsights/Public/Views/SSignalFlowGraph.h
--- a/Engine/Plugins/AudioInsights/Source/AudioInsights/Public/Views/SSignalFlowGraph.h
+++ b/Engine/Plugins/AudioInsights/Source/AudioInsights/Public/Views/SSignalFlowGraph.h
@@ -263,9 +263,25 @@ namespace UE::Audio::Insights
 												TArray<FVector2D>& OutBottomPoints,
 												FVector2D& OutLastPoint) const;
 
-		void CreateGraphNodes(TArray<FSignalFlowEntryKey>& OutRemainingNodeKeys);
+		struct FGraphLimits
+		{
+			float MinPos = TNumericLimits<float>::Max();
+			float MaxPos = TNumericLimits<float>::Lowest();
+
+			bool IsValid() const
+			{
+				return MinPos != TNumericLimits<float>::Max() && MaxPos != TNumericLimits<float>::Lowest();
+			}
+		};
+
+		FGraphLimits GetGraphWidthLimits() const;
+		void CreateGraphNodes(TArray<FSignalFlowEntryKey>& OutRemainingNodeKeys, const FGraphLimits& GraphWidthLimits, TArray<struct TreeDepthPair>& OutTreeDepthStructure);
 		void CreateSendLabels();
 		void CleanUpOldWidgets(const TArray<FSignalFlowEntryKey>& NodeKeysToRemove);
+		
+		void MaintainOldViewportPosition(const FVector2D& OriginalSize, const FGraphLimits& NewGraphWidthLimits, const TArray<struct TreeDepthPair>& NewTreeDepthOrder);
+		float CalculateGraphSizeIncreaseAtRight(const FGraphLimits& NewGraphWidthLimits) const;
+		float CalculateGraphSizeIncreaseAtBottom(const TArray<TreeDepthPair>& NewTreeDepthOrder, const float OriginalGraphHeight, const float OriginalViewportPosY) const;
 
 		void CreateGraphNode(const TSharedPtr<ISignalFlowNode>& Node, const float MinXPos, const float MaxXPos, const int32 DepthTracker, TArray<FSignalFlowEntryKey>& OutRemainingNodeKeys);
 		void CreateSendLabelWidget(const FSignalFlowEntryKey& SenderEntryKey, const FSignalFlowEntryKey& ReceiverEntryKey);
@@ -298,13 +314,13 @@ namespace UE::Audio::Insights
 
 		TMap<FSignalFlowEntryKey, TSharedPtr<SSignalFlowGraphNode>> NodeWidgets;
 
-		struct ConnectedNodeKeys
+		struct FConnectedNodeKeys
 		{
 			FSignalFlowEntryKey RealInputKey;
 			FSignalFlowEntryKey RealOutputKey;
 		};
 
-		TMap<FSignalFlowEntryKey, ConnectedNodeKeys> DummyToRealConnectionKeys;
+		TMap<FSignalFlowEntryKey, FConnectedNodeKeys> DummyToRealConnectionKeys;
 		TMultiMap<FSignalFlowEntryKey, TSharedPtr<SSignalFlowGraphSendLabel>> SendLevelWidgets;
 
 		TArray<TSharedPtr<ISignalFlowNode>>* ListSource = nullptr;
@@ -328,6 +344,8 @@ namespace UE::Audio::Insights
 
 		EOrientation Orientation = EOrientation::Orient_Vertical;
 
+		FGraphLimits CachedGraphWidthLimits;
+		TArray<struct TreeDepthPair> CachedTreeDepthPairs;
 		float ZoomScaleFactor = 1.0f;
 		float WireSplineAmplitudePowerFactor = 0.4f;
 		float WireSplineThicknessScalarMax = 14.0f;
diff --git a/Engine/Plugins/AudioInsights/Source/AudioInsights/Public/Views/SignalFlowNodes.h b/Engine/Plugins/AudioInsights/Source/AudioInsights/Public/Views/SignalFlowNodes.h
--- a/Engine/Plugins/AudioInsights/Source/AudioInsights/Public/Views/SignalFlowNodes.h
+++ b/Engine/Plugins/AudioInsights/Source/AudioInsights/Public/Views/SignalFlowNodes.h
@@ -35,6 +35,26 @@ namespace UE::Audio::Insights
 			return TreeDepth < Other.TreeDepth;
 		}
 
+		bool operator>(const TreeDepthPair& Other) const
+		{
+			if (LayerID != Other.LayerID)
+			{
+				return static_cast<int32>(LayerID) > static_cast<int32>(Other.LayerID);
+			}
+
+			return TreeDepth > Other.TreeDepth;
+		}
+
+		bool operator<=(const TreeDepthPair& Other) const
+		{
+			return operator<(Other) || operator==(Other);
+		}
+
+		bool operator>=(const TreeDepthPair& Other) const
+		{
+			return operator>(Other) || operator==(Other);
+		}
+
 		friend int32 GetTypeHash(const TreeDepthPair& Key)
 		{
 			return HashCombine(GetTypeHash(static_cast<int32>(Key.LayerID)), GetTypeHash(Key.TreeDepth));
