[Audio Insights] Clicking empty space in the signal flow graph should de-select the current node

#Synced-CL 50815132

diff --git a/Engine/Plugins/AudioInsights/Source/AudioInsights/Private/Views/SSignalFlowGraph.cpp b/Engine/Plugins/AudioInsights/Source/AudioInsights/Private/Views/SSignalFlowGraph.cpp
--- a/Engine/Plugins/AudioInsights/Source/AudioInsights/Private/Views/SSignalFlowGraph.cpp
+++ b/Engine/Plugins/AudioInsights/Source/AudioInsights/Private/Views/SSignalFlowGraph.cpp
@@ -1012,6 +1012,12 @@ namespace UE::Audio::Insights
 		{
 			return FReply::Handled().DetectDrag(SharedThis(this), EKeys::RightMouseButton);
 		}
+		else if (MouseEvent.GetEffectingButton() == EKeys::LeftMouseButton)
+		{
+			OnNodeSelected.ExecuteIfBound(nullptr);
+
+			return FReply::Handled();
+		}
 
 		return FReply::Unhandled();
 	}
