[Audio Insights] Justification options for signal flow graph
* Graph can be justified towards the top/center/bottom in horizontal mode, and left/center/right in vertical mode
* Hotkeys, settings menu options and toolbar buttons created
* Updated the positioning algorithm to include the new justification options
* Use the previous node position in the row to guard against placing nodes in the same X-Position instead of the old MakeSpaceForNode. This is not only cheaper but makes sure we maintain the order and is more reliable.
* Use average position of inputs to position centered nodes (instead of the median position)
* Allow fractional X positions - remove flooring of x positions to keep everything in a grid. Helps reduce empty space.

#Synced-CL 50815132

diff --git a/Engine/Plugins/AudioInsights/Source/AudioInsights/Private/SignalFlowEditorCommands.cpp b/Engine/Plugins/AudioInsights/Source/AudioInsights/Private/SignalFlowEditorCommands.cpp
--- a/Engine/Plugins/AudioInsights/Source/AudioInsights/Private/SignalFlowEditorCommands.cpp
+++ b/Engine/Plugins/AudioInsights/Source/AudioInsights/Private/SignalFlowEditorCommands.cpp
@@ -17,6 +17,12 @@ namespace UE::Audio::Insights
 	void FSignalFlowEditorCommands::RegisterCommands()
 	{
 		UI_COMMAND(ToggleHorizontalView, "Horizontal Flow", "Toggles between horizontally/vertically orienting the Signal Flow graph.", EUserInterfaceActionType::ToggleButton, FInputChord(EKeys::H, EModifierKey::Control));
+		UI_COMMAND(JustifyBottom, "Bottom", "Justify the Signal Flow Graph towards the bottom in horizontal mode.", EUserInterfaceActionType::RadioButton, FInputChord(EKeys::F1, EModifierKey::Control));
+		UI_COMMAND(JustifyHorizontalCenter, "Center", "Justify the Signal Flow Graph around the center in horizontal mode.", EUserInterfaceActionType::RadioButton, FInputChord(EKeys::F2, EModifierKey::Control));
+		UI_COMMAND(JustifyTop, "Top", "Justify the Signal Flow Graph towards the top in horizontal mode.", EUserInterfaceActionType::RadioButton, FInputChord(EKeys::F3, EModifierKey::Control));
+		UI_COMMAND(JustifyLeft, "Left", "Justify the Signal Flow Graph towards the left in vertical mode.", EUserInterfaceActionType::RadioButton, FInputChord(EKeys::F4, EModifierKey::Control));
+		UI_COMMAND(JustifyVerticalCenter, "Center", "Justify the Signal Flow Graph around the center in vertical mode.", EUserInterfaceActionType::RadioButton, FInputChord(EKeys::F5, EModifierKey::Control));
+		UI_COMMAND(JustifyRight, "Right", "Justify the Signal Flow Graph towards the right in vertical mode.", EUserInterfaceActionType::RadioButton, FInputChord(EKeys::F6, EModifierKey::Control));
 		UI_COMMAND(ShowNodeDetails, "Show Node Details", "Toggles visibility of node details.", EUserInterfaceActionType::ToggleButton, FInputChord(EKeys::D, EModifierKey::Control));
 		UI_COMMAND(PauseOnSelect, "Pause on Select", "Toggles whether selecting a node should pause processing trace data, freezing audio insights at the latest timestamp.", EUserInterfaceActionType::ToggleButton, FInputChord(EKeys::P, EModifierKey::Control));
 		UI_COMMAND(ToggleAnimateWires, "Animate Wires", "Toggles whether connection wire animations driven by amplitude envelope values are enabled in the Signal Flow graph.", EUserInterfaceActionType::ToggleButton, FInputChord());
diff --git a/Engine/Plugins/AudioInsights/Source/AudioInsights/Private/SignalFlowEditorCommands.h b/Engine/Plugins/AudioInsights/Source/AudioInsights/Private/SignalFlowEditorCommands.h
--- a/Engine/Plugins/AudioInsights/Source/AudioInsights/Private/SignalFlowEditorCommands.h
+++ b/Engine/Plugins/AudioInsights/Source/AudioInsights/Private/SignalFlowEditorCommands.h
@@ -17,6 +17,12 @@ namespace UE::Audio::Insights
 
 		// Toggles
 		TSharedPtr<const FUICommandInfo> GetToggleHorizontalViewCommand() const { return ToggleHorizontalView; }
+		TSharedPtr<const FUICommandInfo> GetJustifyBottomCommand() const { return JustifyBottom; }
+		TSharedPtr<const FUICommandInfo> GetJustifyHorizontalCenterCommand() const { return JustifyHorizontalCenter; }
+		TSharedPtr<const FUICommandInfo> GetJustifyTopCommand() const { return JustifyTop; }
+		TSharedPtr<const FUICommandInfo> GetJustifyLeftCommand() const { return JustifyLeft; }
+		TSharedPtr<const FUICommandInfo> GetJustifyVerticalCenterCommand() const { return JustifyVerticalCenter; }
+		TSharedPtr<const FUICommandInfo> GetJustifyRightCommand() const { return JustifyRight; }
 		TSharedPtr<const FUICommandInfo> GetShowNodeDetailsCommand() const { return ShowNodeDetails; }
 		TSharedPtr<const FUICommandInfo> GetPauseOnSelectCommand() const { return PauseOnSelect; }
 		TSharedPtr<const FUICommandInfo> GetToggleAnimateWiresCommand() const { return ToggleAnimateWires; }
@@ -28,6 +34,12 @@ namespace UE::Audio::Insights
 	private:
 		// Toggles
 		TSharedPtr<FUICommandInfo> ToggleHorizontalView;
+		TSharedPtr<FUICommandInfo> JustifyBottom;
+		TSharedPtr<FUICommandInfo> JustifyHorizontalCenter;
+		TSharedPtr<FUICommandInfo> JustifyTop;
+		TSharedPtr<FUICommandInfo> JustifyLeft;
+		TSharedPtr<FUICommandInfo> JustifyVerticalCenter;
+		TSharedPtr<FUICommandInfo> JustifyRight;
 		TSharedPtr<FUICommandInfo> ShowNodeDetails;
 		TSharedPtr<FUICommandInfo> PauseOnSelect;
 		TSharedPtr<FUICommandInfo> ToggleAnimateWires;
diff --git a/Engine/Plugins/AudioInsights/Source/AudioInsights/Private/Views/SignalFlowDashboardViewFactory.cpp b/Engine/Plugins/AudioInsights/Source/AudioInsights/Private/Views/SignalFlowDashboardViewFactory.cpp
--- a/Engine/Plugins/AudioInsights/Source/AudioInsights/Private/Views/SignalFlowDashboardViewFactory.cpp
+++ b/Engine/Plugins/AudioInsights/Source/AudioInsights/Private/Views/SignalFlowDashboardViewFactory.cpp
@@ -288,6 +288,42 @@ namespace UE::Audio::Insights
 			FCanExecuteAction(),
 			FIsActionChecked::CreateLambda([this]() { return GraphWidgetContents.IsValid() && GraphWidgetContents->GetGraphOrientation() == Orient_Horizontal; }));
 
+		CommandList->MapAction(
+			Commands.GetJustifyBottomCommand(),
+			FExecuteAction::CreateLambda([this]() { SetHorizontalJustification(ESignalFlowHorizontalJustification::Bottom); }),
+			FCanExecuteAction::CreateLambda([this](){ return GraphWidgetContents.IsValid() && GraphWidgetContents->GetGraphOrientation() == EOrientation::Orient_Horizontal; }),
+			FIsActionChecked::CreateLambda([this]() { return HorizontalGraphJustification == ESignalFlowHorizontalJustification::Bottom; }));
+
+		CommandList->MapAction(
+			Commands.GetJustifyHorizontalCenterCommand(),
+			FExecuteAction::CreateLambda([this]() { SetHorizontalJustification(ESignalFlowHorizontalJustification::Center); }),
+			FCanExecuteAction::CreateLambda([this]() { return GraphWidgetContents.IsValid() && GraphWidgetContents->GetGraphOrientation() == EOrientation::Orient_Horizontal; }),
+			FIsActionChecked::CreateLambda([this]() { return HorizontalGraphJustification == ESignalFlowHorizontalJustification::Center; }));
+
+		CommandList->MapAction(
+			Commands.GetJustifyTopCommand(),
+			FExecuteAction::CreateLambda([this]() { SetHorizontalJustification(ESignalFlowHorizontalJustification::Top); }),
+			FCanExecuteAction::CreateLambda([this]() { return GraphWidgetContents.IsValid() && GraphWidgetContents->GetGraphOrientation() == EOrientation::Orient_Horizontal; }),
+			FIsActionChecked::CreateLambda([this]() { return HorizontalGraphJustification == ESignalFlowHorizontalJustification::Top; }));
+
+		CommandList->MapAction(
+			Commands.GetJustifyLeftCommand(),
+			FExecuteAction::CreateLambda([this]() { SetVerticalJustification(ESignalFlowVerticalJustification::Left); }),
+			FCanExecuteAction::CreateLambda([this]() { return GraphWidgetContents.IsValid() && GraphWidgetContents->GetGraphOrientation() == EOrientation::Orient_Vertical; }),
+			FIsActionChecked::CreateLambda([this]() { return VerticalGraphJustification == ESignalFlowVerticalJustification::Left; }));
+
+		CommandList->MapAction(
+			Commands.GetJustifyVerticalCenterCommand(),
+			FExecuteAction::CreateLambda([this]() { SetVerticalJustification(ESignalFlowVerticalJustification::Center); }),
+			FCanExecuteAction::CreateLambda([this]() { return GraphWidgetContents.IsValid() && GraphWidgetContents->GetGraphOrientation() == EOrientation::Orient_Vertical; }),
+			FIsActionChecked::CreateLambda([this]() { return VerticalGraphJustification == ESignalFlowVerticalJustification::Center; }));
+
+		CommandList->MapAction(
+			Commands.GetJustifyRightCommand(),
+			FExecuteAction::CreateLambda([this]() { SetVerticalJustification(ESignalFlowVerticalJustification::Right); }),
+			FCanExecuteAction::CreateLambda([this]() { return GraphWidgetContents.IsValid() && GraphWidgetContents->GetGraphOrientation() == EOrientation::Orient_Vertical; }),
+			FIsActionChecked::CreateLambda([this]() { return VerticalGraphJustification == ESignalFlowVerticalJustification::Right; }));
+
 		CommandList->MapAction(
 			Commands.GetShowNodeDetailsCommand(),
 			FExecuteAction::CreateLambda([this]() { ToggleShowNodeDetails(); }),
@@ -503,6 +539,71 @@ namespace UE::Audio::Insights
 			[
 				Style.CreateButtonContentWidget("AudioInsights.Icon.SignalFlow.CenterView", FText::GetEmpty(), ButtonTextStyleName)
 			]
+		]
+		
+		+ SHorizontalBox::Slot()
+		.AutoWidth()
+		.Padding(ToolbarButtonPaddingX, ToolbarButtonPaddingY)
+		[
+			SNew(SHorizontalBox)
+
+			+ SHorizontalBox::Slot()
+			.AutoWidth()
+			[
+				SNew(SSegmentedControl<ESignalFlowHorizontalJustification>)
+				.Visibility_Lambda([this]()
+				{
+					return GraphWidgetContents.IsValid() && GraphWidgetContents->GetGraphOrientation() == EOrientation::Orient_Horizontal
+																										? EVisibility::Visible
+																										: EVisibility::Collapsed;
+				})
+				.OnValueChanged(this, &FSignalFlowDashboardViewFactory::SetHorizontalJustification)
+				.Value_Lambda([this]
+				{
+					return HorizontalGraphJustification;
+				})
+					
+				+ SSegmentedControl<ESignalFlowHorizontalJustification>::Slot(ESignalFlowHorizontalJustification::Bottom)
+				.ToolTip(LOCTEXT("AudioDashboard_SignalFlow_HorizontalJustification_Bottom_Tooltip", "Position nodes towards the bottom."))
+				.Icon(FSlateStyle::Get().GetBrush("AudioInsights.Icon.SignalFlow.AlignHorizontalBottom"))
+				
+				+ SSegmentedControl<ESignalFlowHorizontalJustification>::Slot(ESignalFlowHorizontalJustification::Center)
+				.ToolTip(LOCTEXT("AudioDashboard_SignalFlow_HorizontalJustification_Center_Tooltip", "Position nodes around the center."))
+				.Icon(FSlateStyle::Get().GetBrush("AudioInsights.Icon.SignalFlow.AlignHorizontalCenter"))
+
+				+ SSegmentedControl<ESignalFlowHorizontalJustification>::Slot(ESignalFlowHorizontalJustification::Top)
+				.ToolTip(LOCTEXT("AudioDashboard_SignalFlow_HorizontalJustification_Top_Tooltip", "Position nodes towards the top."))
+				.Icon(FSlateStyle::Get().GetBrush("AudioInsights.Icon.SignalFlow.AlignHorizontalTop"))
+			]
+
+			+ SHorizontalBox::Slot()
+			.AutoWidth()
+			[
+				SNew(SSegmentedControl<ESignalFlowVerticalJustification>)
+				.Visibility_Lambda([this]()
+				{
+					return GraphWidgetContents.IsValid() && GraphWidgetContents->GetGraphOrientation() == EOrientation::Orient_Vertical
+																										? EVisibility::Visible
+																										: EVisibility::Collapsed;
+				})
+				.OnValueChanged(this, &FSignalFlowDashboardViewFactory::SetVerticalJustification)
+				.Value_Lambda([this]
+				{
+					return VerticalGraphJustification;
+				})
+					
+				+ SSegmentedControl<ESignalFlowVerticalJustification>::Slot(ESignalFlowVerticalJustification::Left)
+				.ToolTip(LOCTEXT("AudioDashboard_SignalFlow_VerticalJustification_Left_Tooltip", "Position nodes towards the left."))
+				.Icon(FSlateStyle::Get().GetBrush("AudioInsights.Icon.SignalFlow.AlignVerticalLeft"))
+				
+				+ SSegmentedControl<ESignalFlowVerticalJustification>::Slot(ESignalFlowVerticalJustification::Center)
+				.ToolTip(LOCTEXT("AudioDashboard_SignalFlow_VerticalJustification_Center_Tooltip", "Position nodes around the center."))
+				.Icon(FSlateStyle::Get().GetBrush("AudioInsights.Icon.SignalFlow.AlignVerticalCenter"))
+
+				+ SSegmentedControl<ESignalFlowVerticalJustification>::Slot(ESignalFlowVerticalJustification::Right)
+				.ToolTip(LOCTEXT("AudioDashboard_SignalFlow_VerticalJustification_Right_Tooltip", "Position nodes towards the right."))
+				.Icon(FSlateStyle::Get().GetBrush("AudioInsights.Icon.SignalFlow.AlignVerticalRight"))
+			]
 		];
 	}
 
@@ -524,6 +625,18 @@ namespace UE::Audio::Insights
 
 			MenuBuilder.AddMenuEntry(Commands.GetToggleHorizontalViewCommand());
 			MenuBuilder.AddMenuEntry(Commands.GetPauseOnSelectCommand());
+			MenuBuilder.AddSubMenu
+			(
+				LOCTEXT("AudioDashboard_SignalFlow_Settings_HorizontalJustification", "Horizontal Graph Justification"),
+				LOCTEXT("AudioDashboard_SignalFlow_Settings_HorizontalJustificationTooltip", "Choose where nodes are aligned when viewing the signal flow graph in horizontal mode."),
+				FNewMenuDelegate::CreateRaw(this, &FSignalFlowDashboardViewFactory::BuildHorizontalJustificationMenuContent)
+			);
+			MenuBuilder.AddSubMenu
+			(
+				LOCTEXT("AudioDashboard_SignalFlow_Settings_VerticalJustification", "Vertical Graph Justification"),
+				LOCTEXT("AudioDashboard_SignalFlow_Settings_VerticalJustificationTooltip", "Choose where nodes are aligned when viewing the signal flow graph in vertical mode."),
+				FNewMenuDelegate::CreateRaw(this, &FSignalFlowDashboardViewFactory::BuildVerticalJustificationMenuContent)
+			);
 
 			MenuBuilder.AddMenuEntry(Commands.GetShowNodeDetailsCommand());
 			MenuBuilder.AddSubMenu
@@ -539,13 +652,31 @@ namespace UE::Audio::Insights
 		return MenuBuilder.MakeWidget();
 	}
 
-	void FSignalFlowDashboardViewFactory::BuildNodeDetailFiltersMenuContent(FMenuBuilder& MenuBuilder)
+	void FSignalFlowDashboardViewFactory::BuildHorizontalJustificationMenuContent(FMenuBuilder& OutMenuBuilder)
+	{
+		const FSignalFlowEditorCommands& Commands = FSignalFlowEditorCommands::Get();
+
+		OutMenuBuilder.AddMenuEntry(Commands.GetJustifyBottomCommand());
+		OutMenuBuilder.AddMenuEntry(Commands.GetJustifyHorizontalCenterCommand());
+		OutMenuBuilder.AddMenuEntry(Commands.GetJustifyTopCommand());
+	}
+
+	void FSignalFlowDashboardViewFactory::BuildVerticalJustificationMenuContent(FMenuBuilder& OutMenuBuilder)
+	{
+		const FSignalFlowEditorCommands& Commands = FSignalFlowEditorCommands::Get();
+
+		OutMenuBuilder.AddMenuEntry(Commands.GetJustifyLeftCommand());
+		OutMenuBuilder.AddMenuEntry(Commands.GetJustifyVerticalCenterCommand());
+		OutMenuBuilder.AddMenuEntry(Commands.GetJustifyRightCommand());
+	}
+
+	void FSignalFlowDashboardViewFactory::BuildNodeDetailFiltersMenuContent(FMenuBuilder& OutMenuBuilder)
 	{
 		using namespace FSignalFlowDashboardViewFactoryPrivate;
 
 		const FSignalFlowEditorCommands& Commands = FSignalFlowEditorCommands::Get();
 
-		MenuBuilder.AddMenuEntry
+		OutMenuBuilder.AddMenuEntry
 		(
 			LOCTEXT("AudioDashboard_SignalFlow_Settings_ToggleAllNodeDetails", "Show all"),
 			FText::GetEmpty(),
@@ -566,7 +697,7 @@ namespace UE::Audio::Insights
 		{
 			const ESignalFlowNodeDetailParam ParamType = static_cast<ESignalFlowNodeDetailParam>(FilterID);
 
-			MenuBuilder.AddMenuEntry
+			OutMenuBuilder.AddMenuEntry
 			(
 				GetParamFilterName(ParamType),
 				FText::GetEmpty(),
@@ -693,6 +824,8 @@ namespace UE::Audio::Insights
 #if WITH_EDITOR
 		FSignalFlowSettings::OnRequestWriteSettings.Broadcast();
 #endif // WITH_EDITOR	
+
+		OnRequestGraphRefresh();
 	}
 
 	void FSignalFlowDashboardViewFactory::ToggleShowNodeDetails()
@@ -770,6 +903,28 @@ namespace UE::Audio::Insights
 		TimingViewExtender.ResetMessageProcessType();
 	}
 
+	void FSignalFlowDashboardViewFactory::SetHorizontalJustification(const ESignalFlowHorizontalJustification Justification)
+	{
+		HorizontalGraphJustification = Justification;
+
+#if WITH_EDITOR
+		FSignalFlowSettings::OnRequestWriteSettings.Broadcast();
+#endif // WITH_EDITOR
+
+		OnRequestGraphRefresh();
+	}
+
+	void FSignalFlowDashboardViewFactory::SetVerticalJustification(const ESignalFlowVerticalJustification Justification)
+	{
+		VerticalGraphJustification = Justification;
+
+#if WITH_EDITOR
+		FSignalFlowSettings::OnRequestWriteSettings.Broadcast();
+#endif // WITH_EDITOR
+
+		OnRequestGraphRefresh();
+	}
+
 	TSharedPtr<FSignalFlowDashboardEntry> FSignalFlowDashboardViewFactory::GetFilteredAudioDeviceEntry() const
 	{
 		const TSharedPtr<const FSignalFlowTraceProvider> Provider = FindProvider<const FSignalFlowTraceProvider>();
@@ -1775,6 +1930,12 @@ namespace UE::Audio::Insights
 
 	void FSignalFlowDashboardViewFactory::CalcXPositions()
 	{
+		if (!GraphWidgetContents.IsValid())
+		{
+			return;
+		}
+
+		const EOrientation GraphOrientation = GraphWidgetContents->GetGraphOrientation();
 		int32 GraphWidth = 0;
 
 		TreeDepthPair CurrentTreeDepthPair{ ESignalFlowEntryType::InvalidMax, TNumericLimits<int32>::Max() };
@@ -1799,7 +1960,7 @@ namespace UE::Audio::Insights
 			// If we have moved on to a new row, calculate the positions of the current structure
 			if (CurrentTreeDepthPair != NextNodeToProcess->GetTreeDepth())
 			{
-				CalcXPositionsForRows(ProcessedNodeRows, GraphWidth);
+				CalcXPositionsForRows(ProcessedNodeRows, GraphOrientation, GraphWidth);
 
 				// Create the next row
 				RowIndex = RowIndex == INDEX_NONE ? 0 : RowIndex + 1;
@@ -1816,10 +1977,10 @@ namespace UE::Audio::Insights
 		}
 
 		// Do one final positioning process for the final row
-		CalcXPositionsForRows(ProcessedNodeRows, GraphWidth);
+		CalcXPositionsForRows(ProcessedNodeRows, GraphOrientation, GraphWidth);
 	}
 
-	void FSignalFlowDashboardViewFactory::CalcXPositionsForRows(const TArray<TArray<int32>>& NodeIndexRows, int32& OutGraphWidth)
+	void FSignalFlowDashboardViewFactory::CalcXPositionsForRows(const TArray<TArray<int32>>& NodeIndexRows, const EOrientation GraphOrientation, int32& OutGraphWidth)
 	{
 		if (NodeIndexRows.IsEmpty())
 		{
@@ -1833,52 +1994,49 @@ namespace UE::Audio::Insights
 
 		if (OutGraphWidth <= RowToProcessNumNodes)
 		{
-			// new largest row, position the nodes around X = 0.0
-			PositionRowAroundOrigin(RowToProcess);
+			OutGraphWidth = RowToProcessNumNodes;
+
+			// new largest row, position the nodes with an x position of 1.0 apart
+			PositionRowEquidistant(RowToProcess);
 
 			// Propagate change down the hierarchy - moving nodes relative to their inputs
 			for (int32 PreviousRowsIndex = RowToProcessIndex - 1; PreviousRowsIndex >= 0; --PreviousRowsIndex)
 			{
 				const TArray<int32> PrevRowOfIndexes = NodeIndexRows[PreviousRowsIndex];
-				TArray<TSharedPtr<ISignalFlowNode>> PrevPositionedNodes;
-				PrevPositionedNodes.Reserve(PrevRowOfIndexes.Num());
 
-				for (const int32 NodeIndexInGraph : PrevRowOfIndexes)
+				TOptional<float> PrevXPos;
+				for (int32 NodeIndexInRow = 0; NodeIndexInRow < PrevRowOfIndexes.Num(); ++NodeIndexInRow)
 				{
-					const TSharedPtr<ISignalFlowNode>& Node = FilteredGraphNodes[NodeIndexInGraph];
-					CalcXPositionForNode(Node, Node->FilteredInputs, PrevPositionedNodes, PrevRowOfIndexes.Num());
+					const TSharedPtr<ISignalFlowNode>& Node = FilteredGraphNodes[PrevRowOfIndexes[NodeIndexInRow]];
 
-					PrevPositionedNodes.Add(Node);
+					if (Node.IsValid())
+					{
+						CalcXPositionForNode(Node, Node->FilteredInputs, PrevXPos, GraphOrientation);
+						PrevXPos = Node->XPos;
+					}
 				}
 			}
-
-			OutGraphWidth = RowToProcessNumNodes;
 		}
 		else
 		{
-			TArray<TSharedPtr<ISignalFlowNode>> PrevPositionedNodes;
-			PrevPositionedNodes.Reserve(RowToProcess.Num());
-
 			// Process just this row - moving nodes relative to their outputs
+			TOptional<float> PrevXPos;
 			for (int32 NodeIndexInRow = 0; NodeIndexInRow < RowToProcess.Num(); ++NodeIndexInRow)
 			{
 				TSharedPtr<ISignalFlowNode>& Node = FilteredGraphNodes[RowToProcess[NodeIndexInRow]];
 				if (Node.IsValid())
 				{
-					CalcXPositionForNode(Node, Node->FilteredOutputs, PrevPositionedNodes, RowToProcess.Num());
+					CalcXPositionForNode(Node, Node->FilteredOutputs, PrevXPos, GraphOrientation);
+					PrevXPos = Node->XPos;
 				}
-
-				PrevPositionedNodes.Add(Node);
 			}
 		}
 	}
 
-	void FSignalFlowDashboardViewFactory::PositionRowAroundOrigin(const TArray<int32>& NodeIndexRow)
+	void FSignalFlowDashboardViewFactory::PositionRowEquidistant(const TArray<int32>& NodeIndexRow)
 	{
-		// Space a row of nodes around zero (origin) with a X distance of 1.0
+		// Space a row of nodes with a X distance of 1.0
 		const int32 NumNodes = NodeIndexRow.Num();
-		const bool bIsEvenNumNodes = NumNodes % 2 == 0;
-		const float HalfNumNodes = static_cast<float>(NumNodes) / 2.0f;
 
 		for (int32 NodeIndexInRow = 0; NodeIndexInRow < NodeIndexRow.Num(); ++NodeIndexInRow)
 		{
@@ -1888,95 +2046,143 @@ namespace UE::Audio::Insights
 				continue;
 			}
 
-			if (bIsEvenNumNodes)
-			{
-				// Even number of nodes - space the nodes around the origin with an empty space at the origin
-				Node->XPos = NumNodes - HalfNumNodes - NodeIndexInRow;
-				if (NodeIndexInRow >= HalfNumNodes)
+			Node->XPos = NodeIndexInRow;
+		}
+	}
+
+	void FSignalFlowDashboardViewFactory::CalcXPositionForNode(const TSharedPtr<ISignalFlowNode>& Node, const TArray<FSignalFlowEntryKey>& NodesToPositionRelativeTo, const TOptional<float>& PrevXPos, const EOrientation GraphOrientation)
+	{
+		switch (GraphOrientation)
+		{
+			case EOrientation::Orient_Horizontal:
+
+				switch (HorizontalGraphJustification)
 				{
-					Node->XPos--;
+					case ESignalFlowHorizontalJustification::Bottom:
+						AlignNodeTowardsMaxX(Node, NodesToPositionRelativeTo, PrevXPos);
+						break;
+
+					case ESignalFlowHorizontalJustification::Center:
+						AlignNodeTowardsCenterX(Node, NodesToPositionRelativeTo, PrevXPos);
+						break;
+
+					case ESignalFlowHorizontalJustification::Top:
+						AlignNodeTowardsLowestX(Node, NodesToPositionRelativeTo, PrevXPos);
+						break;
 				}
-			}
-			else
+
+				break;
+
+			case EOrientation::Orient_Vertical:
+
+				switch (VerticalGraphJustification)
+				{
+				case ESignalFlowVerticalJustification::Left:
+					AlignNodeTowardsLowestX(Node, NodesToPositionRelativeTo, PrevXPos);
+					break;
+
+				case ESignalFlowVerticalJustification::Center:
+					AlignNodeTowardsCenterX(Node, NodesToPositionRelativeTo, PrevXPos);
+					break;
+
+				case ESignalFlowVerticalJustification::Right:
+					AlignNodeTowardsMaxX(Node, NodesToPositionRelativeTo, PrevXPos);
+					break;
+				}
+
+				break;
+		}
+	}
+
+	void FSignalFlowDashboardViewFactory::AlignNodeTowardsCenterX(const TSharedPtr<ISignalFlowNode>& Node, const TArray<FSignalFlowEntryKey>& NodesToPositionRelativeTo, const TOptional<float>& PrevXPos)
+	{
+		if (!Node.IsValid() || NodesToPositionRelativeTo.IsEmpty())
+		{
+			return;
+		}
+
+		float TotalXPositions = 0.0f;
+		for (int32 AligningNodeIndex = 0; AligningNodeIndex < NodesToPositionRelativeTo.Num(); ++AligningNodeIndex)
+		{
+			TSharedPtr<ISignalFlowNode>* AliningNode = FilteredEntries.Find(NodesToPositionRelativeTo[AligningNodeIndex]);
+			if (AliningNode && AliningNode->IsValid())
 			{
-				// Odd number of nodes - evenly space the nodes around the origin
-				Node->XPos = static_cast<float>(((NumNodes - NodeIndexInRow - 1) * 2) - (NumNodes - 1)) / 2.0f;
+				TotalXPositions += (*AliningNode)->XPos;
 			}
 		}
+
+		const float AverageXPosition = TotalXPositions / NodesToPositionRelativeTo.Num();
+
+		if (PrevXPos.IsSet())
+		{
+			Node->XPos = FMath::Max(AverageXPosition, PrevXPos.GetValue() + 1.0f);
+		}
+		else
+		{
+			Node->XPos = AverageXPosition;
+		}
 	}
 
-	void FSignalFlowDashboardViewFactory::CalcXPositionForNode(const TSharedPtr<ISignalFlowNode>& Node, const TArray<FSignalFlowEntryKey>& NodesToPositionRelativeTo, const TArray<TSharedPtr<ISignalFlowNode>>& PrevPositionedNodes, const int32 NumNodesInRow)
+	void FSignalFlowDashboardViewFactory::AlignNodeTowardsLowestX(const TSharedPtr<ISignalFlowNode>& Node, const TArray<FSignalFlowEntryKey>& NodesToPositionRelativeTo, const TOptional<float>& PrevXPos)
 	{
 		if (!Node.IsValid() || NodesToPositionRelativeTo.IsEmpty())
 		{
 			return;
 		}
 
-		const int32 NumPositioningNodes = NodesToPositionRelativeTo.Num();
-		const bool bEvenNumPositioningNodes = NumPositioningNodes % 2 == 0;
-
-		if (bEvenNumPositioningNodes)
+		TSharedPtr<ISignalFlowNode> NodeToAlignWith = nullptr;
+		float MinXPos = TNumericLimits<float>::Max();
+		for (int32 AligningNodeIndex = 0; AligningNodeIndex < NodesToPositionRelativeTo.Num(); ++AligningNodeIndex)
 		{
-			// Even num positioners: nodes are positioned in between the two middle positioning nodes
-			const int32 NodeToAlignWithIndex02 = NumPositioningNodes / 2;
-			const int32 NodeToAlignWithIndex01 = NodeToAlignWithIndex02 - 1;
-			
-			if (ensure(NodesToPositionRelativeTo.IsValidIndex(NodeToAlignWithIndex01) && NodesToPositionRelativeTo.IsValidIndex(NodeToAlignWithIndex02)))
+			TSharedPtr<ISignalFlowNode>* AliningNode = FilteredEntries.Find(NodesToPositionRelativeTo[AligningNodeIndex]);
+			if (AliningNode && AliningNode->IsValid() && (*AliningNode)->XPos < MinXPos)
 			{
-				TSharedPtr<ISignalFlowNode>* NodeToAlignWith01 = FilteredEntries.Find(NodesToPositionRelativeTo[NodeToAlignWithIndex01]);
-				TSharedPtr<ISignalFlowNode>* NodeToAlignWith02 = FilteredEntries.Find(NodesToPositionRelativeTo[NodeToAlignWithIndex02]);
-
-				if (NodeToAlignWith01 && NodeToAlignWith02 && NodeToAlignWith01->IsValid() && NodeToAlignWith02->IsValid())
-				{
-					Node->XPos = FMath::Floor(((*NodeToAlignWith01)->XPos + (*NodeToAlignWith02)->XPos) / 2.0f);
-					MakeSpaceForNodeInRow(Node, PrevPositionedNodes, NumNodesInRow);
-				}
+				NodeToAlignWith = (*AliningNode);
+				MinXPos = (*AliningNode)->XPos;
 			}
 		}
-		else
+
+		if (NodeToAlignWith.IsValid())
 		{
-			// Odd num positioners: nodes are positioned in line with the middle positioning node
-			const int32 NodeToAlignWithIndex = NumPositioningNodes / 2;
-			if (ensure(NodesToPositionRelativeTo.IsValidIndex(NodeToAlignWithIndex)))
+			if (PrevXPos.IsSet())
 			{
-				TSharedPtr<ISignalFlowNode>* NodeToAlignWith = FilteredEntries.Find(NodesToPositionRelativeTo[NodeToAlignWithIndex]);
-				if (NodeToAlignWith && NodeToAlignWith->IsValid())
-				{
-					Node->XPos = (*NodeToAlignWith)->XPos;
-					MakeSpaceForNodeInRow(Node, PrevPositionedNodes, NumNodesInRow);
-				}
+				Node->XPos = FMath::Max(NodeToAlignWith->XPos, PrevXPos.GetValue() + 1.0f);
+			}
+			else
+			{
+				Node->XPos = NodeToAlignWith->XPos;
 			}
 		}
 	}
 
-	void FSignalFlowDashboardViewFactory::MakeSpaceForNodeInRow(const TSharedPtr<ISignalFlowNode>& Node, const TArray<TSharedPtr<ISignalFlowNode>>& PrevPositionedNodes, const int32 NumNodesInRow)
+	void FSignalFlowDashboardViewFactory::AlignNodeTowardsMaxX(const TSharedPtr<ISignalFlowNode>& Node, const TArray<FSignalFlowEntryKey>& NodesToPositionRelativeTo, const TOptional<float>& PrevXPos)
 	{
-		if (!Node->IsRealNode())
+		if (!Node.IsValid() || NodesToPositionRelativeTo.IsEmpty())
 		{
 			return;
 		}
 
-		// In cases where nodes share inputs/outputs, we can end up with nodes placed in the same XPos
-		// Check if we have hit this edge case and reposition the row accordingly
-		const TSharedPtr<ISignalFlowNode>* PrevNodeAtExistingPosition = PrevPositionedNodes.FindByPredicate([XPos = Node->XPos](const TSharedPtr<ISignalFlowNode>& AlreadyPositionedNode)
+		TSharedPtr<ISignalFlowNode> NodeToAlignWith = nullptr;
+		float MaxXPos = TNumericLimits<float>::Lowest();
+		for (int32 AligningNodeIndex = 0; AligningNodeIndex < NodesToPositionRelativeTo.Num(); ++AligningNodeIndex)
 		{
-			return AlreadyPositionedNode.IsValid() && AlreadyPositionedNode->XPos == XPos;
-		});
+			TSharedPtr<ISignalFlowNode>* AliningNode = FilteredEntries.Find(NodesToPositionRelativeTo[AligningNodeIndex]);
+			if (AliningNode && AliningNode->IsValid() && (*AliningNode)->XPos > MaxXPos)
+			{
+				NodeToAlignWith = (*AliningNode);
+				MaxXPos = (*AliningNode)->XPos;
+			}
+		}
 
-		if (PrevNodeAtExistingPosition)
+		if (NodeToAlignWith.IsValid())
 		{
-			// there already is a node at this position - shift nodes to the left
-			for (const TSharedPtr<ISignalFlowNode>& PrevNode : PrevPositionedNodes)
+			if (PrevXPos.IsSet())
 			{
-				if (!PrevNode.IsValid())
-				{
-					continue;
-				}
-			
-				if (PrevNode->XPos <= Node->XPos)
-				{
-					PrevNode->XPos -= 1.0f;
-				} 
+				Node->XPos = FMath::Max(NodeToAlignWith->XPos, PrevXPos.GetValue() + 1.0f);
+			}
+			else
+			{
+				Node->XPos = NodeToAlignWith->XPos;
 			}
 		}
 	}
@@ -1996,6 +2202,8 @@ namespace UE::Audio::Insights
 		}
 
 		bPauseGraphOnSelect = InSettings.bPauseGraphOnSelect;
+		HorizontalGraphJustification = InSettings.HorizontalGraphJustification;
+		VerticalGraphJustification = InSettings.VerticalGraphJustification;
 		bShowNodeDetails = InSettings.bShowNodeDetails;
 		NodeDetailFilters = InSettings.NodeDetailFilters;
 		bDisplayAmpPeakInDb = InSettings.AmplitudeDisplayMode == EAudioAmplitudeDisplayMode::Decibels;
@@ -2016,6 +2224,8 @@ namespace UE::Audio::Insights
 			OutSettings.bHorizontalFlow = bGraphIsHorizontal;
 		}
 
+		OutSettings.HorizontalGraphJustification = HorizontalGraphJustification;
+		OutSettings.VerticalGraphJustification = VerticalGraphJustification;
 		OutSettings.bShowNodeDetails = bShowNodeDetails;
 		OutSettings.bPauseGraphOnSelect = bPauseGraphOnSelect;
 		OutSettings.NodeDetailFilters = NodeDetailFilters;
diff --git a/Engine/Plugins/AudioInsights/Source/AudioInsights/Public/AudioInsightsStyle.h b/Engine/Plugins/AudioInsights/Source/AudioInsights/Public/AudioInsightsStyle.h
--- a/Engine/Plugins/AudioInsights/Source/AudioInsights/Public/AudioInsightsStyle.h
+++ b/Engine/Plugins/AudioInsights/Source/AudioInsights/Public/AudioInsightsStyle.h
@@ -265,6 +265,13 @@ namespace UE::Audio::Insights
 			Set("AudioInsights.Icon.EventLog.MarkedForDelete",		  new IMAGE_BRUSH_SVG(TEXT("Icons/EventLog/cache_marked_for_delete"), Icon16));
 
 			/* Signal Flow */
+			Set("AudioInsights.Icon.SignalFlow.AlignHorizontalBottom", new IMAGE_BRUSH_SVG(TEXT("Icons/SignalFlow/align_horizontal_bottom"), Icon20));
+			Set("AudioInsights.Icon.SignalFlow.AlignHorizontalCenter", new IMAGE_BRUSH_SVG(TEXT("Icons/SignalFlow/align_horizontal_center"), Icon20));
+			Set("AudioInsights.Icon.SignalFlow.AlignHorizontalTop",	   new IMAGE_BRUSH_SVG(TEXT("Icons/SignalFlow/align_horizontal_top"),	 Icon20));
+			Set("AudioInsights.Icon.SignalFlow.AlignVerticalCenter",   new IMAGE_BRUSH_SVG(TEXT("Icons/SignalFlow/align_vertical_center"),	 Icon20));
+			Set("AudioInsights.Icon.SignalFlow.AlignVerticalLeft",	   new IMAGE_BRUSH_SVG(TEXT("Icons/SignalFlow/align_vertical_left"),	 Icon20));
+			Set("AudioInsights.Icon.SignalFlow.AlignVerticalRight",	   new IMAGE_BRUSH_SVG(TEXT("Icons/SignalFlow/align_vertical_right"),	 Icon20));
+
 			Set("AudioInsights.Icon.SignalFlow.CenterView",			  new IMAGE_BRUSH_SVG(TEXT("Icons/SignalFlow/center_view"),			 Icon20));
 			Set("AudioInsights.Icon.SignalFlow.HorizontalFlow",		  new IMAGE_BRUSH_SVG(TEXT("Icons/SignalFlow/horizontal_flow"),		 Icon20));
 			Set("AudioInsights.Icon.SignalFlow.NodeDetails",		  new IMAGE_BRUSH_SVG(TEXT("Icons/SignalFlow/node_details"),		 Icon20));
diff --git a/Engine/Plugins/AudioInsights/Source/AudioInsights/Public/Settings/SignalFlowSettings.h b/Engine/Plugins/AudioInsights/Source/AudioInsights/Public/Settings/SignalFlowSettings.h
--- a/Engine/Plugins/AudioInsights/Source/AudioInsights/Public/Settings/SignalFlowSettings.h
+++ b/Engine/Plugins/AudioInsights/Source/AudioInsights/Public/Settings/SignalFlowSettings.h
@@ -12,6 +12,40 @@
 
 #define LOCTEXT_NAMESPACE "AudioInsights"
 
+/** ESignalFlowHorizontalJustification
+ *
+ * Where the nodes in the signal flow graph are aligned when viewed in horizontal mode
+ */
+UENUM(BlueprintType)
+enum class ESignalFlowHorizontalJustification : uint8
+{
+	/** Position nodes towards the bottom */
+	Bottom		UMETA(DisplayName = "Bottom"),
+
+	/** Position nodes around the center */
+	Center		UMETA(DisplayName = "Center"),
+
+	/** Position nodes towards the top */
+	Top			UMETA(DisplayName = "Top")
+};
+
+/** ESignalFlowVerticalJustification
+ *
+ * Where the nodes in the signal flow graph are aligned when viewed in vertical mode
+ */
+UENUM(BlueprintType)
+enum class ESignalFlowVerticalJustification : uint8
+{
+	/** Position nodes towards the left */
+	Left		UMETA(DisplayName = "Left"),
+
+	/** Position nodes around the center */
+	Center		UMETA(DisplayName = "Center"),
+
+	/** Position nodes towards the right */
+	Right		UMETA(DisplayName = "Right")
+};
+
 USTRUCT()
 struct FSignalFlowNodeDetailFilterSettings
 {
@@ -78,6 +112,14 @@ struct FSignalFlowSettings
 	UPROPERTY(EditAnywhere, config, Category = SignalFlow)
 	bool bPauseGraphOnSelect = false;
 
+	/** Where the nodes in the signal flow graph are aligned when viewed in horizontal mode */
+	UPROPERTY(EditAnywhere, config, Category = SignalFlow, meta = (DisplayName = "Horizontal Graph Justification"))
+	ESignalFlowHorizontalJustification HorizontalGraphJustification = ESignalFlowHorizontalJustification::Top;
+
+	/** Where the nodes in the signal flow graph are aligned when viewed in vertical mode */
+	UPROPERTY(EditAnywhere, config, Category = SignalFlow, meta = (DisplayName = "Vertical Graph Justification"))
+	ESignalFlowVerticalJustification VerticalGraphJustification = ESignalFlowVerticalJustification::Left;
+
 	/** Toggles visibility of node details. */
 	UPROPERTY(EditAnywhere, config, Category = SignalFlow)
 	bool bShowNodeDetails = false;
diff --git a/Engine/Plugins/AudioInsights/Source/AudioInsights/Public/Views/SignalFlowDashboardViewFactory.h b/Engine/Plugins/AudioInsights/Source/AudioInsights/Public/Views/SignalFlowDashboardViewFactory.h
--- a/Engine/Plugins/AudioInsights/Source/AudioInsights/Public/Views/SignalFlowDashboardViewFactory.h
+++ b/Engine/Plugins/AudioInsights/Source/AudioInsights/Public/Views/SignalFlowDashboardViewFactory.h
@@ -77,7 +77,9 @@ namespace UE::Audio::Insights
 		TSharedRef<SWidget> CreateToolbarButtonActions();
 
 		TSharedRef<SWidget> OnGetSettingsMenuContent();
-		void BuildNodeDetailFiltersMenuContent(FMenuBuilder& MenuBuilder);
+		void BuildHorizontalJustificationMenuContent(FMenuBuilder& OutMenuBuilder);
+		void BuildVerticalJustificationMenuContent(FMenuBuilder& OutMenuBuilder);
+		void BuildNodeDetailFiltersMenuContent(FMenuBuilder& OutMenuBuilder);
 		void CreateAmpDisplayModeToggle(FMenuBuilder& OutMenuBuilder);
 
 		FReply ScrollGraphViewToSelectedNode(const bool bResetZoom);
@@ -86,6 +88,8 @@ namespace UE::Audio::Insights
 		void TogglePauseGraphEnabled();
 		void ToggleEnableAllNodeDetailFilters();
 		void ToggleAnimateWires();
+		void SetHorizontalJustification(const ESignalFlowHorizontalJustification Justification);
+		void SetVerticalJustification(const ESignalFlowVerticalJustification Justification);
 
 		bool CanPauseTimestamp() const;
 		bool IsTimestampPaused() const;
@@ -129,10 +133,13 @@ namespace UE::Audio::Insights
 		void SortNodeOrder();
 
 		void CalcXPositions();
-		void CalcXPositionsForRows(const TArray<TArray<int32>>& NodeIndexRows, int32& OutGraphWidth);
-		void PositionRowAroundOrigin(const TArray<int32>& NodeIndexRow);
-		void CalcXPositionForNode(const TSharedPtr<ISignalFlowNode>& Node, const TArray<FSignalFlowEntryKey>& NodesToPositionRelativeTo, const TArray<TSharedPtr<ISignalFlowNode>>& PrevPositionedNodes, const int32 NumNodesInRow);
-		void MakeSpaceForNodeInRow(const TSharedPtr<ISignalFlowNode>& Node, const TArray<TSharedPtr<ISignalFlowNode>>& PrevPositionedNodes, const int32 NumNodesInRow);
+		void CalcXPositionsForRows(const TArray<TArray<int32>>& NodeIndexRows, const EOrientation GraphOrientation, int32& OutGraphWidth);
+		void PositionRowEquidistant(const TArray<int32>& NodeIndexRow);
+		
+		void CalcXPositionForNode(const TSharedPtr<ISignalFlowNode>& Node, const TArray<FSignalFlowEntryKey>& NodesToPositionRelativeTo, const TOptional<float>& PrevXPos, const EOrientation GraphOrientation);
+		void AlignNodeTowardsCenterX(const TSharedPtr<ISignalFlowNode>& Node, const TArray<FSignalFlowEntryKey>& NodesToPositionRelativeTo, const TOptional<float>& PrevXPos);
+		void AlignNodeTowardsLowestX(const TSharedPtr<ISignalFlowNode>& Node, const TArray<FSignalFlowEntryKey>& NodesToPositionRelativeTo, const TOptional<float>& PrevXPos);
+		void AlignNodeTowardsMaxX(const TSharedPtr<ISignalFlowNode>& Node, const TArray<FSignalFlowEntryKey>& NodesToPositionRelativeTo, const TOptional<float>& PrevXPos);
 
 #if WITH_EDITOR
 		void OnReadEditorSettings(const FSignalFlowSettings& InSettings);
@@ -166,6 +173,9 @@ namespace UE::Audio::Insights
 
 		FSignalFlowNodeDetailFilterSettings NodeDetailFilters;
 
+		ESignalFlowHorizontalJustification HorizontalGraphJustification = ESignalFlowHorizontalJustification::Top;
+		ESignalFlowVerticalJustification VerticalGraphJustification = ESignalFlowVerticalJustification::Left;
+
 		bool bGraphRequiresRefresh = false;
 		bool bAutoFocusAfterGraphRefresh = false;
 		bool bShowNodeDetails = false;
