[Audio Insights] Output Meter: added trace messages and provider classes + adjusted dashboard so TruePeak is (and future loudness meters could be) available in standalone.

#Synced-CL 50901149

diff --git a/Engine/Plugins/AudioInsights/Source/AudioInsights/Private/Messages/OutputMeterTraceMessages.cpp b/Engine/Plugins/AudioInsights/Source/AudioInsights/Private/Messages/OutputMeterTraceMessages.cpp
new file mode 100644
--- /dev/null
+++ b/Engine/Plugins/AudioInsights/Source/AudioInsights/Private/Messages/OutputMeterTraceMessages.cpp
@@ -0,0 +1,48 @@
+// Copyright Epic Games, Inc. All Rights Reserved.
+#include "Messages/OutputMeterTraceMessages.h"
+
+namespace UE::Audio::Insights
+{
+	namespace OutputMeterMessageNames
+	{
+		const FName MainSubmixLoaded = "OutputMeterMainSubmixLoaded";
+		const FName TruePeak = "OutputMeterTruePeak";
+	};
+
+	// FOutputMeterMessageBase
+	FOutputMeterMessageBase::FOutputMeterMessageBase(const Trace::IAnalyzer::FOnEventContext& InContext)
+	{
+		const Trace::IAnalyzer::FEventData& EventData = InContext.EventData;
+
+		DeviceId  = static_cast<::Audio::FDeviceId>(EventData.GetValue<uint32>("DeviceId"));
+		SubmixId  = EventData.GetValue<uint32>("SubmixId");
+		Timestamp = InContext.EventTime.AsSeconds(EventData.GetValue<uint64>("Timestamp"));
+	}
+
+	// FOutputMeterMainSubmixLoadedMessage
+	FOutputMeterMainSubmixLoadedMessage::FOutputMeterMainSubmixLoadedMessage(const Trace::IAnalyzer::FOnEventContext& InContext)
+		: FOutputMeterMessageBase(InContext)
+	{
+	}
+
+	uint32 FOutputMeterMainSubmixLoadedMessage::GetSizeOf() const
+	{
+		// Returning fixed size of the struct (no dynamically sized members in it)
+		return sizeof(FOutputMeterMainSubmixLoadedMessage);
+	}
+
+	// FOutputMeterTruePeakMessage
+	FOutputMeterTruePeakMessage::FOutputMeterTruePeakMessage(const Trace::IAnalyzer::FOnEventContext& InContext)
+		: FOutputMeterMessageBase(InContext)
+	{
+		const Trace::IAnalyzer::FEventData& EventData = InContext.EventData;
+
+		TruePeakMaxValueDb = EventData.GetValue<float>("TruePeakMaxValueDb");
+	}
+	
+	uint32 FOutputMeterTruePeakMessage::GetSizeOf() const
+	{
+		// Returning fixed size of the struct (no dynamically sized members in it)
+		return sizeof(FOutputMeterTruePeakMessage);
+	}
+} // namespace UE::Audio::Insights
diff --git a/Engine/Plugins/AudioInsights/Source/AudioInsights/Private/Providers/OutputMeterTraceProvider.cpp b/Engine/Plugins/AudioInsights/Source/AudioInsights/Private/Providers/OutputMeterTraceProvider.cpp
new file mode 100644
--- /dev/null
+++ b/Engine/Plugins/AudioInsights/Source/AudioInsights/Private/Providers/OutputMeterTraceProvider.cpp
@@ -0,0 +1,219 @@
+// Copyright Epic Games, Inc. All Rights Reserved.
+#include "Providers/OutputMeterTraceProvider.h"
+
+#include "Async/ParallelFor.h"
+#include "AudioInsightsModule.h"
+#include "Cache/AudioInsightsCacheManager.h"
+#include "Trace/Analyzer.h"
+#include "TraceServices/Model/AnalysisSession.h"
+
+namespace UE::Audio::Insights
+{
+	FName FOutputMeterTraceProvider::GetName_Static()
+	{
+		return "OutputMeterProvider";
+	}
+
+	UE::Trace::IAnalyzer* FOutputMeterTraceProvider::ConstructAnalyzer(TraceServices::IAnalysisSession& InSession)
+	{
+		class FOutputMeterTraceAnalyzer : public FTraceAnalyzerBase
+		{
+		public:
+			FOutputMeterTraceAnalyzer(TSharedRef<FOutputMeterTraceProvider> InProvider, TraceServices::IAnalysisSession& InSession)
+				: FTraceAnalyzerBase(InProvider)
+				, Session(InSession)
+			{
+			}
+
+			virtual void OnAnalysisBegin(const FOnAnalysisContext& Context) override
+			{
+				FTraceAnalyzerBase::OnAnalysisBegin(Context);
+
+				UE::Trace::IAnalyzer::FInterfaceBuilder& Builder = Context.InterfaceBuilder;
+				
+				Builder.RouteEvent(RouteId_MainSubmixLoaded, "Audio", "MainSubmixLoaded");
+				Builder.RouteEvent(RouteId_TruePeakMaxValue, "Audio", "MainSubmixTruePeakMaxValue");
+			}
+
+			virtual bool OnHandleEvent(uint16 RouteId, EStyle Style, const FOnEventContext& Context) override
+			{
+				LLM_SCOPE_BYNAME(TEXT("Insights/FOutputMeterTraceAnalyzer"));
+
+				FOutputMeterMessages& Messages = GetProvider<FOutputMeterTraceProvider>().TraceMessages;
+				
+				switch (RouteId)
+				{
+					case RouteId_MainSubmixLoaded:
+					{
+						CacheMessage<FOutputMeterMainSubmixLoadedMessage>(Context, Messages.MainSubmixLoadedMessages);
+						break;
+					}
+
+					case RouteId_TruePeakMaxValue:
+					{
+						CacheMessage<FOutputMeterTruePeakMessage>(Context, Messages.TruePeakMessages);
+						break;
+					}
+
+					default:
+					{
+						return OnEventFailure(RouteId, Style, Context);
+					}
+				}
+
+				const double Timestamp = Context.EventTime.AsSeconds(Context.EventData.GetValue<uint64>("Timestamp"));
+
+				{
+					TraceServices::FAnalysisSessionEditScope SessionEditScope(Session);
+					Session.UpdateDurationSeconds(Timestamp);
+				}
+
+				return OnEventSuccess(RouteId, Style, Context);
+			}
+
+		private:
+			enum : uint16
+			{
+				RouteId_MainSubmixLoaded,
+				RouteId_TruePeakMaxValue
+			};
+
+			TraceServices::IAnalysisSession& Session;
+		};
+
+		return new FOutputMeterTraceAnalyzer(AsShared(), InSession);
+	}
+
+	bool FOutputMeterTraceProvider::ProcessMessages()
+	{
+		auto CreateEntryFunc = [this](const FOutputMeterMessageBase& Msg)
+		{
+			TSharedPtr<FOutputMeterDashboardEntry>* ToReturn = nullptr;
+			UpdateDeviceEntry(Msg.DeviceId, Msg.SubmixId, [&ToReturn, &Msg](TSharedPtr<FOutputMeterDashboardEntry>& Entry)
+			{
+				if (!Entry.IsValid())
+				{
+					Entry = MakeShared<FOutputMeterDashboardEntry>();
+					Entry->DeviceId  = Msg.DeviceId;
+					Entry->SubmixId = Msg.SubmixId;
+				}
+				Entry->Timestamp = Msg.Timestamp;
+
+				ToReturn = &Entry;
+			});
+
+			return ToReturn;
+		};
+
+		auto GetEntryFunc = [this](const FOutputMeterMessageBase& Msg)
+		{
+			return FindDeviceEntry(Msg.DeviceId, Msg.SubmixId);
+		};
+
+		ProcessMessageQueue<FOutputMeterMainSubmixLoadedMessage>(TraceMessages.MainSubmixLoadedMessages, CreateEntryFunc,
+		[this](const FOutputMeterMainSubmixLoadedMessage& Msg, TSharedPtr<FOutputMeterDashboardEntry>* OutEntry)
+		{
+		});
+
+		ProcessMessageQueue<FOutputMeterTruePeakMessage>(TraceMessages.TruePeakMessages, GetEntryFunc,
+		[this](const FOutputMeterTruePeakMessage& Msg, TSharedPtr<FOutputMeterDashboardEntry>* OutEntry)
+		{
+			if (OutEntry)
+			{
+				FOutputMeterDashboardEntry& EntryRef = *OutEntry->Get();
+				EntryRef.TruePeakMaxValueDb = Msg.TruePeakMaxValueDb;
+			}
+		});
+
+		return true;
+	}
+
+	void FOutputMeterTraceProvider::OnTimingViewTimeMarkerChanged(double TimeMarker)
+	{
+		// Helper lambdas
+		auto ProcessOutputMeterMessage = [this](const FOutputMeterMainSubmixLoadedMessage& OutputMeterCachedMessage)
+		{
+			UpdateDeviceEntry(OutputMeterCachedMessage.DeviceId, OutputMeterCachedMessage.SubmixId, [&OutputMeterCachedMessage](TSharedPtr<FOutputMeterDashboardEntry>& Entry)
+			{
+				if (!Entry.IsValid())
+				{
+					Entry = MakeShared<FOutputMeterDashboardEntry>();
+					Entry->DeviceId  = OutputMeterCachedMessage.DeviceId;
+					Entry->SubmixId = OutputMeterCachedMessage.SubmixId;
+				}
+				
+				Entry->Timestamp = OutputMeterCachedMessage.Timestamp;
+			});
+		};
+
+		// Process cached messages
+		DeviceDataMap.Empty();
+
+		const FAudioInsightsCacheManager& CacheManager = FAudioInsightsModule::GetChecked().GetCacheManager();
+
+		// Main submix loaded cached messages
+		CacheManager.IterateTo<FOutputMeterMainSubmixLoadedMessage>(OutputMeterMessageNames::MainSubmixLoaded, TimeMarker, [&ProcessOutputMeterMessage](const FOutputMeterMainSubmixLoadedMessage& OutputMeterMainSubmixLoadedMessage)
+		{
+			ProcessOutputMeterMessage(OutputMeterMainSubmixLoadedMessage);
+		});
+
+		CollectParamsForTimestamp(TimeMarker);
+
+		// Call parent method to update LastMessageId
+		FTraceProviderBase::OnTimingViewTimeMarkerChanged(TimeMarker);
+	}
+
+	void FOutputMeterTraceProvider::CollectParamsForTimestamp(const double InTimeMarker)
+	{
+		const FDeviceData* DeviceData = FindFilteredDeviceData();
+		if (DeviceData == nullptr)
+		{
+			return;
+		}
+
+		// Collect update messages from output meter (based on submix id)
+		struct CachedEntryInfo
+		{
+			TOptional<FOutputMeterTruePeakMessage> TruePeakMessage;
+		};
+
+		TArray<uint32> MainSubmixIdArray;
+		DeviceData->GenerateKeyArray(MainSubmixIdArray);
+
+		TArray<CachedEntryInfo> CachedEntryInfos;
+		CachedEntryInfos.SetNum(MainSubmixIdArray.Num());
+
+		const FAudioInsightsCacheManager& CacheManager = FAudioInsightsModule::GetChecked().GetCacheManager();
+
+		ParallelFor(MainSubmixIdArray.Num(), [&MainSubmixIdArray, &CachedEntryInfos, &CacheManager, InTimeMarker, this](const int32 Index)
+		{
+			const uint32 MainSubmixId = MainSubmixIdArray[Index];
+
+			// TruePeak
+			const FOutputMeterTruePeakMessage* FoundTruePeakCachedMessage = CacheManager.FindClosestMessage<FOutputMeterTruePeakMessage>(OutputMeterMessageNames::TruePeak, InTimeMarker, MainSubmixId);
+			if (FoundTruePeakCachedMessage)
+			{
+				CachedEntryInfos[Index].TruePeakMessage = *FoundTruePeakCachedMessage;
+			}
+		});
+
+		// Update the device entries with the collected info
+		for (const CachedEntryInfo& CachedEntryInfo : CachedEntryInfos)
+		{
+			// TruePeak
+			if (CachedEntryInfo.TruePeakMessage.IsSet())
+			{
+				const FOutputMeterTruePeakMessage& OutputMeterTruePeakMessage = CachedEntryInfo.TruePeakMessage.GetValue();
+
+				UpdateDeviceEntry(OutputMeterTruePeakMessage.DeviceId, OutputMeterTruePeakMessage.SubmixId, [&OutputMeterTruePeakMessage](TSharedPtr<FOutputMeterDashboardEntry>& Entry)
+				{
+					if (ensure(Entry.IsValid()))
+					{
+						Entry->Timestamp = OutputMeterTruePeakMessage.Timestamp;
+						Entry->TruePeakMaxValueDb = OutputMeterTruePeakMessage.TruePeakMaxValueDb;
+					}
+				});
+			}
+		}
+	}
+} // namespace UE::Audio::Insights
diff --git a/Engine/Plugins/AudioInsights/Source/AudioInsights/Private/Providers/OutputMeterTraceProvider.h b/Engine/Plugins/AudioInsights/Source/AudioInsights/Private/Providers/OutputMeterTraceProvider.h
new file mode 100644
--- /dev/null
+++ b/Engine/Plugins/AudioInsights/Source/AudioInsights/Private/Providers/OutputMeterTraceProvider.h
@@ -0,0 +1,32 @@
+// Copyright Epic Games, Inc. All Rights Reserved.
+#pragma once
+
+#include "AudioInsightsTraceProviderBase.h"
+#include "Messages/OutputMeterTraceMessages.h"
+
+namespace UE::Audio::Insights
+{
+	class FOutputMeterTraceProvider
+		: public TDeviceDataMapTraceProvider<uint32, TSharedPtr<FOutputMeterDashboardEntry>>
+		, public TSharedFromThis<FOutputMeterTraceProvider>
+	{
+	public:
+		FOutputMeterTraceProvider()
+			: TDeviceDataMapTraceProvider<uint32, TSharedPtr<FOutputMeterDashboardEntry>>(GetName_Static())
+		{
+		}
+
+		virtual ~FOutputMeterTraceProvider() = default;
+		
+		virtual UE::Trace::IAnalyzer* ConstructAnalyzer(TraceServices::IAnalysisSession& InSession) override;
+
+		static FName GetName_Static();
+
+	private:
+		virtual bool ProcessMessages() override;
+		virtual void OnTimingViewTimeMarkerChanged(double TimeMarker) override;
+		void CollectParamsForTimestamp(const double InTimeMarker);
+
+		FOutputMeterMessages TraceMessages;
+	};
+} // namespace UE::Audio::Insights
diff --git a/Engine/Plugins/AudioInsights/Source/AudioInsights/Private/Views/OutputMeterDashboardViewFactory.cpp b/Engine/Plugins/AudioInsights/Source/AudioInsights/Private/Views/OutputMeterDashboardViewFactory.cpp
--- a/Engine/Plugins/AudioInsights/Source/AudioInsights/Private/Views/OutputMeterDashboardViewFactory.cpp
+++ b/Engine/Plugins/AudioInsights/Source/AudioInsights/Private/Views/OutputMeterDashboardViewFactory.cpp
@@ -3,9 +3,9 @@
 
 #include "Algo/Find.h"
 #include "AudioDefines.h"
+#include "AudioInsightsModule.h"
 #include "AudioInsightsStyle.h"
 #include "DSP/Dsp.h"
-#include "ISessionServicesModule.h"
 #include "SAudioMeterWidget.h"
 #include "Widgets/Colors/SColorBlock.h"
 #include "Widgets/SOverlay.h"
@@ -14,17 +14,29 @@
 
 namespace UE::Audio::Insights
 {
-	FOutputMeterDashboardViewFactory::FOutputMeterDashboardViewFactory(TSharedRef<FSubmixDashboardViewFactory> SubmixDashboard)
-		: SubmixProvider(SubmixDashboard->FindProvider<FSubmixTraceProvider>())
+	namespace FOutputMeterDashboardViewFactoryPrivate
 	{
-	}
+		const FOutputMeterDashboardEntry& CastEntry(const IDashboardDataViewEntry& InData)
+		{
+			return static_cast<const FOutputMeterDashboardEntry&>(InData);
+		};
+	} // namespace FOutputMeterDashboardViewFactoryPrivate
 
-	FOutputMeterDashboardViewFactory::~FOutputMeterDashboardViewFactory()
+
+	FOutputMeterDashboardViewFactory::FOutputMeterDashboardViewFactory(TSharedRef<FSubmixDashboardViewFactory> SubmixDashboard)
+		: FTraceDashboardViewFactoryBase()
+		, SubmixProvider(SubmixDashboard->FindProvider<FSubmixTraceProvider>())
 	{
-		if (TickerHandle.IsValid())
+		FTraceModule& AudioInsightsTraceModule = static_cast<FTraceModule&>(FAudioInsightsModule::GetChecked().GetTraceModule());
+
+		OutputMeterProvider = MakeShared<FOutputMeterTraceProvider>();
+
+		AudioInsightsTraceModule.AddTraceProvider(OutputMeterProvider);
+
+		Providers = TArray<TSharedPtr<FTraceProviderBase>>
 		{
-			FTSTicker::RemoveTicker(TickerHandle);
-		}
+			OutputMeterProvider
+		};
 	}
 
 	EDefaultDashboardTabStack FOutputMeterDashboardViewFactory::GetDefaultTabStack() const
@@ -47,65 +59,205 @@ namespace UE::Audio::Insights
 		return FSlateStyle::Get().CreateIcon("AudioInsights.Icon");
 	}
 
-	TSharedRef<SWidget> FOutputMeterDashboardViewFactory::MakeWidget(TSharedRef<SDockTab> OwnerTab, const FSpawnTabArgs& SpawnTabArgs)
+	void FOutputMeterDashboardViewFactory::InitLoudnessMeterWidgetView()
 	{
-		if (!TickerHandle.IsValid())
+		using namespace AudioWidgetsCore;
+
+		// ValuesDisplayOrderParams
+		FLoudnessMeterWidgetView::FDisplayOrderParams ValuesDisplayOrderParams;
+
+		ValuesDisplayOrderParams.PermutationIndex.BindLambda([this]() { return LoudnessMeterSettings.ValuesOrderingPermutation; });
+		ValuesDisplayOrderParams.OnPermutationChanged.BindLambda([this](int32 NewPermutationIndex)
+		{
+			LoudnessMeterSettings.ValuesOrderingPermutation = NewPermutationIndex;
+		});
+
+		// MetersDisplayOrderParams
+		FLoudnessMeterWidgetView::FDisplayOrderParams MetersDisplayOrderParams;
+
+		MetersDisplayOrderParams.PermutationIndex.BindLambda([this]() { return LoudnessMeterSettings.MetersOrderingPermutation; });
+		MetersDisplayOrderParams.OnPermutationChanged.BindLambda([this](int32 NewPermutationIndex)
+		{
+			LoudnessMeterSettings.MetersOrderingPermutation = NewPermutationIndex;
+		});
+
+		// LoudnessMetrics
+		LoudnessMeterSettings.TruePeak.bShowValue = false;
+		LoudnessMeterSettings.TruePeak.bShowMeter = true;
+
+		FLoudnessMeterWidgetView::FLoudnessMetric LoudnessMetrics[] =
+		{
+			{
+				.Name = GET_MEMBER_NAME_CHECKED(FLoudnessMeterSettings, TruePeak),
+				.DisplayName = LOCTEXT("TruePeakDisplayName", "True Peak"),
+				.MeterMetric = EAudioMeterMetric::Decibels,
+				.Value = TAttribute<TOptional<float>>::CreateSPLambda(this, [this]()
+				{
+					return TruePeakMaxValueDb;
+				}),
+				.OnClicked = FSimpleDelegate::CreateSPLambda(this, [this]() { TruePeakMaxValueDb = MIN_VOLUME_DECIBELS; })
+			}
+		};
+
+		for (AudioWidgetsCore::FLoudnessMeterWidgetView::FLoudnessMetric& LoudnessMetric : LoudnessMetrics)
 		{
-			TickerHandle = FTSTicker::GetCoreTicker().AddTicker(TEXT("FOutputMeterDashboardViewFactory"), 0.0f, [this](float DeltaTime)
+			if (const FProperty* DisplayOptionsProperty = FLoudnessMeterSettings::StaticStruct()->FindPropertyByName(LoudnessMetric.Name))
+			{
+				// Show value
+				LoudnessMetric.bShowValue.BindLambda([this, DisplayOptionsProperty]()
+				{
+					const FLoudnessMeterDisplayOptions* DisplayOptions = DisplayOptionsProperty->ContainerPtrToValuePtr<FLoudnessMeterDisplayOptions>(&LoudnessMeterSettings);
+					return DisplayOptions->bShowValue;
+				});
+
+				LoudnessMetric.OnShowValueToggleRequested.BindLambda([this, DisplayOptionsProperty]()
+				{
+					FLoudnessMeterDisplayOptions* DisplayOptions = DisplayOptionsProperty->ContainerPtrToValuePtr<FLoudnessMeterDisplayOptions>(&LoudnessMeterSettings);
+					DisplayOptions->bShowValue = (!DisplayOptions->bShowValue) ? true : false;
+				});
+
+				// Show Meter
+				LoudnessMetric.bShowMeter.BindLambda([this, DisplayOptionsProperty]()
 				{
-					Tick(DeltaTime);
-					return true;
+					const FLoudnessMeterDisplayOptions* DisplayOptions = DisplayOptionsProperty->ContainerPtrToValuePtr<FLoudnessMeterDisplayOptions>(&LoudnessMeterSettings);
+					return DisplayOptions->bShowMeter;
 				});
+
+				LoudnessMetric.OnShowMeterToggleRequested.BindLambda([this, DisplayOptionsProperty]()
+				{
+					FLoudnessMeterDisplayOptions* DisplayOptions = DisplayOptionsProperty->ContainerPtrToValuePtr<FLoudnessMeterDisplayOptions>(&LoudnessMeterSettings);
+					DisplayOptions->bShowMeter = (!DisplayOptions->bShowMeter) ? true : false;
+				});
+
+				// Color
+				LoudnessMetric.Color.BindLambda([this, DisplayOptionsProperty]()
+				{
+					const FLoudnessMeterDisplayOptions* DisplayOptions = DisplayOptionsProperty->ContainerPtrToValuePtr<FLoudnessMeterDisplayOptions>(&LoudnessMeterSettings);
+					return DisplayOptions->Color;
+				});
+
+				LoudnessMetric.OnColorChanged.BindLambda([this, DisplayOptionsProperty](FLinearColor NewColor)
+				{
+					FLoudnessMeterDisplayOptions* DisplayOptions = DisplayOptionsProperty->ContainerPtrToValuePtr<FLoudnessMeterDisplayOptions>(&LoudnessMeterSettings);
+					DisplayOptions->Color = NewColor;
+				});
+
+				LoudnessMetric.OnResetColorRequested.BindLambda([this, DisplayOptionsProperty]()
+				{
+					FLoudnessMeterDisplayOptions* DisplayOptions = DisplayOptionsProperty->ContainerPtrToValuePtr<FLoudnessMeterDisplayOptions>(&LoudnessMeterSettings);
+					DisplayOptions->Color = FLoudnessMeterWidgetView::GetDefaultMeterColor();
+				});
+			}
+		}
+
+		LoudnessMeterWidgetView.InitValuesDisplayOrder(ValuesDisplayOrderParams);
+		LoudnessMeterWidgetView.InitMetersDisplayOrder(MetersDisplayOrderParams);
+
+		for (const FLoudnessMeterWidgetView::FLoudnessMetric& LoudnessMetric : LoudnessMetrics)
+		{
+			LoudnessMeterWidgetView.AddLoudnessMetric(LoudnessMetric);
+		}
+	}
+
+	TSharedRef<SWidget> FOutputMeterDashboardViewFactory::MakeWidget(TSharedRef<SDockTab> OwnerTab, const FSpawnTabArgs& SpawnTabArgs)
+	{
+		if (!bIsLoudnessMeterInitialized)
+		{
+			InitLoudnessMeterWidgetView();
+			bIsLoudnessMeterInitialized = true;
 		}
 
 		return SNew(SOverlay)
+			.Clipping(EWidgetClipping::ClipToBounds)
 			+ SOverlay::Slot()
 			[
 				SNew(SColorBlock)
-					.Color(FSlateStyle::Get().GetColor("AudioInsights.Analyzers.BackgroundColor"))
+				.Color(FSlateStyle::Get().GetColor("AudioInsights.Analyzers.BackgroundColor"))
 			]
 			+ SOverlay::Slot()
 			[
-				SNew(SVerticalBox)
+				SNew(SSplitter)
+				.Orientation(Orient_Vertical)
+				.ResizeMode(ESplitterResizeMode::Fill)
+				.Clipping(EWidgetClipping::ClipToBounds)
+				.PhysicalSplitterHandleSize(4.0f)
+				.HitDetectionSplitterHandleSize(6.0f)
+				+ SSplitter::Slot()
+				.Value(0.5f)
+				[
+					LoudnessMeterWidgetView.MakeWidget()
+				]
+				+ SSplitter::Slot()
+				.Value(0.5f)
+				[
+					SNew(SVerticalBox)
 					+ SVerticalBox::Slot()
 					.FillHeight(1.0f)
 					.HAlign(HAlign_Center)
 					.VAlign(VAlign_Fill)
 					[
 						SNew(SAudioMeterWidget)
-							.Orientation(EOrientation::Orient_Vertical)
-							.BackgroundColor(FLinearColor::Transparent)
-							.MeterBackgroundColor(FAudioMeterDefaultColorWidgetStyle::GetDefault().MeterBackgroundColor)
-							.MeterValueColor(FAudioMeterDefaultColorWidgetStyle::GetDefault().MeterValueColor)
-							.MeterPeakColor(FAudioMeterDefaultColorWidgetStyle::GetDefault().MeterPeakColor)
-							.MeterClippingColor(FAudioMeterDefaultColorWidgetStyle::GetDefault().MeterClippingColor)
-							.MeterScaleColor(FAudioMeterDefaultColorWidgetStyle::GetDefault().MeterScaleColor)
-							.MeterScaleLabelColor(FAudioMeterDefaultColorWidgetStyle::GetDefault().MeterScaleLabelColor)
-							.MeterChannelInfo(TAttribute<TArray<FAudioMeterChannelInfo>>::CreateSP(this, &FOutputMeterDashboardViewFactory::GetAudioMeterChannelInfo))
+						.Orientation(EOrientation::Orient_Vertical)
+						.BackgroundColor(FLinearColor::Transparent)
+						.MeterBackgroundColor(FAudioMeterDefaultColorWidgetStyle::GetDefault().MeterBackgroundColor)
+						.MeterValueColor(FAudioMeterDefaultColorWidgetStyle::GetDefault().MeterValueColor)
+						.MeterPeakColor(FAudioMeterDefaultColorWidgetStyle::GetDefault().MeterPeakColor)
+						.MeterClippingColor(FAudioMeterDefaultColorWidgetStyle::GetDefault().MeterClippingColor)
+						.MeterScaleColor(FAudioMeterDefaultColorWidgetStyle::GetDefault().MeterScaleColor)
+						.MeterScaleLabelColor(FAudioMeterDefaultColorWidgetStyle::GetDefault().MeterScaleLabelColor)
+						.MeterChannelInfo(TAttribute<TArray<FAudioMeterChannelInfo>>::CreateSP(this, &FOutputMeterDashboardViewFactory::GetAudioMeterChannelInfo))
 					]
+				]
 			];
 	}
 
-	void FOutputMeterDashboardViewFactory::Tick(float InElapsed)
+	void FOutputMeterDashboardViewFactory::UpdateDataViewEntries()
 	{
-		// Try and subscribe to OnSessionInstanceUpdated events if not already:
-		if (!SessionInstanceUpdatedDelegateHandle.IsValid())
+		const TSharedPtr<const FOutputMeterTraceProvider> Provider = FindProvider<const FOutputMeterTraceProvider>();
+		if (Provider.IsValid())
 		{
-			ISessionServicesModule& SessionServicesModule = FModuleManager::LoadModuleChecked<ISessionServicesModule>("SessionServices");
-			if (TSharedPtr<ISessionManager> SessionManager = SessionServicesModule.GetSessionManager())
+			if (const FOutputMeterTraceProvider::FDeviceData* DeviceData = Provider->FindFilteredDeviceData())
 			{
-				SessionInstanceUpdatedDelegateHandle = SessionManager->OnSessionInstanceUpdated().AddSP(this, &FOutputMeterDashboardViewFactory::HandleSessionInstanceUpdated);
+				DataViewEntries.Reset();
+
+				auto TransformEntry = [](const typename FOutputMeterTraceProvider::FEntryPair& Pair)
+				{
+					return StaticCastSharedPtr<IDashboardDataViewEntry>(Pair.Value);
+				};
+
+				Algo::Transform(*DeviceData, DataViewEntries, TransformEntry);
+			}
+			else
+			{
+				if (!DataViewEntries.IsEmpty())
+				{
+					DataViewEntries.Empty();
+				}
 			}
 		}
+	}
 
+	void FOutputMeterDashboardViewFactory::UpdateAudioMeterInfo()
+	{
 		if (SubmixProvider.IsValid())
 		{
 			const uint64 LastUpdateId = SubmixProvider->GetLastUpdateId();
+
 			if (ProcessedUpdateId != LastUpdateId)
 			{
 				if (const FSubmixTraceProvider::FDeviceData* DeviceData = SubmixProvider->FindFilteredDeviceData())
 				{
-					ProcessDeviceData(*DeviceData);
+					// Try and find the FSubmixDashboardEntry for the main submix with updated audio meter values
+					const auto IsMainSubmixEntryPair = [](const FSubmixTraceProvider::FEntryPair& EntryPair) { return EntryPair.Value->IsMainSubmix(); };
+
+					const FSubmixTraceProvider::FEntryPair* MainSubmixEntryPair = Algo::FindByPredicate(*DeviceData, IsMainSubmixEntryPair);
+
+					TSharedPtr<FSubmixDashboardEntry> MainSubmixDashboardEntry = MainSubmixEntryPair ? MainSubmixEntryPair->Value : nullptr;
+
+					if (MainSubmixDashboardEntry.IsValid())
+					{
+						AudioMeterInfo = MainSubmixDashboardEntry->AudioMeterInfo.Get();
+					}
 				}
 
 				ProcessedUpdateId = LastUpdateId;
@@ -113,58 +265,43 @@ namespace UE::Audio::Insights
 		}
 	}
 
-	void FOutputMeterDashboardViewFactory::HandleSessionInstanceUpdated()
+	void FOutputMeterDashboardViewFactory::UpdateLoudnessValues()
 	{
-		if (SubmixDashboardEntry.IsValid())
+		if (!DataViewEntries.IsEmpty())
 		{
-			// As we may have just gotten a valid session, we request starting of the envelope follower:
-			FSubmixDashboardViewFactory::SendSubmixEnvelopeFollowerCVar(true, *SubmixDashboardEntry);
+			const FOutputMeterDashboardEntry& OutputMeterEntry = FOutputMeterDashboardViewFactoryPrivate::CastEntry(*DataViewEntries[0].Get());
+			TruePeakMaxValueDb = FMath::Max<float>(TruePeakMaxValueDb, OutputMeterEntry.TruePeakMaxValueDb);
 		}
 	}
 
-	void FOutputMeterDashboardViewFactory::ProcessDeviceData(const FSubmixTraceProvider::FDeviceData& DeviceData)
+	void FOutputMeterDashboardViewFactory::ProcessEntries(FTraceDashboardViewFactoryBase::EProcessReason Reason)
 	{
-		// Try and find the FSubmixDashboardEntry for the main submix:
-		const auto IsMainSubmixEntryPair = [](const FSubmixTraceProvider::FEntryPair& EntryPair) { return EntryPair.Value->IsMainSubmix(); };
-		const FSubmixTraceProvider::FEntryPair* MainSubmixEntryPair = Algo::FindByPredicate(DeviceData, IsMainSubmixEntryPair);
-		TSharedPtr<FSubmixDashboardEntry> MainSubmixDashboardEntry = (MainSubmixEntryPair) ? MainSubmixEntryPair->Value : nullptr;
+		UpdateDataViewEntries();
 
-		if (SubmixDashboardEntry != MainSubmixDashboardEntry)
-		{
-			SubmixDashboardEntry = MainSubmixDashboardEntry;
+		UpdateAudioMeterInfo();
+		UpdateLoudnessValues();
 
-			if (SubmixDashboardEntry.IsValid())
-			{
-				// Request starting of envelope follower:
-				FSubmixDashboardViewFactory::SendSubmixEnvelopeFollowerCVar(true, *SubmixDashboardEntry);
-			}
-		}
+		LoudnessMeterWidgetView.RefreshVisibleLoudnessMetrics();
 	}
 
 	TArray<FAudioMeterChannelInfo> FOutputMeterDashboardViewFactory::GetAudioMeterChannelInfo() const
 	{
 		TArray<FAudioMeterChannelInfo> MeterChannelInfo;
 
-		if (SubmixDashboardEntry.IsValid())
+		if (!AudioMeterInfo.EnvelopeValues.IsEmpty())
 		{
-			const FAudioMeterInfo& AudioMeterInfo = *SubmixDashboardEntry->AudioMeterInfo;
-			const TArray<float>& EnvelopeValues = AudioMeterInfo.EnvelopeValues;
-
-			if (!EnvelopeValues.IsEmpty())
+			MeterChannelInfo.Reserve(AudioMeterInfo.EnvelopeValues.Num());
+			for (const float EnvelopeValue : AudioMeterInfo.EnvelopeValues)
 			{
-				MeterChannelInfo.Reserve(EnvelopeValues.Num());
-				for (const float EnvelopeValue : EnvelopeValues)
-				{
-					MeterChannelInfo.Emplace(::Audio::ConvertToDecibels(EnvelopeValue), MIN_VOLUME_DECIBELS, MIN_VOLUME_DECIBELS);
-				}
+				MeterChannelInfo.Emplace(::Audio::ConvertToDecibels(EnvelopeValue), MIN_VOLUME_DECIBELS, MIN_VOLUME_DECIBELS);
 			}
-			else
+		}
+		else
+		{
+			MeterChannelInfo.Reserve(AudioMeterInfo.NumChannels);
+			while (MeterChannelInfo.Num() < AudioMeterInfo.NumChannels)
 			{
-				MeterChannelInfo.Reserve(AudioMeterInfo.NumChannels);
-				while (MeterChannelInfo.Num() < AudioMeterInfo.NumChannels)
-				{
-					MeterChannelInfo.Emplace(MIN_VOLUME_DECIBELS, MIN_VOLUME_DECIBELS, MIN_VOLUME_DECIBELS);
-				}
+				MeterChannelInfo.Emplace(MIN_VOLUME_DECIBELS, MIN_VOLUME_DECIBELS, MIN_VOLUME_DECIBELS);
 			}
 		}
 
diff --git a/Engine/Plugins/AudioInsights/Source/AudioInsights/Private/Views/OutputMeterDashboardViewFactory.h b/Engine/Plugins/AudioInsights/Source/AudioInsights/Private/Views/OutputMeterDashboardViewFactory.h
--- a/Engine/Plugins/AudioInsights/Source/AudioInsights/Private/Views/OutputMeterDashboardViewFactory.h
+++ b/Engine/Plugins/AudioInsights/Source/AudioInsights/Private/Views/OutputMeterDashboardViewFactory.h
@@ -2,17 +2,19 @@
 #pragma once
 
 #include "AudioMeterChannelInfo.h"
-#include "Containers/Ticker.h"
+#include "LoudnessMeterSettings.h"
+#include "LoudnessMeterWidgetView.h"
+#include "Providers/OutputMeterTraceProvider.h"
 #include "Providers/SubmixTraceProvider.h"
 #include "Views/SubmixDashboardViewFactory.h"
 
 namespace UE::Audio::Insights
 {
-	class FOutputMeterDashboardViewFactory : public IDashboardViewFactory, public TSharedFromThis<FOutputMeterDashboardViewFactory>
+	class FOutputMeterDashboardViewFactory : public FTraceDashboardViewFactoryBase, public TSharedFromThis<FOutputMeterDashboardViewFactory>
 	{
 	public:
 		FOutputMeterDashboardViewFactory(TSharedRef<FSubmixDashboardViewFactory> SubmixDashboard);
-		virtual ~FOutputMeterDashboardViewFactory();
+		virtual ~FOutputMeterDashboardViewFactory() = default;
 
 		virtual EDefaultDashboardTabStack GetDefaultTabStack() const override;
 		virtual FText GetDisplayName() const override;
@@ -21,16 +23,29 @@ namespace UE::Audio::Insights
 		virtual TSharedRef<SWidget> MakeWidget(TSharedRef<SDockTab> OwnerTab, const FSpawnTabArgs& SpawnTabArgs) override;
 
 	private:
-		void Tick(float InElapsed);
-		void HandleSessionInstanceUpdated();
-		void ProcessDeviceData(const FSubmixTraceProvider::FDeviceData& DeviceData);
+		void InitLoudnessMeterWidgetView();
+
+		void UpdateDataViewEntries();
+		void UpdateAudioMeterInfo();
+		void UpdateLoudnessValues();
+
+		virtual void ProcessEntries(FTraceDashboardViewFactoryBase::EProcessReason Reason) override;
+
 		TArray<FAudioMeterChannelInfo> GetAudioMeterChannelInfo() const;
 
+		TSharedPtr<FOutputMeterTraceProvider> OutputMeterProvider;
+		TArray<TSharedPtr<IDashboardDataViewEntry>> DataViewEntries;
+
+		// Audio meter related members
 		TSharedPtr<const FSubmixTraceProvider> SubmixProvider;
-		FTSTicker::FDelegateHandle TickerHandle;
-		FDelegateHandle SessionInstanceUpdatedDelegateHandle;
 		TOptional<uint64> ProcessedUpdateId;
-		TSharedPtr<FSubmixDashboardEntry> SubmixDashboardEntry;
+		FAudioMeterInfo AudioMeterInfo;
+
+		// Loudness meters related members
+		FLoudnessMeterSettings LoudnessMeterSettings;
+		AudioWidgetsCore::FLoudnessMeterWidgetView LoudnessMeterWidgetView;
+		float TruePeakMaxValueDb = MIN_VOLUME_DECIBELS;
+		bool bIsLoudnessMeterInitialized = false;
 	};
 } // namespace UE::Audio::Insights
 
diff --git a/Engine/Plugins/AudioInsights/Source/AudioInsights/Public/Messages/OutputMeterTraceMessages.h b/Engine/Plugins/AudioInsights/Source/AudioInsights/Public/Messages/OutputMeterTraceMessages.h
new file mode 100644
--- /dev/null
+++ b/Engine/Plugins/AudioInsights/Source/AudioInsights/Public/Messages/OutputMeterTraceMessages.h
@@ -0,0 +1,73 @@
+// Copyright Epic Games, Inc. All Rights Reserved.
+#pragma once
+
+#include "AudioDefines.h"
+#include "AudioInsightsDataSource.h"
+#include "Cache/IAudioCachedMessage.h"
+#include "Messages/AnalyzerMessageQueue.h"
+
+namespace UE::Audio::Insights
+{
+	namespace OutputMeterMessageNames
+	{
+		extern const FName MainSubmixLoaded;
+		extern const FName TruePeak;
+	};
+
+	struct FOutputMeterMessageBase : public IAudioCachedMessage
+	{
+		FOutputMeterMessageBase() = default;
+		FOutputMeterMessageBase(const Trace::IAnalyzer::FOnEventContext& InContext);
+
+		virtual uint64 GetID() const override { return SubmixId; }
+
+		::Audio::FDeviceId DeviceId = INDEX_NONE;
+		uint32 SubmixId = INDEX_NONE;
+	};
+
+	struct FOutputMeterMainSubmixLoadedMessage : public FOutputMeterMessageBase
+	{
+		FOutputMeterMainSubmixLoadedMessage() = default;
+		FOutputMeterMainSubmixLoadedMessage(const Trace::IAnalyzer::FOnEventContext& InContext);
+
+		virtual const FName GetMessageName() const override { return OutputMeterMessageNames::MainSubmixLoaded; }
+		virtual uint32 GetSizeOf() const override;
+	};
+
+	struct FOutputMeterTruePeakMessage : public FOutputMeterMessageBase
+	{
+		FOutputMeterTruePeakMessage() = default;
+		FOutputMeterTruePeakMessage(const Trace::IAnalyzer::FOnEventContext& InContext);
+
+		virtual const FName GetMessageName() const override { return OutputMeterMessageNames::TruePeak; }
+		virtual uint32 GetSizeOf() const override;
+
+		float TruePeakMaxValueDb = MIN_VOLUME_DECIBELS;
+	};
+
+	class FOutputMeterDashboardEntry : public IDashboardDataViewEntry
+	{
+	public:
+		FOutputMeterDashboardEntry() = default;
+		virtual ~FOutputMeterDashboardEntry() = default;
+
+		virtual bool IsValid() const override
+		{
+			return SubmixId != static_cast<uint32>(INDEX_NONE);
+		}
+
+		::Audio::FDeviceId DeviceId = INDEX_NONE;
+		double Timestamp = 0.0;
+		uint32 SubmixId = static_cast<uint32>(INDEX_NONE);
+
+		float TruePeakMaxValueDb = MIN_VOLUME_DECIBELS;
+	};
+
+	class FOutputMeterMessages
+	{
+		TAnalyzerMessageQueue<FOutputMeterMainSubmixLoadedMessage> MainSubmixLoadedMessages;
+		TAnalyzerMessageQueue<FOutputMeterTruePeakMessage> TruePeakMessages;
+
+		friend class FOutputMeterTraceProvider;
+	};
+} // namespace UE::Audio::Insights
