[Audio Insights] - Support TimingViewExtender via Interface access

#Synced-CL 50904638

diff --git a/Engine/Plugins/AudioInsights/Source/AudioInsights/Private/AudioInsightsTimingViewExtender.h b/Engine/Plugins/AudioInsights/Source/AudioInsights/Private/AudioInsightsTimingViewExtender.h
deleted file mode 100644
diff --git a/Engine/Plugins/AudioInsights/Source/AudioInsights/Private/IAudioInsightsModule.cpp b/Engine/Plugins/AudioInsights/Source/AudioInsights/Private/IAudioInsightsModule.cpp
--- a/Engine/Plugins/AudioInsights/Source/AudioInsights/Private/IAudioInsightsModule.cpp
+++ b/Engine/Plugins/AudioInsights/Source/AudioInsights/Private/IAudioInsightsModule.cpp
@@ -15,6 +15,12 @@ UE::Audio::Insights::FAudioInsightsCacheManager& IAudioInsightsModule::GetCacheM
 	return AudioInsightsModule.GetCacheManager();
 }
 
+UE::Audio::Insights::FAudioInsightsTimingViewExtender& IAudioInsightsModule::GetTimingViewExtender()
+{
+	IAudioInsightsModule& AudioInsightsModule = static_cast<IAudioInsightsModule&>(UE::Audio::Insights::FAudioInsightsModule::GetChecked());
+	return AudioInsightsModule.GetTimingViewExtender();
+}
+
 bool IAudioInsightsModule::IsModuleLoaded()
 {
 	return UE::Audio::Insights::FAudioInsightsModule::IsModuleLoaded();
diff --git a/Engine/Plugins/AudioInsights/Source/AudioInsights/Public/AudioInsightsTimingViewExtender.h b/Engine/Plugins/AudioInsights/Source/AudioInsights/Public/AudioInsightsTimingViewExtender.h
new file mode 100644
--- /dev/null
+++ b/Engine/Plugins/AudioInsights/Source/AudioInsights/Public/AudioInsightsTimingViewExtender.h
@@ -0,0 +1,79 @@
+// Copyright Epic Games, Inc. All Rights Reserved.
+#pragma once
+
+#include "Delegates/Delegate.h"
+#include "Insights/ITimingViewExtender.h"
+#include "Math/Range.h"
+#include "Misc/Optional.h"
+
+namespace Insights
+{
+	class ITimingViewSession;
+	enum class ETimeChangedFlags : int32;
+}
+
+class IAnalysisSession;
+
+namespace UE::Audio::Insights
+{
+	// How the cache manage handles new incoming messages
+	// Cache - will save the incoming message inside the cache to be retrievable later
+	// Process - sends the message to the Audio Insights providers to be processed and displayed in the UI
+	enum class ECacheAndProcess
+	{
+		Latest = 0,				// Cache and process all new messages
+		CacheLatestNoProcess,	// Cache new messages but do not process them
+		None					// Do not cache or process any new messages
+	};
+
+	enum class ESystemControllingTimeMarker
+	{
+		EventLog = 0,
+		PlotsWidget,
+		SignalFlow,
+		External
+	};
+
+	class  FAudioInsightsTimingViewExtender : public UE::Insights::Timing::ITimingViewExtender
+	{
+	public:
+		// Insights::ITimingViewExtender interface
+		virtual void OnBeginSession(UE::Insights::Timing::ITimingViewSession& InSession) override;
+		virtual void OnEndSession(UE::Insights::Timing::ITimingViewSession& InSession) override;
+		virtual void Tick(UE::Insights::Timing::ITimingViewSession& InSession, const TraceServices::IAnalysisSession& InAnalysisSession) override;
+
+		ECacheAndProcess GetMessageCacheAndProcessingStatus() const { return CacheAndProcessMessageStatus; }
+		TOptional<ESystemControllingTimeMarker> TryGetSystemControllingTimeMarker() const { return SystemControllingTime; }
+
+		void StopProcessingNewMessages() { CacheAndProcessMessageStatus = ECacheAndProcess::CacheLatestNoProcess; }
+		void StopCachingAndProcessingNewMessages() { CacheAndProcessMessageStatus = ECacheAndProcess::None; }
+		void ResetMessageProcessType();
+
+		double GetCurrentDurationSeconds() const;
+		TRange<double> GetPlottingRange() const { return PlottingRange; }
+
+		void SetTimeMarker(const double Timestamp, const ESystemControllingTimeMarker ControllingSystem, TOptional<TRange<double>> InPlottingRange = TOptional<TRange<double>>());
+
+		DECLARE_MULTICAST_DELEGATE_OneParam(FOnTimingViewTimeMarkerChanged, double /*TimeMarker*/);
+		FOnTimingViewTimeMarkerChanged OnTimingViewTimeMarkerChanged;
+
+		DECLARE_MULTICAST_DELEGATE(FOnTimeControlMethodReset);
+		FOnTimeControlMethodReset OnTimeControlMethodReset;
+
+		// Move to AudioInsightsConstants
+		static constexpr double MaxPlottingHistorySeconds = 5.0;
+		static constexpr double PlottingMarginSeconds = 0.2;
+
+	private:
+		void OnTimeMarkerChanged(UE::Insights::Timing::ETimeChangedFlags InFlags, double InTimeMarker);
+
+		const TraceServices::IAnalysisSession* AnalysisSession = nullptr;
+		UE::Insights::Timing::ITimingViewSession* TimingView = nullptr;
+
+		ECacheAndProcess CacheAndProcessMessageStatus = ECacheAndProcess::Latest;
+		TOptional<ESystemControllingTimeMarker> SystemControllingTime;
+		TRange<double> PlottingRange { 0.0, MaxPlottingHistorySeconds };
+		double TraceDurationSeconds = 0.0;
+		bool bUserInputDetected = false;
+	};
+} // namespace UE::Audio::Insights
diff --git a/Engine/Plugins/AudioInsights/Source/AudioInsights/Public/IAudioInsightsModule.h b/Engine/Plugins/AudioInsights/Source/AudioInsights/Public/IAudioInsightsModule.h
--- a/Engine/Plugins/AudioInsights/Source/AudioInsights/Public/IAudioInsightsModule.h
+++ b/Engine/Plugins/AudioInsights/Source/AudioInsights/Public/IAudioInsightsModule.h
@@ -17,6 +17,7 @@ namespace UE::Audio::Insights
 {
 	class IDashboardViewFactory;
 	class FAudioInsightsCacheManager;
+	class FAudioInsightsTimingViewExtender;
 	class FTraceProviderBase;
 } // namespace UE::Audio::Insights
 
@@ -32,6 +33,8 @@ public:
 
 	UE_API virtual UE::Audio::Insights::FAudioInsightsCacheManager& GetCacheManager() override;
 
+	UE_API virtual UE::Audio::Insights::FAudioInsightsTimingViewExtender& GetTimingViewExtender() override;
+
 	virtual TSharedRef<SDockTab> CreateDashboardTabWidget(const FSpawnTabArgs& Args) = 0;
 
 	UE_API static bool IsModuleLoaded();
diff --git a/Engine/Plugins/AudioInsights/Source/AudioInsights/Public/IAudioInsightsModuleInterface.h b/Engine/Plugins/AudioInsights/Source/AudioInsights/Public/IAudioInsightsModuleInterface.h
--- a/Engine/Plugins/AudioInsights/Source/AudioInsights/Public/IAudioInsightsModuleInterface.h
+++ b/Engine/Plugins/AudioInsights/Source/AudioInsights/Public/IAudioInsightsModuleInterface.h
@@ -2,6 +2,7 @@
 #pragma once
 
 #include "AudioDefines.h"
+#include "AudioInsightsTimingViewExtender.h"
 #include "Cache/AudioInsightsCacheManager.h"
 #include "Modules/ModuleInterface.h"
 #include "Templates/SharedPointer.h"
@@ -25,4 +26,5 @@ public:
 
 	virtual IAudioInsightsTraceModule& GetTraceModule() = 0;
 	virtual UE::Audio::Insights::FAudioInsightsCacheManager& GetCacheManager() = 0;
+	virtual UE::Audio::Insights::FAudioInsightsTimingViewExtender& GetTimingViewExtender() = 0;
 };
diff --git a/Engine/Plugins/AudioInsights/Source/AudioInsightsEditor/Private/AudioInsightsEditorModule.cpp b/Engine/Plugins/AudioInsights/Source/AudioInsightsEditor/Private/AudioInsightsEditorModule.cpp
--- a/Engine/Plugins/AudioInsights/Source/AudioInsightsEditor/Private/AudioInsightsEditorModule.cpp
+++ b/Engine/Plugins/AudioInsights/Source/AudioInsightsEditor/Private/AudioInsightsEditorModule.cpp
@@ -102,6 +102,12 @@ namespace UE::Audio::Insights
 		return InsightsModule.GetCacheManager();
 	}
 
+	FAudioInsightsTimingViewExtender& FAudioInsightsEditorModule::GetTimingViewExtender()
+	{
+		IAudioInsightsModule& InsightsModule = IAudioInsightsModule::GetChecked();
+		return InsightsModule.GetTimingViewExtender();
+	}
+
 	TSharedPtr<FEditorDashboardFactory> FAudioInsightsEditorModule::GetDashboardFactory()
 	{
 		return DashboardFactory;
diff --git a/Engine/Plugins/AudioInsights/Source/AudioInsightsEditor/Private/AudioInsightsEditorModule.h b/Engine/Plugins/AudioInsights/Source/AudioInsightsEditor/Private/AudioInsightsEditorModule.h
--- a/Engine/Plugins/AudioInsights/Source/AudioInsightsEditor/Private/AudioInsightsEditorModule.h
+++ b/Engine/Plugins/AudioInsights/Source/AudioInsightsEditor/Private/AudioInsightsEditorModule.h
@@ -36,6 +36,7 @@ namespace UE::Audio::Insights
 		static FAudioInsightsEditorModule& GetChecked();
 		virtual IAudioInsightsTraceModule& GetTraceModule() override;
 		virtual class FAudioInsightsCacheManager& GetCacheManager() override;
+		virtual class FAudioInsightsTimingViewExtender& GetTimingViewExtender() override;
 
 	private:
 		void RegisterMenus();
