[Audio Insights] - TimingViewExtender: Rename TimeMarker functions to PauseTimerMarker() and ResumeTimeMarker()

#Synced-CL 50931749

diff --git a/Engine/Plugins/AudioInsights/Source/AudioInsights/Private/AudioInsightsTimingViewExtender.cpp b/Engine/Plugins/AudioInsights/Source/AudioInsights/Private/AudioInsightsTimingViewExtender.cpp
--- a/Engine/Plugins/AudioInsights/Source/AudioInsights/Private/AudioInsightsTimingViewExtender.cpp
+++ b/Engine/Plugins/AudioInsights/Source/AudioInsights/Private/AudioInsightsTimingViewExtender.cpp
@@ -35,7 +35,7 @@ namespace UE::Audio::Insights
 		}
 	}
 
-	void FAudioInsightsTimingViewExtender::ResetMessageProcessType()
+	void FAudioInsightsTimingViewExtender::ResumeTimeMarker()
 	{
 		CacheAndProcessMessageStatus = ECacheAndProcess::Latest;
 		SystemControllingTime.Reset();
@@ -48,9 +48,9 @@ namespace UE::Audio::Insights
 		return TraceDurationSeconds;
 	}
 
-	void FAudioInsightsTimingViewExtender::SetTimeMarker(const double Timestamp, const ESystemControllingTimeMarker ControllingSystem, TOptional<TRange<double>> InPlottingRange /*= TOptional<TRange<double>>()*/)
+	void FAudioInsightsTimingViewExtender::PauseTimeMarker(const double InTimestamp, const ESystemControllingTimeMarker InControllingSystem, TOptional<TRange<double>> InPlottingRange /*= TOptional<TRange<double>>()*/)
 	{
-		SystemControllingTime = ControllingSystem;
+		SystemControllingTime = InControllingSystem;
 		bUserInputDetected = true;
 
 		if (InPlottingRange.IsSet())
@@ -59,17 +59,17 @@ namespace UE::Audio::Insights
 		}
 		else
 		{
-			PlottingRange = TRange<double>(Timestamp - MaxPlottingHistorySeconds, Timestamp);
+			PlottingRange = TRange<double>(InTimestamp - MaxPlottingHistorySeconds, InTimestamp);
 		}
 
 		// If you set the time marker on the TimingView, it will Broadcast OnTimingViewTimeMarkerChanged. Otherwise we gotta do it ourselves.
 		if (TimingView != nullptr)
 		{
-			TimingView->SetAndCenterOnTimeMarker(Timestamp);
+			TimingView->SetAndCenterOnTimeMarker(InTimestamp);
 		}
 		else
 		{
-			OnTimingViewTimeMarkerChanged.Broadcast(Timestamp);
+			OnTimingViewTimeMarkerChanged.Broadcast(InTimestamp);
 		}
 	}
 
diff --git a/Engine/Plugins/AudioInsights/Source/AudioInsights/Private/AudioInsightsTraceProviderBase.cpp b/Engine/Plugins/AudioInsights/Source/AudioInsights/Private/AudioInsightsTraceProviderBase.cpp
--- a/Engine/Plugins/AudioInsights/Source/AudioInsights/Private/AudioInsightsTraceProviderBase.cpp
+++ b/Engine/Plugins/AudioInsights/Source/AudioInsights/Private/AudioInsightsTraceProviderBase.cpp
@@ -180,7 +180,7 @@ namespace UE::Audio::Insights
 	{
 		if (FAudioInsightsModule* AudioInsightsModulePtr = FAudioInsightsModule::GetModulePtr())
 		{
-			AudioInsightsModulePtr->GetTimingViewExtender().ResetMessageProcessType();
+			AudioInsightsModulePtr->GetTimingViewExtender().ResumeTimeMarker();
 		}
 	}
 } // namespace UE::Audio::Insights
diff --git a/Engine/Plugins/AudioInsights/Source/AudioInsights/Private/Views/AudioEventLogDashboardViewFactory.cpp b/Engine/Plugins/AudioInsights/Source/AudioInsights/Private/Views/AudioEventLogDashboardViewFactory.cpp
--- a/Engine/Plugins/AudioInsights/Source/AudioInsights/Private/Views/AudioEventLogDashboardViewFactory.cpp
+++ b/Engine/Plugins/AudioInsights/Source/AudioInsights/Private/Views/AudioEventLogDashboardViewFactory.cpp
@@ -252,7 +252,7 @@ namespace UE::Audio::Insights
 		}
 
 		FAudioInsightsModule& AudioInsightsModule = FAudioInsightsModule::GetChecked();
-		AudioInsightsModule.GetTimingViewExtender().ResetMessageProcessType();
+		AudioInsightsModule.GetTimingViewExtender().ResumeTimeMarker();
 	}
 
 	FName FAudioEventLogDashboardViewFactory::GetName() const
@@ -950,7 +950,7 @@ namespace UE::Audio::Insights
 		if (SelectInfo != ESelectInfo::Type::Direct)
 		{
 			FAudioInsightsModule& AudioInsightsModule = FAudioInsightsModule::GetChecked();
-			AudioInsightsModule.GetTimingViewExtender().SetTimeMarker(SelectedEventLogEntry.Timestamp, ESystemControllingTimeMarker::EventLog);
+			AudioInsightsModule.GetTimingViewExtender().PauseTimeMarker(SelectedEventLogEntry.Timestamp, ESystemControllingTimeMarker::EventLog);
 		}
 	}
 
diff --git a/Engine/Plugins/AudioInsights/Source/AudioInsights/Private/Views/SignalFlowDashboardViewFactory.cpp b/Engine/Plugins/AudioInsights/Source/AudioInsights/Private/Views/SignalFlowDashboardViewFactory.cpp
--- a/Engine/Plugins/AudioInsights/Source/AudioInsights/Private/Views/SignalFlowDashboardViewFactory.cpp
+++ b/Engine/Plugins/AudioInsights/Source/AudioInsights/Private/Views/SignalFlowDashboardViewFactory.cpp
@@ -900,7 +900,7 @@ namespace UE::Audio::Insights
 			return;
 		}
 
-		TimingViewExtender.ResetMessageProcessType();
+		TimingViewExtender.ResumeTimeMarker();
 	}
 
 	void FSignalFlowDashboardViewFactory::SetHorizontalJustification(const ESignalFlowHorizontalJustification Justification)
@@ -1237,7 +1237,7 @@ namespace UE::Audio::Insights
 				{
 					FAudioInsightsModule& AudioInsightsModule = FAudioInsightsModule::GetChecked();
 					const double PauseTime = AudioInsightsModule.GetCacheManager().GetCacheEndTimeStamp();
-					AudioInsightsModule.GetTimingViewExtender().SetTimeMarker(PauseTime, ESystemControllingTimeMarker::SignalFlow);
+					AudioInsightsModule.GetTimingViewExtender().PauseTimeMarker(PauseTime, ESystemControllingTimeMarker::SignalFlow);
 				}
 			}
 		}
diff --git a/Engine/Plugins/AudioInsights/Source/AudioInsights/Private/Views/SoundPlotsWidgetView.cpp b/Engine/Plugins/AudioInsights/Source/AudioInsights/Private/Views/SoundPlotsWidgetView.cpp
--- a/Engine/Plugins/AudioInsights/Source/AudioInsights/Private/Views/SoundPlotsWidgetView.cpp
+++ b/Engine/Plugins/AudioInsights/Source/AudioInsights/Private/Views/SoundPlotsWidgetView.cpp
@@ -516,7 +516,7 @@ namespace UE::Audio::Insights
 	void FSoundPlotsWidgetView::ResetInspectTimestamp()
 	{
 		FAudioInsightsModule& AudioInsightsModule = FAudioInsightsModule::GetChecked();
-		AudioInsightsModule.GetTimingViewExtender().ResetMessageProcessType();
+		AudioInsightsModule.GetTimingViewExtender().ResumeTimeMarker();
 	}
 
 	void FSoundPlotsWidgetView::CollectDataRecursive(const TSharedPtr<IDashboardDataTreeViewEntry>& Entry, const TArray<TSharedPtr<IDashboardDataTreeViewEntry>>& SelectedEntries, bool& bOutHasNewMetadata)
@@ -852,7 +852,7 @@ namespace UE::Audio::Insights
 #endif // WITH_EDITOR
 
 		FAudioInsightsModule& AudioInsightsModule = UE::Audio::Insights::FAudioInsightsModule::GetChecked();
-		AudioInsightsModule.GetTimingViewExtender().SetTimeMarker(AbsoluteTimestamp, ESystemControllingTimeMarker::PlotsWidget, AbsolutePlottingRange);
+		AudioInsightsModule.GetTimingViewExtender().PauseTimeMarker(AbsoluteTimestamp, ESystemControllingTimeMarker::PlotsWidget, AbsolutePlottingRange);
 	}
 
 	TRange<double> FSoundPlotsWidgetView::GetViewRange()
diff --git a/Engine/Plugins/AudioInsights/Source/AudioInsights/Public/AudioInsightsTimingViewExtender.h b/Engine/Plugins/AudioInsights/Source/AudioInsights/Public/AudioInsightsTimingViewExtender.h
--- a/Engine/Plugins/AudioInsights/Source/AudioInsights/Public/AudioInsightsTimingViewExtender.h
+++ b/Engine/Plugins/AudioInsights/Source/AudioInsights/Public/AudioInsightsTimingViewExtender.h
@@ -47,12 +47,12 @@ namespace UE::Audio::Insights
 
 		void StopProcessingNewMessages() { CacheAndProcessMessageStatus = ECacheAndProcess::CacheLatestNoProcess; }
 		void StopCachingAndProcessingNewMessages() { CacheAndProcessMessageStatus = ECacheAndProcess::None; }
-		void ResetMessageProcessType();
 
 		double GetCurrentDurationSeconds() const;
 		TRange<double> GetPlottingRange() const { return PlottingRange; }
 
-		void SetTimeMarker(const double Timestamp, const ESystemControllingTimeMarker ControllingSystem, TOptional<TRange<double>> InPlottingRange = TOptional<TRange<double>>());
+		void PauseTimeMarker(const double InTimestamp, const ESystemControllingTimeMarker InControllingSystem, TOptional<TRange<double>> InPlottingRange = TOptional<TRange<double>>());
+		void ResumeTimeMarker();
 
 		DECLARE_MULTICAST_DELEGATE_OneParam(FOnTimingViewTimeMarkerChanged, double /*TimeMarker*/);
 		FOnTimingViewTimeMarkerChanged OnTimingViewTimeMarkerChanged;
