[Audio Insights] - TimingViewExtender: Distinguish between External vs Internal API functions

#Synced-CL 51007057

diff --git a/Engine/Plugins/AudioInsights/Source/AudioInsights/Private/AudioInsightsTimingViewExtender.cpp b/Engine/Plugins/AudioInsights/Source/AudioInsights/Private/AudioInsightsTimingViewExtender.cpp
--- a/Engine/Plugins/AudioInsights/Source/AudioInsights/Private/AudioInsightsTimingViewExtender.cpp
+++ b/Engine/Plugins/AudioInsights/Source/AudioInsights/Private/AudioInsightsTimingViewExtender.cpp
@@ -43,7 +43,22 @@ namespace UE::Audio::Insights
 		OnTimeControlMethodReset.Broadcast();
 	}
 
-	double FAudioInsightsTimingViewExtender::GetCurrentDurationSeconds() const
+	void FAudioInsightsTimingViewExtender::StopProcessingNewMessages()
+	{
+		CacheAndProcessMessageStatus = ECacheAndProcess::CacheLatestNoProcess;
+	}
+
+	void FAudioInsightsTimingViewExtender::StopCachingAndProcessingNewMessages()
+	{
+		CacheAndProcessMessageStatus = ECacheAndProcess::None;
+	}
+
+	TRange<double> FAudioInsightsTimingViewExtender::GetPlottingRange() const 
+	{
+		return PlottingRange;
+	}
+
+	double FAudioInsightsTimingViewExtender::GetCurrentTraceDurationSeconds() const
 	{
 		return TraceDurationSeconds;
 	}
diff --git a/Engine/Plugins/AudioInsights/Source/AudioInsights/Private/Views/SoundPlotsWidgetView.cpp b/Engine/Plugins/AudioInsights/Source/AudioInsights/Private/Views/SoundPlotsWidgetView.cpp
--- a/Engine/Plugins/AudioInsights/Source/AudioInsights/Private/Views/SoundPlotsWidgetView.cpp
+++ b/Engine/Plugins/AudioInsights/Source/AudioInsights/Private/Views/SoundPlotsWidgetView.cpp
@@ -906,7 +906,7 @@ namespace UE::Audio::Insights
 
 			PreviousTime = CurrentTime;
 
-			const double TraceCurrentDurationSeconds = FAudioInsightsModule::GetChecked().GetTimingViewExtender().GetCurrentDurationSeconds();
+			const double TraceCurrentDurationSeconds = FAudioInsightsModule::GetChecked().GetTimingViewExtender().GetCurrentTraceDurationSeconds();
 
 			CurrentRangeUpperBound = FMath::FInterpTo(CurrentRangeUpperBound, TraceCurrentDurationSeconds, DeltaTime, 1.0);
 
diff --git a/Engine/Plugins/AudioInsights/Source/AudioInsights/Public/AudioInsightsTimingViewExtender.h b/Engine/Plugins/AudioInsights/Source/AudioInsights/Public/AudioInsightsTimingViewExtender.h
--- a/Engine/Plugins/AudioInsights/Source/AudioInsights/Public/AudioInsightsTimingViewExtender.h
+++ b/Engine/Plugins/AudioInsights/Source/AudioInsights/Public/AudioInsightsTimingViewExtender.h
@@ -6,6 +6,8 @@
 #include "Math/Range.h"
 #include "Misc/Optional.h"
 
+#define UE_API AUDIOINSIGHTS_API
+
 namespace Insights
 {
 	class ITimingViewSession;
@@ -34,25 +36,18 @@ namespace UE::Audio::Insights
 		External
 	};
 
-	class  FAudioInsightsTimingViewExtender : public UE::Insights::Timing::ITimingViewExtender
+	class FAudioInsightsTimingViewExtender final : private UE::Insights::Timing::ITimingViewExtender
 	{
 	public:
-		// Insights::ITimingViewExtender interface
-		virtual void OnBeginSession(UE::Insights::Timing::ITimingViewSession& InSession) override;
-		virtual void OnEndSession(UE::Insights::Timing::ITimingViewSession& InSession) override;
-		virtual void Tick(UE::Insights::Timing::ITimingViewSession& InSession, const TraceServices::IAnalysisSession& InAnalysisSession) override;
 
+		// AudioInsights API:
 		ECacheAndProcess GetMessageCacheAndProcessingStatus() const { return CacheAndProcessMessageStatus; }
 		TOptional<ESystemControllingTimeMarker> TryGetSystemControllingTimeMarker() const { return SystemControllingTime; }
 
-		void StopProcessingNewMessages() { CacheAndProcessMessageStatus = ECacheAndProcess::CacheLatestNoProcess; }
-		void StopCachingAndProcessingNewMessages() { CacheAndProcessMessageStatus = ECacheAndProcess::None; }
-
-		double GetCurrentDurationSeconds() const;
-		TRange<double> GetPlottingRange() const { return PlottingRange; }
+		UE_API double GetCurrentTraceDurationSeconds() const;
 
-		void PauseTimeMarker(const double InTimestamp, const ESystemControllingTimeMarker InControllingSystem, TOptional<TRange<double>> InPlottingRange = TOptional<TRange<double>>());
-		void ResumeTimeMarker();
+		UE_API void PauseTimeMarker(const double InTimestamp, const ESystemControllingTimeMarker InControllingSystem, TOptional<TRange<double>> InPlottingRange = TOptional<TRange<double>>());
+		UE_API void ResumeTimeMarker();
 
 		DECLARE_MULTICAST_DELEGATE_OneParam(FOnTimingViewTimeMarkerChanged, double /*TimeMarker*/);
 		FOnTimingViewTimeMarkerChanged OnTimingViewTimeMarkerChanged;
@@ -64,7 +59,17 @@ namespace UE::Audio::Insights
 		static constexpr double MaxPlottingHistorySeconds = 5.0;
 		static constexpr double PlottingMarginSeconds = 0.2;
 
+		// AudioInsights Internal:
+		void StopProcessingNewMessages();
+		void StopCachingAndProcessingNewMessages();
+		TRange<double> GetPlottingRange() const;
+
 	private:
+		// Insights::ITimingViewExtender interface
+		virtual void OnBeginSession(UE::Insights::Timing::ITimingViewSession& InSession) override;
+		virtual void OnEndSession(UE::Insights::Timing::ITimingViewSession& InSession) override;
+		virtual void Tick(UE::Insights::Timing::ITimingViewSession& InSession, const TraceServices::IAnalysisSession& InAnalysisSession) override;
+
 		void OnTimeMarkerChanged(UE::Insights::Timing::ETimeChangedFlags InFlags, double InTimeMarker);
 
 		const TraceServices::IAnalysisSession* AnalysisSession = nullptr;
@@ -77,3 +82,5 @@ namespace UE::Audio::Insights
 		bool bUserInputDetected = false;
 	};
 } // namespace UE::Audio::Insights
+
+#undef UE_API
