[Audio Insights] - TimingViewExtender: FAudioInsightsTimingViewExtender inherits publicly from UE::Insights::Timing::ITimingViewExtender

#Synced-CL 51011358

diff --git a/Engine/Plugins/AudioInsights/Source/AudioInsights/Public/AudioInsightsTimingViewExtender.h b/Engine/Plugins/AudioInsights/Source/AudioInsights/Public/AudioInsightsTimingViewExtender.h
--- a/Engine/Plugins/AudioInsights/Source/AudioInsights/Public/AudioInsightsTimingViewExtender.h
+++ b/Engine/Plugins/AudioInsights/Source/AudioInsights/Public/AudioInsightsTimingViewExtender.h
@@ -36,7 +36,7 @@ namespace UE::Audio::Insights
 		External
 	};
 
-	class FAudioInsightsTimingViewExtender final : private UE::Insights::Timing::ITimingViewExtender
+	class FAudioInsightsTimingViewExtender final : public UE::Insights::Timing::ITimingViewExtender
 	{
 	public:
 
