[Audio Insights] Move send levels to output pins in Signal Flow
* Removed old Send Label Widgets that were connected to connection wires, along with multimap
* Moved send levels inside SSignalFlowGraphNode - displayed on the right hand side in Horizontal mode, and along the bottom in Vertical mode
* Send levels now have their own node detail setting and can be toggled on/off
* Send levels can be hidden by a node's expander arrow
* Output wires will now align with send level pins, or collapse to the center of the node on the outputting side if pins are hidden
* Updated calculations to node size and graph slot size to account for labels inside the nodes
* Pin color reflects the type of node the output is connected to
* Pin alpha linked to send level

#Synced-CL 51049241

#Synced-CL 51049241

diff --git a/Engine/Plugins/AudioInsights/Source/AudioInsights/Private/Settings/SignalFlowSettings.cpp b/Engine/Plugins/AudioInsights/Source/AudioInsights/Private/Settings/SignalFlowSettings.cpp
--- a/Engine/Plugins/AudioInsights/Source/AudioInsights/Private/Settings/SignalFlowSettings.cpp
+++ b/Engine/Plugins/AudioInsights/Source/AudioInsights/Private/Settings/SignalFlowSettings.cpp
@@ -45,6 +45,9 @@ bool FSignalFlowNodeDetailFilterSettings::GetParameterIsVisible(const UE::Audio:
 		case ESignalFlowNodeDetailParam::RelativeRenderCost:
 			return bRelativeRenderCost;
 
+		case ESignalFlowNodeDetailParam::SendOutputVolume:
+			return bSendOutputVolume;
+
 		default:
 			ensureMsgf(false, TEXT("Unrecognized ESignalFlowNodeDetailParam value"));
 			break;
@@ -53,7 +56,7 @@ bool FSignalFlowNodeDetailFilterSettings::GetParameterIsVisible(const UE::Audio:
     return false;
 }
 
-int32 FSignalFlowNodeDetailFilterSettings::GetNumVisibleParams() const
+int32 FSignalFlowNodeDetailFilterSettings::GetNumVisibleParams(const TSet<UE::Audio::Insights::ESignalFlowNodeDetailParam>& IgnoredParams /* = TSet<UE::Audio::Insights::ESignalFlowNodeDetailParam>()*/) const
 {
 	using namespace UE::Audio::Insights;
 
@@ -62,7 +65,8 @@ int32 FSignalFlowNodeDetailFilterSettings::GetNumVisibleParams() const
 	const int32 NumFilters = static_cast<int32>(ESignalFlowNodeDetailParam::MAX);
 	for (uint8 FilterID = 0u; FilterID < NumFilters; ++FilterID)
 	{
-		if (GetParameterIsVisible(static_cast<ESignalFlowNodeDetailParam>(FilterID)))
+		const ESignalFlowNodeDetailParam ParamType = static_cast<ESignalFlowNodeDetailParam>(FilterID);
+		if (!IgnoredParams.Contains(ParamType) && GetParameterIsVisible(ParamType))
 		{
 			NumVisibleParams++;
 		}
@@ -147,6 +151,10 @@ void FSignalFlowNodeDetailFilterSettings::SetParameterVisibility(const UE::Audio
 			bRelativeRenderCost = bIsVisible;
 			break;
 
+		case ESignalFlowNodeDetailParam::SendOutputVolume:
+			bSendOutputVolume = bIsVisible;
+			break;
+
 		default:
 			ensureMsgf(false, TEXT("Unrecognized ESignalFlowNodeDetailParam value"));
 			break;
diff --git a/Engine/Plugins/AudioInsights/Source/AudioInsights/Private/Views/SSignalFlowGraph.cpp b/Engine/Plugins/AudioInsights/Source/AudioInsights/Private/Views/SSignalFlowGraph.cpp
--- a/Engine/Plugins/AudioInsights/Source/AudioInsights/Private/Views/SSignalFlowGraph.cpp
+++ b/Engine/Plugins/AudioInsights/Source/AudioInsights/Private/Views/SSignalFlowGraph.cpp
@@ -31,10 +31,9 @@ namespace UE::Audio::Insights
 {
 	namespace SSignalFlowGraphStyle
 	{
-		constexpr float NodeMaxWidth = 240.0f;
+		constexpr float NodeMaxWidth = 256.0f;
 
 		constexpr float NodeRoundedBrushRadius = 12.0f;
-		constexpr float SendLabelRoundedBrushRadius = 8.0f;
 
 		constexpr float SplineRadius = 80.0f;
 		constexpr float SplineThicknessScalarMin = 1.0f;
@@ -43,14 +42,15 @@ namespace UE::Audio::Insights
 		constexpr float SplineThicknessMin = 1.0f;
 		constexpr float SplineAnimPointsPerUnit = 0.2f;
 
-		constexpr float SendLabelPropAlongLine = 0.25f; // How far alone the spline we want the send level to appear (range 0 -> 1)
-
 		constexpr float SelectedNodeBorderThickness = 3.0f;
 		constexpr float NodeBorderThickness = 2.0f;
-		constexpr float SendLabelBorderThickness = 1.0f;
 
 		constexpr float NodeBorderPadding = 12.0f;
-		constexpr float SendLevelTextMargin = 2.0f;
+		constexpr float NodeDetailVerticalPadding = 8.0f;
+		constexpr float NodeDetailIndentationPadding = 20.0f;
+
+		constexpr float SendLabelVerticalWidth = 20.0f;
+		constexpr float SendLabelPadding = 8.0f;
 
 		constexpr float AnimationSpeed = 15.0f;
 		constexpr float DefaultZoomScale = 1.0f;
@@ -65,12 +65,6 @@ namespace UE::Audio::Insights
 			return RoundedBoxBrush;
 		}
 
-		const FSlateRoundedBoxBrush& GetSendLabelRoundedBoxBrush()
-		{
-			static const FSlateRoundedBoxBrush RoundedBoxBrush(FLinearColor::White, SendLabelRoundedBrushRadius);
-			return RoundedBoxBrush;
-		}
-
 		float GetNodeGraphSlotWidth(const EOrientation GraphOrientation)
 		{
 			const float XPaddingBetweenNodes = GraphOrientation == EOrientation::Orient_Horizontal ? 48.0f : 8.0f;
@@ -315,56 +309,95 @@ namespace UE::Audio::Insights
 
 		ChildSlot
 		[
-			SNew(SVerticalBox)
-
-			+ SVerticalBox::Slot()
-			.AutoHeight()
-			.HAlign(HAlign_Center)
+			SAssignNew(SelectedNodeBorder, SBorder)
+			.BorderImage(&GetNodeRoundedBoxBrush())
+			.HAlign(HAlign_Fill)
+			.VAlign(VAlign_Fill)
+			.BorderBackgroundColor_Lambda([this]()
+			{
+				return IsSelected.IsBound() && IsSelected.Get() ? FStyleColors::AccentOrange.GetSpecifiedColor() : FLinearColor::Transparent;
+			})
+			.Padding(SelectedNodeBorderThickness)
 			[
-				SAssignNew(SelectedNodeBorder, SBorder)
+				SAssignNew(NodeBorder, SBorder)
 				.BorderImage(&GetNodeRoundedBoxBrush())
 				.HAlign(HAlign_Fill)
 				.VAlign(VAlign_Fill)
-				.BorderBackgroundColor_Lambda([this]()
-				{
-					return IsSelected.IsBound() && IsSelected.Get() ? FStyleColors::AccentOrange.GetSpecifiedColor() : FLinearColor::Transparent;
-				})
-				.Padding(SelectedNodeBorderThickness)
+				.BorderBackgroundColor(GetNodeAccentColor(NodeEntryKey.EntryType))
+				.Padding(NodeBorderThickness)
+				.Content()
 				[
-					SAssignNew(NodeBorder, SBorder)
+					SNew(SBorder)
 					.BorderImage(&GetNodeRoundedBoxBrush())
 					.HAlign(HAlign_Fill)
 					.VAlign(VAlign_Fill)
-					.BorderBackgroundColor(GetNodeAccentColor(NodeEntryKey.EntryType))
-					.Padding(NodeBorderThickness)
+					.BorderBackgroundColor_Lambda([this]()
+					{
+						return (IsSelected.IsBound() && IsSelected.Get()) ? FStyleColors::Header : FStyleColors::Panel;
+					})
+					.Padding(0.0f, NodeBorderPadding, 0.0f, NodeBorderPadding)
 					.Content()
 					[
-						SNew(SBorder)
-						.BorderImage(&GetNodeRoundedBoxBrush())
-						.HAlign(HAlign_Fill)
-						.VAlign(VAlign_Fill)
-						.BorderBackgroundColor_Lambda([this]()
-						{
-							return (IsSelected.IsBound() && IsSelected.Get()) ? FStyleColors::Header : FStyleColors::Panel;
-						})
-						.Padding(NodeBorderPadding)
-						.Content()
+						SNew(SVerticalBox)
+
+						+ SVerticalBox::Slot()
+						.Padding(NodeBorderPadding, 0.0f, NodeBorderPadding, 0.0f)
+						.AutoHeight()
+						[
+							CreateNodeName(InArgs)
+						]
+
+						+ SVerticalBox::Slot()
+						.Padding(NodeBorderPadding + NodeDetailIndentationPadding, 0.0f, 0.0f, 0.0f)
+						.AutoHeight()
 						[
-							SNew(SVerticalBox)
+							SNew(SHorizontalBox)
 
-							+ SVerticalBox::Slot()
-							.AutoHeight()
+							+ SHorizontalBox::Slot()
+							.FillWidth(1.0f)
 							[
-								CreateNodeName(InArgs)
+								CreateNodeDetails()
 							]
 
-							+ SVerticalBox::Slot()
-							.Padding(20.0f, 0.0f)
-							.AutoHeight()
+							+ SHorizontalBox::Slot()
+							.HAlign(HAlign_Right)
+							.VAlign(VAlign_Top)
+							.AutoWidth()
 							[
-								CreateNodeDetails()
+								// In vertical mode, we want to ensure that we maintain padding on the RHS of the node when node details is on
+								// We do this by drawing a box with no content inside.
+								// In horizontal mode, this is replaced by SendLevelsVerticalBox, which is flush to the RHS
+								SNew(SVerticalBox)
+								.Visibility_Lambda([this]()
+								{
+									return GraphOrientation.Get() == EOrientation::Orient_Vertical && ShowNodeDetails.Get()
+																									? EVisibility::Visible
+																									: EVisibility::Collapsed;
+								})
+
+								+ SVerticalBox::Slot()
+								[
+									SNew(SBox)
+									.WidthOverride(NodeBorderPadding)
+								]
+							]
+
+							+ SHorizontalBox::Slot()
+							.HAlign(HAlign_Right)
+							.VAlign(VAlign_Top)
+							.AutoWidth()
+							[
+								CreateSendOutputsVerticalContainer()
 							]
 						]
+
+						+ SVerticalBox::Slot()
+						.HAlign(HAlign_Center)
+						.VAlign(VAlign_Bottom)
+						.AutoHeight()
+						[
+							CreateSendOutputsHorizontalContainer()
+						]
 					]
 				]
 			]
@@ -417,7 +450,8 @@ namespace UE::Audio::Insights
 		{
 			if (bNodeDetailsAreExpanded)
 			{
-				NumParamsVisible = NodeDetailFilterSettings->GetNumVisibleParams();
+				const TSet<ESignalFlowNodeDetailParam> IgnoredParams { ESignalFlowNodeDetailParam::SendOutputVolume };
+				NumParamsVisible = NodeDetailFilterSettings->GetNumVisibleParams(IgnoredParams);
 			}
 			// We always show amplitude if the node is collapsed
 			else if (Entry.IsValid() && Entry->Amplitude.IsSet())
@@ -426,9 +460,28 @@ namespace UE::Audio::Insights
 			}
 		}
 
+		const int32 NumOutputPinsShowing = SendLabels.Num();
 		const FVector2D Size = SCompoundWidget::ComputeDesiredSize(InLayoutScaleMultiplier);
-		const FVector2D ClampedSize = { FMath::Min(Size.X, NodeMaxWidth), FMath::Min(Size.Y, GetNodeMaxHeight(NumParamsVisible, ShowNodeDetails.Get())) };
-		return ClampedSize;
+
+		if (GraphOrientation.Get() == EOrientation::Orient_Horizontal)
+		{
+			const int32 MaxNumParamRows = FMath::Max(NumParamsVisible, NumOutputPinsShowing);
+
+			const FVector2D ClampedSize = { FMath::Min(Size.X, NodeMaxWidth), FMath::Min(Size.Y, GetNodeMaxHeight(MaxNumParamRows, ShowNodeDetails.Get())) };
+			return ClampedSize;
+		}
+		else
+		{
+			const float PinSectionWidth = SendLevelsHorizontalBox.IsValid() ? SendLevelsHorizontalBox->GetDesiredSize().X : 0.0f;
+			const float OutputPinSectionHeight = SendLevelsHorizontalBox.IsValid() ? SendLevelsHorizontalBox->GetDesiredSize().Y : 0.0f;
+
+			// Make sure Max width/height calculations account for the bottom output pin row
+			const float ScaledMaxWidth = FMath::Max(NodeMaxWidth, PinSectionWidth);
+			const float TotalMaxHeight = GetNodeMaxHeight(NumParamsVisible, ShowNodeDetails.Get()) + OutputPinSectionHeight;
+
+			const FVector2D ClampedSize = { FMath::Min(Size.X, ScaledMaxWidth), FMath::Min(Size.Y, TotalMaxHeight) };
+			return ClampedSize;
+		}
 	}
 
 	FVector2D SSignalFlowGraphNode::GetLocalInputPos() const
@@ -440,16 +493,16 @@ namespace UE::Audio::Insights
 		if (GraphOrientation.IsBound() && GraphOrientation.Get() == EOrientation::Orient_Horizontal)
 		{
 			// Input pos = left center of node
-			return FVector2D{ CachedPositionInGraph.X + SelectedNodeBorderThickness, CachedPositionInGraph.Y + (Size.Y / 2.0f) } * ZoomScaleFactor.Get();
+			return FVector2D{ CachedPositionInGraph.X + SelectedNodeBorderThickness, CachedPositionInGraph.Y + (Size.Y * 0.5f) } * ZoomScaleFactor.Get();
 		}
 		else
 		{
 			// Input pos = top center of node
-			return FVector2D{ CachedPositionInGraph.X + (Size.X / 2.0f), CachedPositionInGraph.Y + SelectedNodeBorderThickness } * ZoomScaleFactor.Get();
+			return FVector2D{ CachedPositionInGraph.X + (Size.X * 0.5f), CachedPositionInGraph.Y + SelectedNodeBorderThickness } * ZoomScaleFactor.Get();
 		}
 	}
 	
-	FVector2D SSignalFlowGraphNode::GetLocalOutputPos() const
+	FVector2D SSignalFlowGraphNode::GetLocalCenteredOutputPos() const
 	{
 		using namespace SSignalFlowGraphStyle;
 
@@ -467,6 +520,58 @@ namespace UE::Audio::Insights
 		}
 	}
 
+	FVector2D SSignalFlowGraphNode::GetLocalSendLabelPos(const FSignalFlowEntryKey& ReceiverEntryKey, const int32 LabelIndex) const
+	{
+		using namespace SSignalFlowGraphStyle;
+
+		const FVector2D OutputPos = GetLocalCenteredOutputPos();
+		if (!NodeLabelHorizontalBox.IsValid() || !SendLabelsAreVisible())
+		{
+			return OutputPos;
+		}
+
+		const TSharedRef<SSignalFlowGraphSendLabel>* SendLabelWidget = SendLabels.Find(ReceiverEntryKey);
+		if (SendLabelWidget == nullptr)
+		{
+			return OutputPos;
+		}
+
+		constexpr float BorderThicknessYOffset = (SelectedNodeBorderThickness * 0.5f) + (NodeBorderThickness * 0.5f);
+
+		const FVector2D SendLabelWidgetSize = (*SendLabelWidget)->GetDesiredSize();
+
+		float XPosInGraph;
+		float YPosInGraph;
+
+		if (GraphOrientation.Get() == EOrientation::Orient_Horizontal)
+		{
+			XPosInGraph = OutputPos.X;
+
+			const float NodeLabelHeight = NodeLabelHorizontalBox->GetDesiredSize().Y + NodeBorderPadding;
+			const float YOffsetToLabelArea = CachedPositionInGraph.Y + BorderThicknessYOffset + NodeLabelHeight;
+			const float YOffsetInLabelList = ((SendLabelWidgetSize.Y + NodeDetailVerticalPadding) * LabelIndex) + NodeDetailVerticalPadding;
+			const float HalfLabelHeight = (SendLabelWidgetSize.Y * 0.5f);
+
+			YPosInGraph = (YOffsetToLabelArea + YOffsetInLabelList + HalfLabelHeight) * ZoomScaleFactor.Get();
+		}
+		else
+		{
+			YPosInGraph = OutputPos.Y;
+
+			const float SendLabelWidth = SendLabelWidgetSize.X + SendLabelPadding;
+			const float XOffset = (SendLabelWidth * (LabelIndex + 1)) - ((SendLabelWidth * SendLabels.Num()) * 0.5f) - (SendLabelWidth * 0.5f);
+			XPosInGraph = OutputPos.X + (XOffset * ZoomScaleFactor.Get());
+		}
+
+		const FVector2D SendLabelPosition(XPosInGraph, YPosInGraph);
+		return SendLabelPosition;
+	}
+
+	int32 SSignalFlowGraphNode::GetNumSendLabels() const
+	{
+		return SendLabels.Num();
+	}
+
 	void SSignalFlowGraphNode::UpdateNodeLayout()
 	{
 		using namespace SSignalFlowGraphStyle;
@@ -491,11 +596,111 @@ namespace UE::Audio::Insights
 		return CachedPositionInGraph;
 	}
 
+	void SSignalFlowGraphNode::SetSendLabelWidgets(const TArray<FSignalFlowEntryKey>& ReceiverEntryKeys)
+	{
+		for (const FSignalFlowEntryKey& ReceiverKey : ReceiverEntryKeys)
+		{
+			AddSendLabelWidget(ReceiverKey);
+		}
+
+		for (auto Itr = SendLabels.CreateIterator(); Itr; ++Itr)
+		{
+			if (!ReceiverEntryKeys.Contains(Itr->Key))
+			{
+				if (SendLevelsVerticalBox.IsValid())
+				{
+					SendLevelsVerticalBox->RemoveSlot(Itr->Value);
+				}
+
+				if (SendLevelsHorizontalBox.IsValid())
+				{
+					SendLevelsHorizontalBox->RemoveSlot(Itr->Value);
+				}
+
+				Itr.RemoveCurrent();
+			}
+		}
+	}
+
+	void SSignalFlowGraphNode::AddSendLabelWidget(const FSignalFlowEntryKey& ReceiverEntryKey)
+	{
+		using namespace SSignalFlowGraphStyle;
+
+		if (!SendLevelsVerticalBox.IsValid() || !SendLevelsHorizontalBox.IsValid() || SendLabels.Contains(ReceiverEntryKey))
+		{
+			return;
+		}
+
+		TSharedPtr<SSignalFlowGraphSendLabel> SendLabelWidget = SNew(SSignalFlowGraphSendLabel)
+			.PixelSnappingMethod(EWidgetPixelSnapping::SnapToPixel)
+			.ToolTipText(LOCTEXT("SignalFlowGraph_NodeDetails_SendLabelTooltip", "Send Output Volume"))
+			.ReceiverEntryType(ReceiverEntryKey.EntryType)
+			.SendLevel_Lambda([this, ReceiverEntryKey]() -> TOptional<float>
+			{
+				if (Entry.IsValid())
+				{
+					const FSignalFlowOutputData* SendInfo = Entry->Outputs.Find(ReceiverEntryKey);
+					if (SendInfo && SendInfo->SendLevel.IsSet())
+					{
+						return SendInfo->SendLevel;
+					}
+				}
+
+				return NullOpt;
+			})
+			.GraphOrientation(GraphOrientation)
+			.ZoomScaleFactor(ZoomScaleFactor);
+
+		SendLabels.Add(ReceiverEntryKey, SendLabelWidget.ToSharedRef());
+
+		SendLevelsVerticalBox->AddSlot()
+		.Padding(SendLabelPadding, NodeDetailVerticalPadding, 0.0f, 0.0f)
+		.VAlign(VAlign_Top)
+		.HAlign(HAlign_Right)
+		[
+			SendLabelWidget.ToSharedRef()
+		];
+
+		SendLevelsHorizontalBox->AddSlot()
+		.Padding(SendLabelPadding * 0.5f, NodeDetailVerticalPadding, SendLabelPadding * 0.5f, 0.0f)
+		.VAlign(VAlign_Center)
+		.HAlign(HAlign_Center)
+		[
+			SendLabelWidget.ToSharedRef()
+		];
+	}
+
+	void SSignalFlowGraphNode::RemoveSendLabelWidget(const FSignalFlowEntryKey& ReceiverEntryKey)
+	{
+		TSharedRef<SSignalFlowGraphSendLabel>* SendLabelWidget = SendLabels.Find(ReceiverEntryKey);
+		if (SendLabelWidget == nullptr)
+		{
+			return;
+		}
+
+		if (SendLevelsVerticalBox.IsValid())
+		{
+			SendLevelsVerticalBox->RemoveSlot(*SendLabelWidget);
+		}
+
+		if (SendLevelsHorizontalBox.IsValid())
+		{
+			SendLevelsHorizontalBox->RemoveSlot(*SendLabelWidget);
+		}
+
+		SendLabels.Remove(ReceiverEntryKey);
+	}
+
+	bool SSignalFlowGraphNode::ContainsSendLabel(const FSignalFlowEntryKey& ReceiverEntryKey) const
+	{
+		return SendLabels.Contains(ReceiverEntryKey);
+	}
+
 	TSharedRef<SHorizontalBox> SSignalFlowGraphNode::CreateNodeName(const SSignalFlowGraphNode::FArguments& InArgs)
 	{
 		using namespace SSignalFlowGraphStyle;
 
-		return SNew(SHorizontalBox)
+		return SAssignNew(NodeLabelHorizontalBox, SHorizontalBox)
 
 			+ SHorizontalBox::Slot()
 			.AutoWidth()
@@ -543,67 +748,69 @@ namespace UE::Audio::Insights
 
 	TSharedRef<SVerticalBox> SSignalFlowGraphNode::CreateNodeDetails() const
 	{
+		using namespace SSignalFlowGraphStyle;
+
 		return SNew(SVerticalBox)
 				.Visibility_Lambda([this]() { return ShowNodeDetails.Get() ? EVisibility::Visible : EVisibility::Collapsed; })
 
 				+ SVerticalBox::Slot()
-				.Padding(0.0f, 8.0f, 0.0f, 0.0f)
+				.Padding(0.0f, NodeDetailVerticalPadding, 0.0f, 0.0f)
 				.AutoHeight()
 				[
 					CreateNodeDetailsRow(ESignalFlowNodeDetailParam::Amplitude)
 				]
 			
 				+ SVerticalBox::Slot()
-				.Padding(0.0f, 8.0f, 0.0f, 0.0f)
+				.Padding(0.0f, NodeDetailVerticalPadding, 0.0f, 0.0f)
 				.AutoHeight()
 				[
 					CreateNodeDetailsRow(ESignalFlowNodeDetailParam::Volume)
 				]
 
 				+ SVerticalBox::Slot()
-				.Padding(0.0f, 8.0f, 0.0f, 0.0f)
+				.Padding(0.0f, NodeDetailVerticalPadding, 0.0f, 0.0f)
 				.AutoHeight()
 				[
 					CreateNodeDetailsRow(ESignalFlowNodeDetailParam::Pitch)
 				]
 
 				+ SVerticalBox::Slot()
-				.Padding(0.0f, 8.0f, 0.0f, 0.0f)
+				.Padding(0.0f, NodeDetailVerticalPadding, 0.0f, 0.0f)
 				.AutoHeight()
 				[
 					CreateNodeDetailsRow(ESignalFlowNodeDetailParam::LPFFreq)
 				]
 			
 				+ SVerticalBox::Slot()
-				.Padding(0.0f, 8.0f, 0.0f, 0.0f)
+				.Padding(0.0f, NodeDetailVerticalPadding, 0.0f, 0.0f)
 				.AutoHeight()
 				[
 					CreateNodeDetailsRow(ESignalFlowNodeDetailParam::HPFFreq)
 				]
 
 				+ SVerticalBox::Slot()
-				.Padding(0.0f, 8.0f, 0.0f, 0.0f)
+				.Padding(0.0f, NodeDetailVerticalPadding, 0.0f, 0.0f)
 				.AutoHeight()
 				[
 					CreateNodeDetailsRow(ESignalFlowNodeDetailParam::Priority)
 				]
 
 				+ SVerticalBox::Slot()
-				.Padding(0.0f, 8.0f, 0.0f, 0.0f)
+				.Padding(0.0f, NodeDetailVerticalPadding, 0.0f, 0.0f)
 				.AutoHeight()
 				[
 					CreateNodeDetailsRow(ESignalFlowNodeDetailParam::Distance)
 				]
 
 				+ SVerticalBox::Slot()
-				.Padding(0.0f, 8.0f, 0.0f, 0.0f)
+				.Padding(0.0f, NodeDetailVerticalPadding, 0.0f, 0.0f)
 				.AutoHeight()
 				[
 					CreateNodeDetailsRow(ESignalFlowNodeDetailParam::Attenuation)
 				]
 
 				+ SVerticalBox::Slot()
-				.Padding(0.0f, 8.0f, 0.0f, 0.0f)
+				.Padding(0.0f, NodeDetailVerticalPadding, 0.0f, 0.0f)
 				.AutoHeight()
 				[
 					CreateNodeDetailsRow(ESignalFlowNodeDetailParam::RelativeRenderCost)
@@ -661,6 +868,32 @@ namespace UE::Audio::Insights
 				];
 	}
 
+	TSharedRef<SVerticalBox> SSignalFlowGraphNode::CreateSendOutputsVerticalContainer()
+	{
+		SendLevelsVerticalBox = SNew(SVerticalBox)
+								.Visibility_Lambda([this]()
+								{
+									return GraphOrientation.Get() == EOrientation::Orient_Horizontal && SendLabelsAreVisible() 
+																	? EVisibility::Visible 
+																	: EVisibility::Collapsed;
+								});
+
+		return SendLevelsVerticalBox.ToSharedRef();
+	}
+
+	TSharedRef<SHorizontalBox> SSignalFlowGraphNode::CreateSendOutputsHorizontalContainer()
+	{
+		SendLevelsHorizontalBox = SNew(SHorizontalBox)
+									.Visibility_Lambda([this]()
+									{
+										return GraphOrientation.Get() == EOrientation::Orient_Vertical && SendLabelsAreVisible()
+																	? EVisibility::Visible
+																	: EVisibility::Collapsed;
+									});
+
+		return SendLevelsHorizontalBox.ToSharedRef();
+	}
+
 	bool SSignalFlowGraphNode::NodeDetailsExpanderArrowIsEnabled() const
 	{
 		using namespace SSignalFlowGraphStyle;
@@ -692,7 +925,38 @@ namespace UE::Audio::Insights
 				continue;
 			}
 
-			if (NodeDetailFilterSettings->GetParameterIsVisible(Param) && GetNodeDetailParamValue(Entry, Param).IsSet())
+			if (Param == ESignalFlowNodeDetailParam::SendOutputVolume)
+			{
+				if (NodeDetailFilterSettings->GetParameterIsVisible(Param))
+				{
+					return true;
+				}
+			}
+			else if (NodeDetailFilterSettings->GetParameterIsVisible(Param) && GetNodeDetailParamValue(Entry, Param).IsSet())
+			{
+				return true;
+			}
+		}
+
+		return false;
+	}
+
+	bool SSignalFlowGraphNode::SendLabelsAreVisible() const
+	{
+		if (!ShowNodeDetails.Get() || SendLabels.IsEmpty())
+		{
+			return false;
+		}
+
+		const bool bParamIsFiltered = NodeDetailFilterSettings && NodeDetailFilterSettings->GetParameterIsVisible(ESignalFlowNodeDetailParam::SendOutputVolume);
+		if (!bNodeDetailsAreExpanded || !bParamIsFiltered)
+		{
+			return false;
+		}
+
+		for (const auto& [ReceiverKey, Label] : SendLabels)
+		{
+			if (Label->ValueIsSet())
 			{
 				return true;
 			}
@@ -705,60 +969,67 @@ namespace UE::Audio::Insights
 	{
 		using namespace SSignalFlowGraphStyle;
 
-		SendReceiverEntryKey = InArgs._SendReceiverEntryKey;
-		SenderPositionInGraph = InArgs._SenderPositionInGraph;
-		ReceiverPositionInGraph = InArgs._ReceiverPositionInGraph;
 		SendLevel = InArgs._SendLevel;
 		GraphOrientation = InArgs._GraphOrientation;
 		ZoomScaleFactor = InArgs._ZoomScaleFactor;
 
 		ChildSlot
 		[
-			SNew(SHorizontalBox)
-			.Visibility_Lambda([this]()
-			{
-				return SendLevel.IsBound() && SendLevel.Get().IsSet() ? EVisibility::Visible : EVisibility::Collapsed;
-			})
+			SNew(SOverlay)
 
-			+ SHorizontalBox::Slot()
-			.FillContentWidth(1.0f)
-			.HAlign(HAlign_Center)
-			.VAlign(VAlign_Center)
+			// Horizontal flow label (text then icon)
+			+ SOverlay::Slot()
 			[
-				SAssignNew(LabelBorder, SBorder)
-				.BorderImage(&GetSendLabelRoundedBoxBrush())
-				.HAlign(HAlign_Fill)
-				.VAlign(VAlign_Fill)
-				.BorderBackgroundColor(FStyleColors::Foreground.GetSpecifiedColor())
-				.Padding(SendLabelBorderThickness)
-				.Content()
+				SNew(SHorizontalBox)
+				.Visibility_Lambda([this]()
+				{
+					return GraphOrientation.Get() == EOrientation::Orient_Horizontal && SendLevel.IsBound() && SendLevel.Get().IsSet() 
+																												? EVisibility::Visible 
+																												: EVisibility::Collapsed;
+				})
+
+				+ SHorizontalBox::Slot()
+				.FillContentWidth(1.0f)
+				.HAlign(HAlign_Right)
+				.VAlign(VAlign_Center)
 				[
-					SNew(SBorder)
-					.BorderImage(&GetSendLabelRoundedBoxBrush())
-					.HAlign(HAlign_Fill)
-					.VAlign(VAlign_Fill)
-					.BorderBackgroundColor(FStyleColors::Panel.GetSpecifiedColor())
-					.Content()
-					[
-						SNew(STextBlock)
-						.Text_Lambda([this]() 
-						{
-							if (!SendLevel.IsBound())
-							{
-								return FText::GetEmpty();
-							}
-
-							const TOptional<float> SendLevelOpt = SendLevel.Get();
-							if (!SendLevelOpt.IsSet())
-							{
-								return FText::GetEmpty();
-							}
-
-							return FText::AsNumber(SendLevelOpt.GetValue(), FSlateStyle::Get().GetSendLevelFormat());
-						})
-						.Justification(ETextJustify::Center)
-						.Margin(SendLevelTextMargin)
-					]
+					CreateOutputLabelText(InArgs._ToolTipText.Get())
+				]
+
+				+ SHorizontalBox::Slot()
+				.FillContentWidth(1.0f)
+				.HAlign(HAlign_Right)
+				.VAlign(VAlign_Center)
+				[
+					CreatePinIconWidget(InArgs._ToolTipText.Get(), InArgs._ReceiverEntryType)
+				]
+			]
+
+			// Vertical flow label (text above icon)
+			+ SOverlay::Slot()
+			[
+				SNew(SVerticalBox)
+				.Visibility_Lambda([this]()
+				{
+					return GraphOrientation.Get() == EOrientation::Orient_Vertical && SendLevel.IsBound() && SendLevel.Get().IsSet() 
+																												? EVisibility::Visible 
+																												: EVisibility::Collapsed;
+				})
+
+				+ SVerticalBox::Slot()
+				.FillContentHeight(1.0f)
+				.HAlign(HAlign_Center)
+				.VAlign(VAlign_Bottom)
+				[
+					CreateOutputLabelText(InArgs._ToolTipText.Get())
+				]
+
+				+ SVerticalBox::Slot()
+				.FillContentHeight(1.0f)
+				.HAlign(HAlign_Center)
+				.VAlign(VAlign_Bottom)
+				[
+					CreatePinIconWidget(InArgs._ToolTipText.Get(), InArgs._ReceiverEntryType)
 				]
 			]
 		];
@@ -769,34 +1040,80 @@ namespace UE::Audio::Insights
 		using namespace SSignalFlowGraphStyle;
 
 		const FVector2D Size = SCompoundWidget::ComputeDesiredSize(InLayoutScaleMultiplier);
-		return Size;
+		if (GraphOrientation.Get() == EOrientation::Orient_Horizontal)
+		{
+			return Size;
+		}
+		else
+		{
+			return FVector2D(SendLabelVerticalWidth, Size.Y);
+		}
 	}
 
-	void SSignalFlowGraphSendLabel::UpdateLabelLayout()
+	bool SSignalFlowGraphSendLabel::ValueIsSet() const
+	{
+		return SendLevel.IsBound() && SendLevel.Get().IsSet();
+	}
+
+	TSharedRef<STextBlock> SSignalFlowGraphSendLabel::CreateOutputLabelText(const FText& ToolTipText) const
+	{
+		return SNew(STextBlock)
+				.Text_Lambda([this]()
+				{
+					if (!SendLevel.IsBound())
+					{
+						return FText::GetEmpty();
+					}
+
+					const TOptional<float> SendLevelOpt = SendLevel.Get();
+					if (!SendLevelOpt.IsSet())
+					{
+						return FText::GetEmpty();
+					}
+
+					return FText::AsNumber(SendLevelOpt.GetValue(), FSlateStyle::Get().GetSendLevelFormat());
+				})
+				.ToolTipText(ToolTipText)
+				.Justification(ETextJustify::Center)
+				.Font(FCoreStyle::Get().GetFontStyle("SmallFont"));
+	}
+
+	TSharedRef<SOverlay> SSignalFlowGraphSendLabel::CreatePinIconWidget(const FText& ToolTipText, const ESignalFlowEntryType ReceiverEntryType)
 	{
 		using namespace SSignalFlowGraphStyle;
 
-		// Ensure border thickness does not get scaled to avoid rendering artifacts
-		if (LabelBorder.IsValid())
-		{
-			LabelBorder->SetPadding(SendLabelBorderThickness / ZoomScaleFactor.Get());
-		}
+		return SNew(SOverlay)
 
-		// Update the label's position in the graph
-		const FVector2D SenderNodePos = SenderPositionInGraph.Get() / ZoomScaleFactor.Get();
-		const FVector2D ReceiverNodePos = ReceiverPositionInGraph.Get() / ZoomScaleFactor.Get();
+			+ SOverlay::Slot()
+			[
+				SAssignNew(OutputPinImage, SImage)
+				.Image(FSlateStyle::Get().CreateIcon("AudioInsights.Icon.SignalFlow.Disconnected").GetIcon())
+				.ColorAndOpacity(GetNodeAccentColor(ReceiverEntryType))
+			]
 
-		const EOrientation Orientation = GraphOrientation.Get();
-		const FVector2D PointAlongSpline = FMath::CubicInterp(SenderNodePos, GetSplineStartDirection(Orientation), ReceiverNodePos, GetSplineEndDirection(Orientation), SendLabelPropAlongLine);
+			+ SOverlay::Slot()
+			[
+				SNew(SImage)
+				.Image(FSlateStyle::Get().CreateIcon("AudioInsights.Icon.SignalFlow.Connected").GetIcon())
+				.ToolTipText(ToolTipText)
+				.ColorAndOpacity_Lambda([this, ReceiverEntryType]()
+				{
+					const FSlateColor BaseColor = GetNodeAccentColor(ReceiverEntryType);
 
-		const FVector2D LabelSize = GetDesiredSize();
+					if (!SendLevel.IsBound())
+					{
+						return BaseColor;
+					}
 
-		CachedPositionInGraph = PointAlongSpline - (LabelSize / 2.0f);
-	}
+					const TOptional<float> SendLevelOpt = SendLevel.Get();
+					if (!SendLevelOpt.IsSet())
+					{
+						return BaseColor;
+					}
 
-	FVector2D SSignalFlowGraphSendLabel::GetCachedPositionInGraph() const
-	{
-		return CachedPositionInGraph;
+					return FSlateColor(BaseColor.GetSpecifiedColor().CopyWithNewOpacity(SendLevelOpt.GetValue()));
+				})
+			];
 	}
 
 	void SSignalFlowGraph::Construct(const SSignalFlowGraph::FArguments& InArgs)
@@ -879,7 +1196,6 @@ namespace UE::Audio::Insights
 
 		NodeWidgets.Reset();
 		DummyToRealConnectionKeys.Reset();
-		SendLevelWidgets.Reset();
 	}
 
 	void SSignalFlowGraph::ToggleOrientation()
@@ -1123,19 +1439,20 @@ namespace UE::Audio::Insights
 				continue;
 			}
 
-			const FVector2D ConnectionStartPos = (*InputNode)->GetLocalOutputPos();
-			const FVector2D ConnectionStartPosInView = ConnectionStartPos - ViewportPosInGraph;
-
-			for (const FSignalFlowEntryKey& NodeOutputKey : Node->FilteredOutputs)
+			for (int32 OutputIndex = 0; OutputIndex < Node->FilteredOutputs.Num(); ++OutputIndex)
 			{
+				const FSignalFlowEntryKey& NodeOutputKey = Node->FilteredOutputs[OutputIndex];
 				const TSharedPtr<SSignalFlowGraphNode> OutputNode = FindRealNode(NodeOutputKey, ENodeConnectionType::Output);
 
 				if (OutputNode.IsValid())
 				{
+					const FVector2D SendLabelPos = (*InputNode)->GetLocalSendLabelPos(OutputNode->GetNodeEntryKey(), OutputIndex);
+					const FVector2D SendLabelPosInView = SendLabelPos - ViewportPosInGraph;
+
 					// Check if the node connection wire is inside the viewport
 					const FVector2D ConnectionEndPos = OutputNode->GetLocalInputPos();
 					const FVector2D ConnectionEndPosInView = ConnectionEndPos - ViewportPosInGraph;
-					if (!FMath::LineBoxIntersection2D(LocalBoundsBox, ConnectionStartPosInView, ConnectionEndPosInView))
+					if (!FMath::LineBoxIntersection2D(LocalBoundsBox, SendLabelPosInView, ConnectionEndPosInView))
 					{
 						continue;
 					}
@@ -1148,7 +1465,7 @@ namespace UE::Audio::Insights
 						FSignalFlowOutputData* OutputData = Entry->Outputs.Find(OutputNode->GetNodeEntryKey());
 						const float SendLevel = OutputData && OutputData->SendLevel.IsSet() ? OutputData->SendLevel.GetValue() : 1.0f;
 
-						LayerId = AnimateWire(AllottedGeometry, OutDrawElements, LayerId, Entry, ConnectionStartPos, ConnectionEndPos, SendLevel);
+						LayerId = AnimateWire(AllottedGeometry, OutDrawElements, LayerId, Entry, SendLabelPos, ConnectionEndPos, SendLevel);
 					}
 					else
 					{
@@ -1158,7 +1475,7 @@ namespace UE::Audio::Insights
 							OutDrawElements,
 							++LayerId,
 							AllottedGeometry.ToPaintGeometry(),
-							ConnectionStartPos,
+							SendLabelPos,
 							GetSplineStartDirection(Orientation),
 							ConnectionEndPos,
 							GetSplineEndDirection(Orientation),
@@ -1317,10 +1634,16 @@ namespace UE::Audio::Insights
 		float GraphWidth = 0.0f;
 		float GraphHeight = 0.0f;
 
-		const int32 NumParamsVisible = NodeDetailFilterSettings ? NodeDetailFilterSettings->GetNumVisibleParams() : 0;
+		const TSet<ESignalFlowNodeDetailParam> IgnoredParams{ ESignalFlowNodeDetailParam::SendOutputVolume };
+		const int32 NumParamsVisible = NodeDetailFilterSettings ? NodeDetailFilterSettings->GetNumVisibleParams(IgnoredParams) : 0;
+
+		const bool bSendLabelsAreVisible = ShowNodeDetails.Get() && NodeDetailFilterSettings && NodeDetailFilterSettings->GetParameterIsVisible(ESignalFlowNodeDetailParam::SendOutputVolume);
+		const int32 MaxNumVisibleSendLabels = bSendLabelsAreVisible ? MaxNumNodeOutputsInGraph : 0;
+
+		const int32 MaxNumParamRows = FMath::Max(NumParamsVisible, MaxNumVisibleSendLabels);
 
 		const float NodeGraphSlotWidth = GetNodeGraphSlotWidth(GetGraphOrientation());
-		const float NodeGraphSlotHeight = GetNodeGraphSlotHeight(GetGraphOrientation(), NumParamsVisible, ShowNodeDetails.Get());
+		const float NodeGraphSlotHeight = GetNodeGraphSlotHeight(GetGraphOrientation(), MaxNumParamRows, ShowNodeDetails.Get());
 
 		for (const auto& [Key, NodeWidget] : NodeWidgets)
 		{
@@ -1397,6 +1720,8 @@ namespace UE::Audio::Insights
 
 	void SSignalFlowGraph::CreateSendLabels()
 	{
+		MaxNumNodeOutputsInGraph = 0;
+
 		// Run through the list source again to create send label widgets between nodes
 		for (const TSharedPtr<ISignalFlowNode>& Node : *ListSource)
 		{
@@ -1405,14 +1730,22 @@ namespace UE::Audio::Insights
 				continue;
 			}
 
+			const TSharedPtr<SSignalFlowGraphNode> NodeWidget = FindNodeWidget(Node->GetEntryKey());
+			if (!NodeWidget.IsValid())
+			{
+				continue;
+			}
+
 			for (const FSignalFlowEntryKey& OutputKey : Node->FilteredOutputs)
 			{
 				const TOptional<FSignalFlowEntryKey> RealNodeOutputKey = FindRealNodeKey(OutputKey, ENodeConnectionType::Output);
 				if (RealNodeOutputKey.IsSet())
 				{
-					CreateSendLabelWidget(Node->GetEntryKey(), RealNodeOutputKey.GetValue());
+					NodeWidget->AddSendLabelWidget(RealNodeOutputKey.GetValue());
 				}
 			}
+
+			MaxNumNodeOutputsInGraph = FMath::Max(MaxNumNodeOutputsInGraph, Node->FilteredOutputs.Num());
 		}
 	}
 
@@ -1421,7 +1754,7 @@ namespace UE::Audio::Insights
 		// Clean up and remove any old nodes that we did not detect in the new graph
 		for (const FSignalFlowEntryKey& Key : NodeKeysToRemove)
 		{
-			TSharedPtr<SSignalFlowGraphNode>* RemovedNode = NodeWidgets.Find(Key);;
+			TSharedPtr<SSignalFlowGraphNode>* RemovedNode = NodeWidgets.Find(Key);
 			if (!RemovedNode || !RemovedNode->IsValid())
 			{
 				NodeWidgets.Remove(Key);
@@ -1430,15 +1763,6 @@ namespace UE::Audio::Insights
 
 			Canvas->RemoveSlot(RemovedNode->ToSharedRef());
 
-			TArray<TSharedPtr<SSignalFlowGraphSendLabel>> SendLabels;
-			SendLevelWidgets.MultiFind(Key, SendLabels);
-			for (const TSharedPtr<SSignalFlowGraphSendLabel>& SendWidget : SendLabels)
-			{
-				Canvas->RemoveSlot(SendWidget.ToSharedRef());
-			}
-
-			SendLevelWidgets.Remove(Key);
-
 			if ((*RemovedNode)->GetEntry().IsValid())
 			{
 				for (const FSignalFlowEntryKey& InputKey : (*RemovedNode)->GetEntry()->Inputs)
@@ -1449,14 +1773,13 @@ namespace UE::Audio::Insights
 						continue;
 					}
 
-					const TSharedPtr<SSignalFlowGraphSendLabel> InputSendLabel = FindSendLabel(RealInputKey.GetValue(), Key);
-					if (!InputSendLabel.IsValid())
+					TSharedPtr<SSignalFlowGraphNode>* InputNode = NodeWidgets.Find(RealInputKey.GetValue());
+					if (InputNode == nullptr || !InputNode->IsValid())
 					{
 						continue;
 					}
 
-					SendLevelWidgets.RemoveSingle(RealInputKey.GetValue(), InputSendLabel);
-					Canvas->RemoveSlot(InputSendLabel.ToSharedRef());
+					(*InputNode)->RemoveSendLabelWidget(Key);
 				}
 			}
 
@@ -1556,8 +1879,16 @@ namespace UE::Audio::Insights
 		// Calculate how much the graph has increased below the original viewport (positive increase)
 		const int32 YPosDiffAtBottom = FMath::Max(NewYPos - OriginalYPos, 0);
 
-		const int32 NumParamsVisible = NodeDetailFilterSettings ? NodeDetailFilterSettings->GetNumVisibleParams() : 0;
-		const float NodeGraphSlotHeight = GetNodeGraphSlotHeight(GetGraphOrientation(), NumParamsVisible, ShowNodeDetails.Get());
+		// The max height of a node in the graph is driven by either the number of visible node details, or the number of send label outputs
+		const TSet<ESignalFlowNodeDetailParam> IgnoredParams{ ESignalFlowNodeDetailParam::SendOutputVolume };
+		const int32 NumParamsVisible = NodeDetailFilterSettings ? NodeDetailFilterSettings->GetNumVisibleParams(IgnoredParams) : 0;
+
+		const bool bSendLabelsAreVisible = ShowNodeDetails.Get() && NodeDetailFilterSettings && NodeDetailFilterSettings->GetParameterIsVisible(ESignalFlowNodeDetailParam::SendOutputVolume);
+		const int32 MaxNumVisibleSendLabels = bSendLabelsAreVisible ? MaxNumNodeOutputsInGraph : 0;
+
+		const int32 MaxNumParamRows = FMath::Max(NumParamsVisible, MaxNumVisibleSendLabels);
+
+		const float NodeGraphSlotHeight = GetNodeGraphSlotHeight(GetGraphOrientation(), MaxNumParamRows, ShowNodeDetails.Get());
 
 		return YPosDiffAtBottom * NodeGraphSlotHeight;
 	}
@@ -1595,36 +1926,18 @@ namespace UE::Audio::Insights
 
 				OutRemainingNodeKeys.RemoveSwap(EntryKey);
 
-				TArray<TSharedPtr<SSignalFlowGraphSendLabel>> SendLabels;
-				SendLevelWidgets.MultiFind(EntryKey, SendLabels);
-
-				TSet<FSignalFlowEntryKey> ActiveOutputKeys;
+				TArray<FSignalFlowEntryKey> ActiveOutputKeys;
 				for (const FSignalFlowEntryKey& NodeOutputKey : Node->FilteredOutputs)
 				{
 					const TOptional<FSignalFlowEntryKey> RealNodeOutputKeyOpt = FindRealNodeKey(NodeOutputKey, ENodeConnectionType::Output);
 					if (RealNodeOutputKeyOpt.IsSet())
 					{
 						const FSignalFlowEntryKey RealNodeOutputKey = RealNodeOutputKeyOpt.GetValue();
-						ActiveOutputKeys.Add(RealNodeOutputKey);
-
-						const TSharedPtr<SSignalFlowGraphSendLabel>* FoundExistingSendLabel = FindSendLabelInList(SendLabels, RealNodeOutputKey);
-						
-						if (FoundExistingSendLabel == nullptr)
-						{
-							CreateSendLabelWidget(EntryKey, RealNodeOutputKey);
-						}
+						ActiveOutputKeys.AddUnique(RealNodeOutputKey);
 					}
 				}
 
-				// Ensure any removed connections between active nodes are cleaned up
-				for (const TSharedPtr<SSignalFlowGraphSendLabel>& ExistingSendLabel : SendLabels)
-				{
-					if (!ActiveOutputKeys.Contains(ExistingSendLabel->GetSendReceiverEntryKey()))
-					{
-						SendLevelWidgets.RemoveSingle(EntryKey, ExistingSendLabel);
-						Canvas->RemoveSlot(ExistingSendLabel.ToSharedRef());
-					}
-				}
+				(*NodeWidget)->SetSendLabelWidgets(ActiveOutputKeys);
 
 				return;
 			}
@@ -1646,11 +1959,17 @@ namespace UE::Audio::Insights
 				const TSharedPtr<SSignalFlowGraphNode>* Node = NodeWidgets.Find(EntryKey);
 				if (Node && Node->IsValid())
 				{
-					const int32 NumParamsVisible = NodeDetailFilterSettings ? NodeDetailFilterSettings->GetNumVisibleParams() : 0;
+					const TSet<ESignalFlowNodeDetailParam> IgnoredParams{ ESignalFlowNodeDetailParam::SendOutputVolume };
+					const int32 NumParamsVisible = NodeDetailFilterSettings ? NodeDetailFilterSettings->GetNumVisibleParams(IgnoredParams) : 0;
+
+					const bool bSendLabelsAreVisible = ShowNodeDetails.Get() && NodeDetailFilterSettings && NodeDetailFilterSettings->GetParameterIsVisible(ESignalFlowNodeDetailParam::SendOutputVolume);
+					const int32 MaxNumVisibleSendLabels = bSendLabelsAreVisible ? MaxNumNodeOutputsInGraph : 0;
+
+					const int32 MaxNumParamRows = FMath::Max(NumParamsVisible, MaxNumVisibleSendLabels);
 
 					const FVector2D NodeSize = (*Node)->GetDesiredSize();
 					const float NodeGraphSlotWidth = GetNodeGraphSlotWidth(GetGraphOrientation());
-					const float NodeGraphSlotHeight = GetNodeGraphSlotHeight(GetGraphOrientation(), NumParamsVisible, ShowNodeDetails.Get());
+					const float NodeGraphSlotHeight = GetNodeGraphSlotHeight(GetGraphOrientation(), MaxNumParamRows, ShowNodeDetails.Get());
 
 					const float XOffset = (NodeGraphSlotWidth - NodeSize.X) / 2.0f;
 					const float YOffset = (NodeGraphSlotHeight - NodeSize.Y) / 2.0f;
@@ -1740,98 +2059,6 @@ namespace UE::Audio::Insights
 			];
 	}
 
-	void SSignalFlowGraph::CreateSendLabelWidget(const FSignalFlowEntryKey& SenderEntryKey, const FSignalFlowEntryKey& ReceiverEntryKey)
-	{
-		using namespace SSignalFlowGraphStyle;
-
-		if (FindSendLabel(SenderEntryKey, ReceiverEntryKey).IsValid())
-		{
-			return;
-		}
-
-		TSharedPtr<SSignalFlowGraphSendLabel> SendLabelWidget = SNew(SSignalFlowGraphSendLabel)
-			.PixelSnappingMethod(EWidgetPixelSnapping::SnapToPixel)
-			.SendReceiverEntryKey(ReceiverEntryKey)
-			.SenderPositionInGraph_Lambda([this, SenderEntryKey]() -> FVector2D
-			{
-				const TSharedPtr<SSignalFlowGraphNode>* SenderNode = NodeWidgets.Find(SenderEntryKey);
-				if (SenderNode && SenderNode->IsValid())
-				{
-					return (*SenderNode)->GetLocalOutputPos();
-				}
-
-				return FVector2D::ZeroVector;
-			})
-			.ReceiverPositionInGraph_Lambda([this, ReceiverEntryKey]() -> FVector2D
-			{
-				const TSharedPtr<SSignalFlowGraphNode>* ReceiverNode = NodeWidgets.Find(ReceiverEntryKey);
-				if (ReceiverNode && ReceiverNode->IsValid())
-				{
-					return (*ReceiverNode)->GetLocalInputPos();
-				}
-
-				return FVector2D::ZeroVector;
-			})
-			.SendLevel_Lambda([this, SenderEntryKey, ReceiverEntryKey]() -> TOptional<float>
-			{
-				const TSharedPtr<SSignalFlowGraphNode>* SenderNode = NodeWidgets.Find(SenderEntryKey);
-				if (SenderNode && SenderNode->IsValid() && (*SenderNode)->GetEntry().IsValid())
-				{
-					const TSharedPtr<FSignalFlowDashboardEntry>& SenderEntry = (*SenderNode)->GetEntry();
-					const FSignalFlowOutputData* SendInfo = SenderEntry->Outputs.Find(ReceiverEntryKey);
-					if (SendInfo && SendInfo->SendLevel.IsSet())
-					{
-						return SendInfo->SendLevel;
-					}
-				}
-
-				return NullOpt;
-			})
-			.GraphOrientation_Lambda([this]()
-			{
-				return Orientation;
-			})
-			.ZoomScaleFactor_Lambda([this]()
-			{
-				return ZoomScaleFactor;
-			});
-
-		SendLevelWidgets.Add(SenderEntryKey, SendLabelWidget);
-
-		Canvas->AddSlot()
-
-		.Position_Lambda([this, SenderEntryKey, ReceiverEntryKey]() -> FVector2D
-		{
-			const TSharedPtr<SSignalFlowGraphSendLabel> FoundExistingSendLabel = FindSendLabel(SenderEntryKey, ReceiverEntryKey);
-
-			if (FoundExistingSendLabel.IsValid())
-			{
-				FoundExistingSendLabel->UpdateLabelLayout();
-				return FoundExistingSendLabel->GetCachedPositionInGraph();
-			}
-
-			return FVector2D::ZeroVector;
-		})
-
-		.Size_Lambda([this, SenderEntryKey, ReceiverEntryKey]() -> FVector2D
-		{
-			const TSharedPtr<SSignalFlowGraphSendLabel> FoundExistingSendLabel = FindSendLabel(SenderEntryKey, ReceiverEntryKey);
-
-			if (FoundExistingSendLabel.IsValid())
-			{
-				return FoundExistingSendLabel->GetDesiredSize();
-			}
-
-			return FVector2D::ZeroVector;
-		})
-
-		.HAlign(HAlign_Left)
-		.VAlign(VAlign_Top)
-		[
-			SendLabelWidget.ToSharedRef()
-		];
-	}
-
 	TSharedPtr<SSignalFlowGraphNode> SSignalFlowGraph::FindRealNode(const FSignalFlowEntryKey& EntryKey, const ENodeConnectionType ConnectionType) const
 	{
 		const TSharedPtr<SSignalFlowGraphNode>* FoundNode = NodeWidgets.Find(EntryKey);
@@ -1871,29 +2098,6 @@ namespace UE::Audio::Insights
 		return NullOpt;
 	}
 
-	TSharedPtr<SSignalFlowGraphSendLabel> SSignalFlowGraph::FindSendLabel(const FSignalFlowEntryKey& SenderKey, const FSignalFlowEntryKey& ReceiverKey) const
-	{
-		TArray<TSharedPtr<SSignalFlowGraphSendLabel>> SendLabels;
-		SendLevelWidgets.MultiFind(SenderKey, SendLabels);
-
-		const TSharedPtr<SSignalFlowGraphSendLabel>* FoundSendLabel = FindSendLabelInList(SendLabels, ReceiverKey);
-
-		if (FoundSendLabel && FoundSendLabel->IsValid())
-		{
-			return *FoundSendLabel;
-		}
-
-		return nullptr;
-	}
-
-	const TSharedPtr<SSignalFlowGraphSendLabel>* SSignalFlowGraph::FindSendLabelInList(const TArray<TSharedPtr<SSignalFlowGraphSendLabel>>& SendLabelWidgets, const FSignalFlowEntryKey& ReceiverKey) const
-	{
-		return SendLabelWidgets.FindByPredicate([&ReceiverKey](const TSharedPtr<SSignalFlowGraphSendLabel>& SendLabel)
-		{
-			return SendLabel.IsValid() && SendLabel->GetSendReceiverEntryKey() == ReceiverKey;
-		});
-	}
-
 	SSignalFlowGraph::EAnimationResult SSignalFlowGraph::AnimateScrollTo(const TSharedRef<SScrollBox>& ScrollBox, const FVector2D& ScrollToPos, const float DeltaTime)
 	{
 		using namespace SSignalFlowGraphStyle;
diff --git a/Engine/Plugins/AudioInsights/Source/AudioInsights/Private/Views/SignalFlowDashboardViewFactory.cpp b/Engine/Plugins/AudioInsights/Source/AudioInsights/Private/Views/SignalFlowDashboardViewFactory.cpp
--- a/Engine/Plugins/AudioInsights/Source/AudioInsights/Private/Views/SignalFlowDashboardViewFactory.cpp
+++ b/Engine/Plugins/AudioInsights/Source/AudioInsights/Private/Views/SignalFlowDashboardViewFactory.cpp
@@ -63,16 +63,19 @@ namespace UE::Audio::Insights
 					return LOCTEXT("AudioDashboard_SignalFlow_Param_HPFFreq", "HPF Freq");
 
 				case ESignalFlowNodeDetailParam::Priority:
-					return LOCTEXT("AudioDashboard_SignalFlow_Param_Priority", "Priority");;
+					return LOCTEXT("AudioDashboard_SignalFlow_Param_Priority", "Priority");
 
 				case ESignalFlowNodeDetailParam::Distance:
-					return LOCTEXT("AudioDashboard_SignalFlow_Param_Distance", "Distance");;
+					return LOCTEXT("AudioDashboard_SignalFlow_Param_Distance", "Distance");
 
 				case ESignalFlowNodeDetailParam::Attenuation:
-					return LOCTEXT("AudioDashboard_SignalFlow_Param_Attenuation", "Distance/Occlusion Attenuation");;
+					return LOCTEXT("AudioDashboard_SignalFlow_Param_Attenuation", "Distance/Occlusion Attenuation");
 
 				case ESignalFlowNodeDetailParam::RelativeRenderCost:
-					return LOCTEXT("AudioDashboard_SignalFlow_Param_RelativeRanderCost", "Rel. Render Cost");;
+					return LOCTEXT("AudioDashboard_SignalFlow_Param_RelativeRanderCost", "Rel. Render Cost");
+
+				case ESignalFlowNodeDetailParam::SendOutputVolume:
+					return LOCTEXT("AudioDashboard_SignalFlow_Param_SendOutputVolume", "Send Output Vol.");
 
 				default:
 					ensureMsgf(false, TEXT("Unrecognized ESignalFlowNodeDetailParam value"));
diff --git a/Engine/Plugins/AudioInsights/Source/AudioInsights/Public/AudioInsightsStyle.h b/Engine/Plugins/AudioInsights/Source/AudioInsights/Public/AudioInsightsStyle.h
--- a/Engine/Plugins/AudioInsights/Source/AudioInsights/Public/AudioInsightsStyle.h
+++ b/Engine/Plugins/AudioInsights/Source/AudioInsights/Public/AudioInsightsStyle.h
@@ -272,12 +272,14 @@ namespace UE::Audio::Insights
 			Set("AudioInsights.Icon.SignalFlow.AlignVerticalLeft",	   new IMAGE_BRUSH_SVG(TEXT("Icons/SignalFlow/align_vertical_left"),	 Icon20));
 			Set("AudioInsights.Icon.SignalFlow.AlignVerticalRight",	   new IMAGE_BRUSH_SVG(TEXT("Icons/SignalFlow/align_vertical_right"),	 Icon20));
 
-			Set("AudioInsights.Icon.SignalFlow.CenterView",			  new IMAGE_BRUSH_SVG(TEXT("Icons/SignalFlow/center_view"),			 Icon20));
-			Set("AudioInsights.Icon.SignalFlow.HorizontalFlow",		  new IMAGE_BRUSH_SVG(TEXT("Icons/SignalFlow/horizontal_flow"),		 Icon20));
-			Set("AudioInsights.Icon.SignalFlow.NodeDetails",		  new IMAGE_BRUSH_SVG(TEXT("Icons/SignalFlow/node_details"),		 Icon20));
-			Set("AudioInsights.Icon.SignalFlow.Settings",			  new IMAGE_BRUSH_SVG(TEXT("Icons/SignalFlow/settings"),			 Icon20));
+			Set("AudioInsights.Icon.SignalFlow.CenterView",			  new IMAGE_BRUSH_SVG(TEXT("Icons/SignalFlow/center_view"),				 Icon20));
+			Set("AudioInsights.Icon.SignalFlow.HorizontalFlow",		  new IMAGE_BRUSH_SVG(TEXT("Icons/SignalFlow/horizontal_flow"),			 Icon20));
+			Set("AudioInsights.Icon.SignalFlow.NodeDetails",		  new IMAGE_BRUSH_SVG(TEXT("Icons/SignalFlow/node_details"),			 Icon20));
+			Set("AudioInsights.Icon.SignalFlow.Settings",			  new IMAGE_BRUSH_SVG(TEXT("Icons/SignalFlow/settings"),				 Icon20));
 			
-			Set("AudioInsights.Icon.SignalFlow.Pill",				  new IMAGE_BRUSH_SVG(TEXT("Icons/SignalFlow/pill"),				 Icon16));
+			Set("AudioInsights.Icon.SignalFlow.Pill",				  new IMAGE_BRUSH_SVG(TEXT("Icons/SignalFlow/pill"),					 Icon16));
+			Set("AudioInsights.Icon.SignalFlow.Connected",			  new IMAGE_BRUSH_SVG(TEXT("Icons/SignalFlow/pin_trigger_connected"),	 Icon16));
+			Set("AudioInsights.Icon.SignalFlow.Disconnected",		  new IMAGE_BRUSH_SVG(TEXT("Icons/SignalFlow/pin_trigger_disconnected"), Icon16));
 
 			/* Modulation */
 			// Icons
diff --git a/Engine/Plugins/AudioInsights/Source/AudioInsights/Public/Messages/SignalFlowEntryKey.h b/Engine/Plugins/AudioInsights/Source/AudioInsights/Public/Messages/SignalFlowEntryKey.h
--- a/Engine/Plugins/AudioInsights/Source/AudioInsights/Public/Messages/SignalFlowEntryKey.h
+++ b/Engine/Plugins/AudioInsights/Source/AudioInsights/Public/Messages/SignalFlowEntryKey.h
@@ -34,6 +34,7 @@ namespace UE::Audio::Insights
 		Distance,
 		Attenuation,
 		RelativeRenderCost,
+		SendOutputVolume,
 
 		///
 		MAX
diff --git a/Engine/Plugins/AudioInsights/Source/AudioInsights/Public/Settings/SignalFlowSettings.h b/Engine/Plugins/AudioInsights/Source/AudioInsights/Public/Settings/SignalFlowSettings.h
--- a/Engine/Plugins/AudioInsights/Source/AudioInsights/Public/Settings/SignalFlowSettings.h
+++ b/Engine/Plugins/AudioInsights/Source/AudioInsights/Public/Settings/SignalFlowSettings.h
@@ -52,7 +52,7 @@ struct FSignalFlowNodeDetailFilterSettings
 	GENERATED_BODY()
 
 	bool GetParameterIsVisible(const UE::Audio::Insights::ESignalFlowNodeDetailParam Param) const;
-	int32 GetNumVisibleParams() const;
+	int32 GetNumVisibleParams(const TSet<UE::Audio::Insights::ESignalFlowNodeDetailParam>& IgnoredParams = TSet<UE::Audio::Insights::ESignalFlowNodeDetailParam>()) const;
 	ECheckBoxState GetShowAllNodeDetailFiltersCheckboxState() const;
 
 	void ToggleEnableAllNodeDetailFilters();
@@ -93,6 +93,10 @@ struct FSignalFlowNodeDetailFilterSettings
 	/** Toggles visibility of relative render cost parameters. */
 	UPROPERTY(EditAnywhere, config, Category = SignalFlow)
 	bool bRelativeRenderCost = false;
+
+	/** Toggles visibility of send output volume parameters on node output pins. */
+	UPROPERTY(EditAnywhere, config, Category = SignalFlow)
+	bool bSendOutputVolume = true;
 };
 
 USTRUCT()
diff --git a/Engine/Plugins/AudioInsights/Source/AudioInsights/Public/Views/SSignalFlowGraph.h b/Engine/Plugins/AudioInsights/Source/AudioInsights/Public/Views/SSignalFlowGraph.h
--- a/Engine/Plugins/AudioInsights/Source/AudioInsights/Public/Views/SSignalFlowGraph.h
+++ b/Engine/Plugins/AudioInsights/Source/AudioInsights/Public/Views/SSignalFlowGraph.h
@@ -16,7 +16,10 @@ struct FSignalFlowNodeDetailFilterSettings;
 
 class SBorder;
 class SHorizontalBox;
+class SImage;
+class SOverlay;
 class SScrollBox;
+class STextBlock;
 class SVerticalBox;
 
 #if WITH_EDITOR
@@ -95,11 +98,18 @@ namespace UE::Audio::Insights
 		UE_API virtual FVector2D ComputeDesiredSize(float InLayoutScaleMultiplier) const override;
 
 		UE_API FVector2D GetLocalInputPos() const;
-		UE_API FVector2D GetLocalOutputPos() const;
+		UE_API FVector2D GetLocalCenteredOutputPos() const;
+		UE_API FVector2D GetLocalSendLabelPos(const FSignalFlowEntryKey& ReceiverEntryKey, const int32 LabelIndex) const;
+		UE_API int32 GetNumSendLabels() const;
 
 		UE_API void UpdateNodeLayout();
 		UE_API FVector2D GetCachedPositionInGraph() const;
 
+		UE_API void SetSendLabelWidgets(const TArray<FSignalFlowEntryKey>& ReceiverEntryKeys);
+		UE_API void AddSendLabelWidget(const FSignalFlowEntryKey& ReceiverEntryKey);
+		UE_API void RemoveSendLabelWidget(const FSignalFlowEntryKey& ReceiverEntryKey);
+		UE_API bool ContainsSendLabel(const FSignalFlowEntryKey& ReceiverEntryKey) const;
+
 		const FSignalFlowEntryKey& GetNodeEntryKey() const { return NodeEntryKey; }
 		const TSharedPtr<FSignalFlowDashboardEntry> GetEntry() const { return Entry; }
 
@@ -113,8 +123,11 @@ namespace UE::Audio::Insights
 		TSharedRef<SHorizontalBox> CreateNodeName(const SSignalFlowGraphNode::FArguments& InArgs);
 		TSharedRef<SVerticalBox> CreateNodeDetails() const;
 		TSharedRef<SHorizontalBox> CreateNodeDetailsRow(const ESignalFlowNodeDetailParam ParamType) const;
+		TSharedRef<SVerticalBox> CreateSendOutputsVerticalContainer();
+		TSharedRef<SHorizontalBox> CreateSendOutputsHorizontalContainer();
 
 		bool NodeDetailsExpanderArrowIsEnabled() const;
+		bool SendLabelsAreVisible() const;
 
 		TSharedPtr<FSignalFlowDashboardEntry> Entry;
 		TAttribute<FVector2d> PositionInGraph;
@@ -126,6 +139,8 @@ namespace UE::Audio::Insights
 		TAttribute<bool> DisplayAmpPeakInDb;
 		TAttribute<EOrientation> GraphOrientation;
 
+		TMap<FSignalFlowEntryKey, TSharedRef<class SSignalFlowGraphSendLabel>> SendLabels;
+
 		FOnSignalFlowNodeSelected OnNodeSelected;
 
 #if WITH_EDITOR
@@ -140,6 +155,9 @@ namespace UE::Audio::Insights
 
 		TSharedPtr<SBorder> SelectedNodeBorder;
 		TSharedPtr<SBorder> NodeBorder;
+		TSharedPtr<SHorizontalBox> NodeLabelHorizontalBox;
+		TSharedPtr<SVerticalBox> SendLevelsVerticalBox;
+		TSharedPtr<SHorizontalBox> SendLevelsHorizontalBox;
 
 		bool bNodeDetailsAreExpanded = true;
 	};
@@ -148,18 +166,15 @@ namespace UE::Audio::Insights
 	{
 	public:
 		SLATE_BEGIN_ARGS(SSignalFlowGraphSendLabel)
-			: _SendReceiverEntryKey(),
-			_SenderPositionInGraph(FVector2d::ZeroVector),
-			_ReceiverPositionInGraph(FVector2d::ZeroVector),
+			: _ReceiverEntryType(ESignalFlowEntryType::AudioDevice),
 			_SendLevel(),
 			_GraphOrientation(EOrientation::Orient_Vertical),
 			_ZoomScaleFactor(1.0f)
 		{
 		}
 
-		SLATE_ARGUMENT(FSignalFlowEntryKey, SendReceiverEntryKey)
-		SLATE_ATTRIBUTE(FVector2d, SenderPositionInGraph)
-		SLATE_ATTRIBUTE(FVector2d, ReceiverPositionInGraph)
+		SLATE_ARGUMENT(ESignalFlowEntryType, ReceiverEntryType)
+
 		SLATE_ATTRIBUTE(TOptional<float>, SendLevel)
 		SLATE_ATTRIBUTE(EOrientation, GraphOrientation)
 		SLATE_ATTRIBUTE(float, ZoomScaleFactor)
@@ -167,24 +182,20 @@ namespace UE::Audio::Insights
 		SLATE_END_ARGS()
 
 		UE_API void Construct(const FArguments& InArgs);
+
 		UE_API virtual FVector2D ComputeDesiredSize(float InLayoutScaleMultiplier) const override;
-		
-		UE_API void UpdateLabelLayout();
-		UE_API FVector2D GetCachedPositionInGraph() const;
 
-		const FSignalFlowEntryKey& GetSendReceiverEntryKey() const { return SendReceiverEntryKey; }
+		UE_API bool ValueIsSet() const;
 
 	private:
-		FSignalFlowEntryKey SendReceiverEntryKey;
-		TAttribute<FVector2d> SenderPositionInGraph;
-		TAttribute<FVector2d> ReceiverPositionInGraph;
+		TSharedRef<STextBlock> CreateOutputLabelText(const FText& ToolTipText) const;
+		TSharedRef<SOverlay> CreatePinIconWidget(const FText& ToolTipText, const ESignalFlowEntryType ReceiverEntryType);
+
 		TAttribute<TOptional<float>> SendLevel;
 		TAttribute<EOrientation> GraphOrientation;
 		TAttribute<float> ZoomScaleFactor;
 
-		FVector2d CachedPositionInGraph = FVector2d::Zero();
-
-		TSharedPtr<SBorder> LabelBorder;
+		TSharedPtr<SImage> OutputPinImage;
 	};
 
 	class SSignalFlowGraph : public SCompoundWidget
@@ -284,7 +295,6 @@ namespace UE::Audio::Insights
 		float CalculateGraphSizeIncreaseAtBottom(const TArray<TreeDepthPair>& NewTreeDepthOrder, const float OriginalGraphHeight, const float OriginalViewportPosY) const;
 
 		void CreateGraphNode(const TSharedPtr<ISignalFlowNode>& Node, const float MinXPos, const float MaxXPos, const int32 DepthTracker, TArray<FSignalFlowEntryKey>& OutRemainingNodeKeys);
-		void CreateSendLabelWidget(const FSignalFlowEntryKey& SenderEntryKey, const FSignalFlowEntryKey& ReceiverEntryKey);
 
 		enum class ENodeConnectionType : uint8
 		{
@@ -295,9 +305,6 @@ namespace UE::Audio::Insights
 		TSharedPtr<SSignalFlowGraphNode> FindRealNode(const FSignalFlowEntryKey& EntryKey, const ENodeConnectionType ConnectionType) const;
 		TOptional<FSignalFlowEntryKey> FindRealNodeKey(const FSignalFlowEntryKey& EntryKey, const ENodeConnectionType ConnectionType) const;
 
-		TSharedPtr<SSignalFlowGraphSendLabel> FindSendLabel(const FSignalFlowEntryKey& SenderKey, const FSignalFlowEntryKey& ReceiverKey) const;
-		const TSharedPtr<SSignalFlowGraphSendLabel>* FindSendLabelInList(const TArray<TSharedPtr<SSignalFlowGraphSendLabel>>& SendLabelWidgets, const FSignalFlowEntryKey& ReceiverKey) const;
-
 		enum class EAnimationResult : uint8
 		{
 			InProgress,
@@ -321,7 +328,6 @@ namespace UE::Audio::Insights
 		};
 
 		TMap<FSignalFlowEntryKey, FConnectedNodeKeys> DummyToRealConnectionKeys;
-		TMultiMap<FSignalFlowEntryKey, TSharedPtr<SSignalFlowGraphSendLabel>> SendLevelWidgets;
 
 		TArray<TSharedPtr<ISignalFlowNode>>* ListSource = nullptr;
 		FSignalFlowNodeDetailFilterSettings* NodeDetailFilterSettings = nullptr;
@@ -346,6 +352,8 @@ namespace UE::Audio::Insights
 
 		FGraphLimits CachedGraphWidthLimits;
 		TArray<struct TreeDepthPair> CachedTreeDepthPairs;
+		int32 MaxNumNodeOutputsInGraph = 0;
+
 		float ZoomScaleFactor = 1.0f;
 		float WireSplineAmplitudePowerFactor = 0.4f;
 		float WireSplineThicknessScalarMax = 14.0f;
