[Audio Insights] Don't refresh signal flow graph if selecting background when nothing is currently selected

#Synced-CL 51049241

#Synced-CL 51049241

diff --git a/Engine/Plugins/AudioInsights/Source/AudioInsights/Private/Views/SSignalFlowGraph.cpp b/Engine/Plugins/AudioInsights/Source/AudioInsights/Private/Views/SSignalFlowGraph.cpp
--- a/Engine/Plugins/AudioInsights/Source/AudioInsights/Private/Views/SSignalFlowGraph.cpp
+++ b/Engine/Plugins/AudioInsights/Source/AudioInsights/Private/Views/SSignalFlowGraph.cpp
@@ -1328,7 +1328,7 @@ namespace UE::Audio::Insights
 		{
 			return FReply::Handled().DetectDrag(SharedThis(this), EKeys::RightMouseButton);
 		}
-		else if (MouseEvent.GetEffectingButton() == EKeys::LeftMouseButton)
+		else if (MouseEvent.GetEffectingButton() == EKeys::LeftMouseButton && SelectedItem.Get().IsValid())
 		{
 			OnNodeSelected.ExecuteIfBound(nullptr);
 
