[Audio Insights] Improvements to node ordering in signal flow
* Move higher timestamped branches to the bottom/right of the graph (instead of the top/left)
* Use the lowest tree depth value as a metric to keep branches to one side (instead of if a node has no inputs)
* Lock nodes with no output to the bottom/right (instead of top/left)

#Synced-CL 51049241

diff --git a/Engine/Plugins/AudioInsights/Source/AudioInsights/Private/Views/SignalFlowDashboardViewFactory.cpp b/Engine/Plugins/AudioInsights/Source/AudioInsights/Private/Views/SignalFlowDashboardViewFactory.cpp
--- a/Engine/Plugins/AudioInsights/Source/AudioInsights/Private/Views/SignalFlowDashboardViewFactory.cpp
+++ b/Engine/Plugins/AudioInsights/Source/AudioInsights/Private/Views/SignalFlowDashboardViewFactory.cpp
@@ -1789,38 +1789,34 @@ namespace UE::Audio::Insights
 
 			if (FirstNode && SecondNode && FirstNode->IsValid() && SecondNode->IsValid())
 			{
-				const TreeDepthPair FirstNodeTreeDepth = GetOriginTreeDepthForNode((*FirstNode));
-				const TreeDepthPair SecondNodeTreeDepth = GetOriginTreeDepthForNode((*SecondNode));
+				const TreeDepthPair FirstNodeTreeDepth = GetRealTreeDepthForNode((*FirstNode));
+				const TreeDepthPair SecondNodeTreeDepth = GetRealTreeDepthForNode((*SecondNode));
 
 				// First, sort by tree depth
 				if (FirstNodeTreeDepth != SecondNodeTreeDepth)
 				{
-					return FirstNodeTreeDepth < SecondNodeTreeDepth;
+					return FirstNodeTreeDepth > SecondNodeTreeDepth;
 				}
 
-				// Move nodes that have no inputs to one side of the graph
-				const int32 FirstNumInputs = (*FirstNode)->FilteredInputs.Num();
-				const int32 SecondNumInputs = (*SecondNode)->FilteredInputs.Num();
+				// Next, try to order nodes by the minimum timestamp of their parent sound source/owner object nodes (if any exist)
+				double FirstBranchMinTimestamp = TNumericLimits<double>::Max();
+				TreeDepthPair FirstBranchMinTreeDepth = FirstNodeTreeDepth;
+				GetBranchMinTimestampAndTreeDepthRecursive((*FirstNode)->FilteredInputs, FirstBranchMinTreeDepth, FirstBranchMinTimestamp);
 
-				const bool bNumInputsMatch = FirstNumInputs == SecondNumInputs;
-				const bool bEitherNodeHasZeroInputs = (FirstNumInputs == 0 || SecondNumInputs == 0);
+				double SecondInputTimestamp = TNumericLimits<double>::Max();
+				TreeDepthPair SecondBranchMinTreeDepth = SecondNodeTreeDepth;
+				GetBranchMinTimestampAndTreeDepthRecursive((*SecondNode)->FilteredInputs, SecondBranchMinTreeDepth, SecondInputTimestamp);
 
-				if (bEitherNodeHasZeroInputs && !bNumInputsMatch)
+				// Move branches that are not connected all the way up to the lowest tree depth to one side of the graph
+				if (FirstBranchMinTreeDepth != SecondBranchMinTreeDepth)
 				{
-					return FirstNumInputs > SecondNumInputs;
+					return FirstBranchMinTreeDepth > SecondBranchMinTreeDepth;
 				}
 
-				// Next, try to order nodes by the minimum timestamp of their parent sound source/owner object nodes (if any exist)
-				double FirstInputTimestamp = TNumericLimits<double>::Max();
-				double SecondInputTimestamp = TNumericLimits<double>::Max();
-
-				GetMinTimestampRecursive((*FirstNode)->FilteredInputs, FirstInputTimestamp);
-				GetMinTimestampRecursive((*SecondNode)->FilteredInputs, SecondInputTimestamp);
-
 				// Sort by timestamp of the parent nodes
-				if (!FMath::IsNearlyEqual(FirstInputTimestamp, SecondInputTimestamp))
+				if (!FMath::IsNearlyEqual(FirstBranchMinTimestamp, SecondInputTimestamp))
 				{
-					return FirstInputTimestamp < SecondInputTimestamp;
+					return FirstBranchMinTimestamp > SecondInputTimestamp;
 				}
 
 				// Try timestamp of this node
@@ -1828,7 +1824,7 @@ namespace UE::Audio::Insights
 				{
 					if (!FMath::IsNearlyEqual((*FirstNode)->GetTimestamp(), (*SecondNode)->GetTimestamp()))
 					{
-						return (*FirstNode)->GetTimestamp() < (*SecondNode)->GetTimestamp();
+						return (*FirstNode)->GetTimestamp() > (*SecondNode)->GetTimestamp();
 					}
 				}
 
@@ -1845,7 +1841,7 @@ namespace UE::Audio::Insights
 		Node->FilteredInputs.Sort(SortInput);
 	}
 
-	TreeDepthPair FSignalFlowDashboardViewFactory::GetOriginTreeDepthForNode(const TSharedPtr<ISignalFlowNode>& Node) const
+	TreeDepthPair FSignalFlowDashboardViewFactory::GetRealTreeDepthForNode(const TSharedPtr<ISignalFlowNode>& Node) const
 	{
 		// Either grab the tree depth for this node,
 		// or if it's a dummy connection node, grab the tree depth of the original input to the dummy connection
@@ -1864,7 +1860,7 @@ namespace UE::Audio::Insights
 		return Node->GetTreeDepth();
 	}
 
-	void FSignalFlowDashboardViewFactory::GetMinTimestampRecursive(const TArray<FSignalFlowEntryKey>& NodeKeys, double& OutMinTimestamp) const
+	void FSignalFlowDashboardViewFactory::GetBranchMinTimestampAndTreeDepthRecursive(const TArray<FSignalFlowEntryKey>& NodeKeys, TreeDepthPair& OutMinTreeDepth, double& OutMinTimestamp) const
 	{
 		for (const FSignalFlowEntryKey& NodeKey : NodeKeys)
 		{
@@ -1872,13 +1868,15 @@ namespace UE::Audio::Insights
 
 			if (Node && Node->IsValid())
 			{
+				OutMinTreeDepth = FMath::Min(OutMinTreeDepth, (*Node)->GetTreeDepth());
+
 				// Only count entries that are transient when finding the minimum timestamp of a branch
 				if (NodeIsOwnerOrSound(**Node))
 				{
 					OutMinTimestamp = FMath::Min((*Node)->GetTimestamp(), OutMinTimestamp);
 				}
 				
-				GetMinTimestampRecursive((*Node)->FilteredInputs, OutMinTimestamp);
+				GetBranchMinTimestampAndTreeDepthRecursive((*Node)->FilteredInputs, OutMinTreeDepth, OutMinTimestamp);
 			}
 		}
 	}
@@ -1897,7 +1895,7 @@ namespace UE::Audio::Insights
 			if (Node.IsValid() && !Node->NodeOrderIsValid())
 			{
 				// Force nodes with no outputs to one side of the graph
-				Node->SetNodeOrderID(TNumericLimits<int32>::Max());
+				Node->SetNodeOrderID(TNumericLimits<int32>::Lowest());
 			}
 		}
 	}
@@ -1925,7 +1923,7 @@ namespace UE::Audio::Insights
 			}
 
 			// Fallback to timestamp
-			return First->GetTimestamp() < Second->GetTimestamp();
+			return First->GetTimestamp() > Second->GetTimestamp();
 		};
 
 		FilteredGraphNodes.Sort(SortSignalNodes);
diff --git a/Engine/Plugins/AudioInsights/Source/AudioInsights/Public/Views/SignalFlowDashboardViewFactory.h b/Engine/Plugins/AudioInsights/Source/AudioInsights/Public/Views/SignalFlowDashboardViewFactory.h
--- a/Engine/Plugins/AudioInsights/Source/AudioInsights/Public/Views/SignalFlowDashboardViewFactory.h
+++ b/Engine/Plugins/AudioInsights/Source/AudioInsights/Public/Views/SignalFlowDashboardViewFactory.h
@@ -125,8 +125,8 @@ namespace UE::Audio::Insights
 		
 		void AssignNodeOrderRecursive(TSharedPtr<ISignalFlowNode> OutNode, const int32 Depth, int32& OutNodeProcessingOrder);
 		void SortNodeInputs(TSharedPtr<ISignalFlowNode> OutNode);
-		TreeDepthPair GetOriginTreeDepthForNode(const TSharedPtr<ISignalFlowNode>& Node) const;
-		void GetMinTimestampRecursive(const TArray<FSignalFlowEntryKey>& NodeKeys, double& OutMinTimestamp) const;
+		TreeDepthPair GetRealTreeDepthForNode(const TSharedPtr<ISignalFlowNode>& Node) const;
+		void GetBranchMinTimestampAndTreeDepthRecursive(const TArray<FSignalFlowEntryKey>& NodeKeys, TreeDepthPair& OutMinTreeDepth, double& OutMinTimestamp) const;
 		bool NodeIsOwnerOrSound(const ISignalFlowNode& Node) const;
 
 		void FixUnassignedNodeOrders();
