[Audio Insights] Improvements to Justification settings and positioning algorithm
* First pass on Justification setting would result in strange paths being generated between nodes.
* This change aims to both simplify the justification settings, as well as make the positioning algorithm more robust
* Further changes to the positioning algorithm are expected (in particular - improvements to nodes retaining previous positions where possible, collapsing the space dummy nodes take up, and revisiting whether FSignalFlowDashboardViewFactory::CalcXPositionsForRows can be simplified)

Simplified Justification settings
* Bottom and Right justifications have been removed - these would align nodes to the side of the graph that has the most volitility, which is undesirable.
* Merged the seperate horizontal and vertical justification settings into one - Edge justified and Center justified.
* Updated hotkeys for Edge (crtl + E) and Center (ctrl + C) justification.

Updated positioning algorithm (FSignalFlowDashboardViewFactory::CalcXPosition...)
* Removed bottom/right positioning alg (AlignNodeTowardsMaxX) and merged horizontal and vertical paths through FSignalFlowDashboardViewFactory::CalcXPositionForNode
* FSignalFlowDashboardViewFactory::AlignNodeTowardsLowestX is now a very simple XPos = IndexInRow process (instead of aligning relative to the left side of the node's inputs/outputs. Dummy (invisible) nodes are included in this process, so the nodes still maintain gaps for long connections to pass through.
* FSignalFlowDashboardViewFactory::AlignNodeTowardsCenterX still positions nodes in the average position of it's inputs/outputs, but now has two major improvements:
** Fixed nodes with no inputs/outputs to locate around being skipped - we now just place below the previously placed node position (or zero if this is the first node in the row)
** If an overlap is detected with the previous node, we now shift all nodes in the row backwards until an overlap is no longer detected

#Synced-CL 51049241

diff --git a/Engine/Plugins/AudioInsights/Source/AudioInsights/Private/SignalFlowEditorCommands.cpp b/Engine/Plugins/AudioInsights/Source/AudioInsights/Private/SignalFlowEditorCommands.cpp
--- a/Engine/Plugins/AudioInsights/Source/AudioInsights/Private/SignalFlowEditorCommands.cpp
+++ b/Engine/Plugins/AudioInsights/Source/AudioInsights/Private/SignalFlowEditorCommands.cpp
@@ -17,12 +17,10 @@ namespace UE::Audio::Insights
 	void FSignalFlowEditorCommands::RegisterCommands()
 	{
 		UI_COMMAND(ToggleHorizontalView, "Horizontal Flow", "Toggles between horizontally/vertically orienting the Signal Flow graph.", EUserInterfaceActionType::ToggleButton, FInputChord(EKeys::H, EModifierKey::Control));
-		UI_COMMAND(JustifyBottom, "Bottom", "Justify the Signal Flow Graph towards the bottom in horizontal mode.", EUserInterfaceActionType::RadioButton, FInputChord(EKeys::F1, EModifierKey::Control));
-		UI_COMMAND(JustifyHorizontalCenter, "Center", "Justify the Signal Flow Graph around the center in horizontal mode.", EUserInterfaceActionType::RadioButton, FInputChord(EKeys::F2, EModifierKey::Control));
-		UI_COMMAND(JustifyTop, "Top", "Justify the Signal Flow Graph towards the top in horizontal mode.", EUserInterfaceActionType::RadioButton, FInputChord(EKeys::F3, EModifierKey::Control));
-		UI_COMMAND(JustifyLeft, "Left", "Justify the Signal Flow Graph towards the left in vertical mode.", EUserInterfaceActionType::RadioButton, FInputChord(EKeys::F4, EModifierKey::Control));
-		UI_COMMAND(JustifyVerticalCenter, "Center", "Justify the Signal Flow Graph around the center in vertical mode.", EUserInterfaceActionType::RadioButton, FInputChord(EKeys::F5, EModifierKey::Control));
-		UI_COMMAND(JustifyRight, "Right", "Justify the Signal Flow Graph towards the right in vertical mode.", EUserInterfaceActionType::RadioButton, FInputChord(EKeys::F6, EModifierKey::Control));
+
+		UI_COMMAND(JustifyEdge, "Edge", "Position nodes towards one edge of the graph.", EUserInterfaceActionType::RadioButton, FInputChord(EKeys::E, EModifierKey::Control));
+		UI_COMMAND(JustifyCenter, "Center", "Position nodes around the center of the graph.", EUserInterfaceActionType::RadioButton, FInputChord(EKeys::C, EModifierKey::Control));
+
 		UI_COMMAND(ShowNodeDetails, "Show Node Details", "Toggles visibility of node details.", EUserInterfaceActionType::ToggleButton, FInputChord(EKeys::D, EModifierKey::Control));
 		UI_COMMAND(PauseOnSelect, "Pause on Select", "Toggles whether selecting a node should pause processing trace data, freezing audio insights at the latest timestamp.", EUserInterfaceActionType::ToggleButton, FInputChord(EKeys::P, EModifierKey::Control));
 		UI_COMMAND(ToggleAnimateWires, "Animate Wires", "Toggles whether connection wire animations driven by amplitude envelope values are enabled in the Signal Flow graph.", EUserInterfaceActionType::ToggleButton, FInputChord());
diff --git a/Engine/Plugins/AudioInsights/Source/AudioInsights/Private/SignalFlowEditorCommands.h b/Engine/Plugins/AudioInsights/Source/AudioInsights/Private/SignalFlowEditorCommands.h
--- a/Engine/Plugins/AudioInsights/Source/AudioInsights/Private/SignalFlowEditorCommands.h
+++ b/Engine/Plugins/AudioInsights/Source/AudioInsights/Private/SignalFlowEditorCommands.h
@@ -17,12 +17,8 @@ namespace UE::Audio::Insights
 
 		// Toggles
 		TSharedPtr<const FUICommandInfo> GetToggleHorizontalViewCommand() const { return ToggleHorizontalView; }
-		TSharedPtr<const FUICommandInfo> GetJustifyBottomCommand() const { return JustifyBottom; }
-		TSharedPtr<const FUICommandInfo> GetJustifyHorizontalCenterCommand() const { return JustifyHorizontalCenter; }
-		TSharedPtr<const FUICommandInfo> GetJustifyTopCommand() const { return JustifyTop; }
-		TSharedPtr<const FUICommandInfo> GetJustifyLeftCommand() const { return JustifyLeft; }
-		TSharedPtr<const FUICommandInfo> GetJustifyVerticalCenterCommand() const { return JustifyVerticalCenter; }
-		TSharedPtr<const FUICommandInfo> GetJustifyRightCommand() const { return JustifyRight; }
+		TSharedPtr<const FUICommandInfo> GetJustifyEdgeCommand() const { return JustifyEdge; }
+		TSharedPtr<const FUICommandInfo> GetJustifyCenterCommand() const { return JustifyCenter; }
 		TSharedPtr<const FUICommandInfo> GetShowNodeDetailsCommand() const { return ShowNodeDetails; }
 		TSharedPtr<const FUICommandInfo> GetPauseOnSelectCommand() const { return PauseOnSelect; }
 		TSharedPtr<const FUICommandInfo> GetToggleAnimateWiresCommand() const { return ToggleAnimateWires; }
@@ -34,12 +30,10 @@ namespace UE::Audio::Insights
 	private:
 		// Toggles
 		TSharedPtr<FUICommandInfo> ToggleHorizontalView;
-		TSharedPtr<FUICommandInfo> JustifyBottom;
-		TSharedPtr<FUICommandInfo> JustifyHorizontalCenter;
-		TSharedPtr<FUICommandInfo> JustifyTop;
-		TSharedPtr<FUICommandInfo> JustifyLeft;
-		TSharedPtr<FUICommandInfo> JustifyVerticalCenter;
-		TSharedPtr<FUICommandInfo> JustifyRight;
+
+		TSharedPtr<FUICommandInfo> JustifyEdge;
+		TSharedPtr<FUICommandInfo> JustifyCenter;
+
 		TSharedPtr<FUICommandInfo> ShowNodeDetails;
 		TSharedPtr<FUICommandInfo> PauseOnSelect;
 		TSharedPtr<FUICommandInfo> ToggleAnimateWires;
diff --git a/Engine/Plugins/AudioInsights/Source/AudioInsights/Private/Views/SignalFlowDashboardViewFactory.cpp b/Engine/Plugins/AudioInsights/Source/AudioInsights/Private/Views/SignalFlowDashboardViewFactory.cpp
--- a/Engine/Plugins/AudioInsights/Source/AudioInsights/Private/Views/SignalFlowDashboardViewFactory.cpp
+++ b/Engine/Plugins/AudioInsights/Source/AudioInsights/Private/Views/SignalFlowDashboardViewFactory.cpp
@@ -292,40 +292,16 @@ namespace UE::Audio::Insights
 			FIsActionChecked::CreateLambda([this]() { return GraphWidgetContents.IsValid() && GraphWidgetContents->GetGraphOrientation() == Orient_Horizontal; }));
 
 		CommandList->MapAction(
-			Commands.GetJustifyBottomCommand(),
-			FExecuteAction::CreateLambda([this]() { SetHorizontalJustification(ESignalFlowHorizontalJustification::Bottom); }),
-			FCanExecuteAction::CreateLambda([this](){ return GraphWidgetContents.IsValid() && GraphWidgetContents->GetGraphOrientation() == EOrientation::Orient_Horizontal; }),
-			FIsActionChecked::CreateLambda([this]() { return HorizontalGraphJustification == ESignalFlowHorizontalJustification::Bottom; }));
-
-		CommandList->MapAction(
-			Commands.GetJustifyHorizontalCenterCommand(),
-			FExecuteAction::CreateLambda([this]() { SetHorizontalJustification(ESignalFlowHorizontalJustification::Center); }),
-			FCanExecuteAction::CreateLambda([this]() { return GraphWidgetContents.IsValid() && GraphWidgetContents->GetGraphOrientation() == EOrientation::Orient_Horizontal; }),
-			FIsActionChecked::CreateLambda([this]() { return HorizontalGraphJustification == ESignalFlowHorizontalJustification::Center; }));
-
-		CommandList->MapAction(
-			Commands.GetJustifyTopCommand(),
-			FExecuteAction::CreateLambda([this]() { SetHorizontalJustification(ESignalFlowHorizontalJustification::Top); }),
-			FCanExecuteAction::CreateLambda([this]() { return GraphWidgetContents.IsValid() && GraphWidgetContents->GetGraphOrientation() == EOrientation::Orient_Horizontal; }),
-			FIsActionChecked::CreateLambda([this]() { return HorizontalGraphJustification == ESignalFlowHorizontalJustification::Top; }));
-
-		CommandList->MapAction(
-			Commands.GetJustifyLeftCommand(),
-			FExecuteAction::CreateLambda([this]() { SetVerticalJustification(ESignalFlowVerticalJustification::Left); }),
-			FCanExecuteAction::CreateLambda([this]() { return GraphWidgetContents.IsValid() && GraphWidgetContents->GetGraphOrientation() == EOrientation::Orient_Vertical; }),
-			FIsActionChecked::CreateLambda([this]() { return VerticalGraphJustification == ESignalFlowVerticalJustification::Left; }));
-
-		CommandList->MapAction(
-			Commands.GetJustifyVerticalCenterCommand(),
-			FExecuteAction::CreateLambda([this]() { SetVerticalJustification(ESignalFlowVerticalJustification::Center); }),
-			FCanExecuteAction::CreateLambda([this]() { return GraphWidgetContents.IsValid() && GraphWidgetContents->GetGraphOrientation() == EOrientation::Orient_Vertical; }),
-			FIsActionChecked::CreateLambda([this]() { return VerticalGraphJustification == ESignalFlowVerticalJustification::Center; }));
+			Commands.GetJustifyEdgeCommand(),
+			FExecuteAction::CreateLambda([this]() { SetJustification(ESignalFlowJustification::Edge); }),
+			FCanExecuteAction(),
+			FIsActionChecked::CreateLambda([this]() { return GraphJustification == ESignalFlowJustification::Edge; }));
 
 		CommandList->MapAction(
-			Commands.GetJustifyRightCommand(),
-			FExecuteAction::CreateLambda([this]() { SetVerticalJustification(ESignalFlowVerticalJustification::Right); }),
-			FCanExecuteAction::CreateLambda([this]() { return GraphWidgetContents.IsValid() && GraphWidgetContents->GetGraphOrientation() == EOrientation::Orient_Vertical; }),
-			FIsActionChecked::CreateLambda([this]() { return VerticalGraphJustification == ESignalFlowVerticalJustification::Right; }));
+			Commands.GetJustifyCenterCommand(),
+			FExecuteAction::CreateLambda([this]() { SetJustification(ESignalFlowJustification::Center); }),
+			FCanExecuteAction(),
+			FIsActionChecked::CreateLambda([this]() { return GraphJustification == ESignalFlowJustification::Center; }));
 
 		CommandList->MapAction(
 			Commands.GetShowNodeDetailsCommand(),
@@ -553,59 +529,38 @@ namespace UE::Audio::Insights
 			+ SHorizontalBox::Slot()
 			.AutoWidth()
 			[
-				SNew(SSegmentedControl<ESignalFlowHorizontalJustification>)
-				.Visibility_Lambda([this]()
-				{
-					return GraphWidgetContents.IsValid() && GraphWidgetContents->GetGraphOrientation() == EOrientation::Orient_Horizontal
-																										? EVisibility::Visible
-																										: EVisibility::Collapsed;
-				})
-				.OnValueChanged(this, &FSignalFlowDashboardViewFactory::SetHorizontalJustification)
+				SNew(SSegmentedControl<ESignalFlowJustification>)
+				.OnValueChanged(this, &FSignalFlowDashboardViewFactory::SetJustification)
 				.Value_Lambda([this]
 				{
-					return HorizontalGraphJustification;
+					return GraphJustification;
 				})
-					
-				+ SSegmentedControl<ESignalFlowHorizontalJustification>::Slot(ESignalFlowHorizontalJustification::Bottom)
-				.ToolTip(LOCTEXT("AudioDashboard_SignalFlow_HorizontalJustification_Bottom_Tooltip", "Position nodes towards the bottom."))
-				.Icon(FSlateStyle::Get().GetBrush("AudioInsights.Icon.SignalFlow.AlignHorizontalBottom"))
-				
-				+ SSegmentedControl<ESignalFlowHorizontalJustification>::Slot(ESignalFlowHorizontalJustification::Center)
-				.ToolTip(LOCTEXT("AudioDashboard_SignalFlow_HorizontalJustification_Center_Tooltip", "Position nodes around the center."))
-				.Icon(FSlateStyle::Get().GetBrush("AudioInsights.Icon.SignalFlow.AlignHorizontalCenter"))
 
-				+ SSegmentedControl<ESignalFlowHorizontalJustification>::Slot(ESignalFlowHorizontalJustification::Top)
-				.ToolTip(LOCTEXT("AudioDashboard_SignalFlow_HorizontalJustification_Top_Tooltip", "Position nodes towards the top."))
-				.Icon(FSlateStyle::Get().GetBrush("AudioInsights.Icon.SignalFlow.AlignHorizontalTop"))
-			]
-
-			+ SHorizontalBox::Slot()
-			.AutoWidth()
-			[
-				SNew(SSegmentedControl<ESignalFlowVerticalJustification>)
-				.Visibility_Lambda([this]()
+				+ SSegmentedControl<ESignalFlowJustification>::Slot(ESignalFlowJustification::Edge)
+				.ToolTip(LOCTEXT("AudioDashboard_SignalFlow_Justification_Top_Tooltip", "Position nodes towards one edge of the graph."))
+				.Icon_Lambda([this]() -> const FSlateBrush*
 				{
-					return GraphWidgetContents.IsValid() && GraphWidgetContents->GetGraphOrientation() == EOrientation::Orient_Vertical
-																										? EVisibility::Visible
-																										: EVisibility::Collapsed;
+					if (!GraphWidgetContents.IsValid())
+					{
+						return nullptr;
+					}
+
+					return GraphWidgetContents->GetGraphOrientation() == EOrientation::Orient_Horizontal ? FSlateStyle::Get().GetBrush("AudioInsights.Icon.SignalFlow.AlignHorizontalTop")
+																										 : FSlateStyle::Get().GetBrush("AudioInsights.Icon.SignalFlow.AlignVerticalLeft");
 				})
-				.OnValueChanged(this, &FSignalFlowDashboardViewFactory::SetVerticalJustification)
-				.Value_Lambda([this]
+
+				+ SSegmentedControl<ESignalFlowJustification>::Slot(ESignalFlowJustification::Center)
+				.ToolTip(LOCTEXT("AudioDashboard_SignalFlow_Justification_Center_Tooltip", "Position nodes around the center of the graph."))
+				.Icon_Lambda([this]() -> const FSlateBrush*
 				{
-					return VerticalGraphJustification;
-				})
-					
-				+ SSegmentedControl<ESignalFlowVerticalJustification>::Slot(ESignalFlowVerticalJustification::Left)
-				.ToolTip(LOCTEXT("AudioDashboard_SignalFlow_VerticalJustification_Left_Tooltip", "Position nodes towards the left."))
-				.Icon(FSlateStyle::Get().GetBrush("AudioInsights.Icon.SignalFlow.AlignVerticalLeft"))
-				
-				+ SSegmentedControl<ESignalFlowVerticalJustification>::Slot(ESignalFlowVerticalJustification::Center)
-				.ToolTip(LOCTEXT("AudioDashboard_SignalFlow_VerticalJustification_Center_Tooltip", "Position nodes around the center."))
-				.Icon(FSlateStyle::Get().GetBrush("AudioInsights.Icon.SignalFlow.AlignVerticalCenter"))
+					if (!GraphWidgetContents.IsValid())
+					{
+						return nullptr;
+					}
 
-				+ SSegmentedControl<ESignalFlowVerticalJustification>::Slot(ESignalFlowVerticalJustification::Right)
-				.ToolTip(LOCTEXT("AudioDashboard_SignalFlow_VerticalJustification_Right_Tooltip", "Position nodes towards the right."))
-				.Icon(FSlateStyle::Get().GetBrush("AudioInsights.Icon.SignalFlow.AlignVerticalRight"))
+					return GraphWidgetContents->GetGraphOrientation() == EOrientation::Orient_Horizontal ? FSlateStyle::Get().GetBrush("AudioInsights.Icon.SignalFlow.AlignHorizontalCenter")
+																										 : FSlateStyle::Get().GetBrush("AudioInsights.Icon.SignalFlow.AlignVerticalCenter");
+				})
 			]
 		];
 	}
@@ -630,15 +585,9 @@ namespace UE::Audio::Insights
 			MenuBuilder.AddMenuEntry(Commands.GetPauseOnSelectCommand());
 			MenuBuilder.AddSubMenu
 			(
-				LOCTEXT("AudioDashboard_SignalFlow_Settings_HorizontalJustification", "Horizontal Graph Justification"),
-				LOCTEXT("AudioDashboard_SignalFlow_Settings_HorizontalJustificationTooltip", "Choose where nodes are aligned when viewing the signal flow graph in horizontal mode."),
-				FNewMenuDelegate::CreateRaw(this, &FSignalFlowDashboardViewFactory::BuildHorizontalJustificationMenuContent)
-			);
-			MenuBuilder.AddSubMenu
-			(
-				LOCTEXT("AudioDashboard_SignalFlow_Settings_VerticalJustification", "Vertical Graph Justification"),
-				LOCTEXT("AudioDashboard_SignalFlow_Settings_VerticalJustificationTooltip", "Choose where nodes are aligned when viewing the signal flow graph in vertical mode."),
-				FNewMenuDelegate::CreateRaw(this, &FSignalFlowDashboardViewFactory::BuildVerticalJustificationMenuContent)
+				LOCTEXT("AudioDashboard_SignalFlow_Settings_Justification", "Graph Justification"),
+				LOCTEXT("AudioDashboard_SignalFlow_Settings_JustificationTooltip", "Choose where nodes are aligned when viewing the signal flow graph."),
+				FNewMenuDelegate::CreateRaw(this, &FSignalFlowDashboardViewFactory::BuildJustificationMenuContent)
 			);
 
 			MenuBuilder.AddMenuEntry(Commands.GetShowNodeDetailsCommand());
@@ -655,22 +604,12 @@ namespace UE::Audio::Insights
 		return MenuBuilder.MakeWidget();
 	}
 
-	void FSignalFlowDashboardViewFactory::BuildHorizontalJustificationMenuContent(FMenuBuilder& OutMenuBuilder)
+	void FSignalFlowDashboardViewFactory::BuildJustificationMenuContent(FMenuBuilder& OutMenuBuilder)
 	{
 		const FSignalFlowEditorCommands& Commands = FSignalFlowEditorCommands::Get();
 
-		OutMenuBuilder.AddMenuEntry(Commands.GetJustifyBottomCommand());
-		OutMenuBuilder.AddMenuEntry(Commands.GetJustifyHorizontalCenterCommand());
-		OutMenuBuilder.AddMenuEntry(Commands.GetJustifyTopCommand());
-	}
-
-	void FSignalFlowDashboardViewFactory::BuildVerticalJustificationMenuContent(FMenuBuilder& OutMenuBuilder)
-	{
-		const FSignalFlowEditorCommands& Commands = FSignalFlowEditorCommands::Get();
-
-		OutMenuBuilder.AddMenuEntry(Commands.GetJustifyLeftCommand());
-		OutMenuBuilder.AddMenuEntry(Commands.GetJustifyVerticalCenterCommand());
-		OutMenuBuilder.AddMenuEntry(Commands.GetJustifyRightCommand());
+		OutMenuBuilder.AddMenuEntry(Commands.GetJustifyEdgeCommand());
+		OutMenuBuilder.AddMenuEntry(Commands.GetJustifyCenterCommand());
 	}
 
 	void FSignalFlowDashboardViewFactory::BuildNodeDetailFiltersMenuContent(FMenuBuilder& OutMenuBuilder)
@@ -906,20 +845,9 @@ namespace UE::Audio::Insights
 		TimingViewExtender.ResumeTimeMarker();
 	}
 
-	void FSignalFlowDashboardViewFactory::SetHorizontalJustification(const ESignalFlowHorizontalJustification Justification)
-	{
-		HorizontalGraphJustification = Justification;
-
-#if WITH_EDITOR
-		FSignalFlowSettings::OnRequestWriteSettings.Broadcast();
-#endif // WITH_EDITOR
-
-		OnRequestGraphRefresh();
-	}
-
-	void FSignalFlowDashboardViewFactory::SetVerticalJustification(const ESignalFlowVerticalJustification Justification)
+	void FSignalFlowDashboardViewFactory::SetJustification(const ESignalFlowJustification Justification)
 	{
-		VerticalGraphJustification = Justification;
+		GraphJustification = Justification;
 
 #if WITH_EDITOR
 		FSignalFlowSettings::OnRequestWriteSettings.Broadcast();
@@ -2005,15 +1933,15 @@ namespace UE::Audio::Insights
 			{
 				const TArray<int32> PrevRowOfIndexes = NodeIndexRows[PreviousRowsIndex];
 
-				TOptional<float> PrevXPos;
+				TArray<TSharedPtr<ISignalFlowNode>> PrevPositionedNodesInRow;
 				for (int32 NodeIndexInRow = 0; NodeIndexInRow < PrevRowOfIndexes.Num(); ++NodeIndexInRow)
 				{
 					const TSharedPtr<ISignalFlowNode>& Node = FilteredGraphNodes[PrevRowOfIndexes[NodeIndexInRow]];
 
 					if (Node.IsValid())
 					{
-						CalcXPositionForNode(Node, Node->FilteredInputs, PrevXPos, GraphOrientation);
-						PrevXPos = Node->XPos;
+						CalcXPositionForNode(Node, Node->FilteredInputs, NodeIndexInRow, PrevPositionedNodesInRow);
+						PrevPositionedNodesInRow.Add(Node);
 					}
 				}
 			}
@@ -2021,14 +1949,14 @@ namespace UE::Audio::Insights
 		else
 		{
 			// Process just this row - moving nodes relative to their outputs
-			TOptional<float> PrevXPos;
+			TArray<TSharedPtr<ISignalFlowNode>> PrevPositionedNodesInRow;
 			for (int32 NodeIndexInRow = 0; NodeIndexInRow < RowToProcess.Num(); ++NodeIndexInRow)
 			{
 				TSharedPtr<ISignalFlowNode>& Node = FilteredGraphNodes[RowToProcess[NodeIndexInRow]];
 				if (Node.IsValid())
 				{
-					CalcXPositionForNode(Node, Node->FilteredOutputs, PrevXPos, GraphOrientation);
-					PrevXPos = Node->XPos;
+					CalcXPositionForNode(Node, Node->FilteredOutputs, NodeIndexInRow, PrevPositionedNodesInRow);
+					PrevPositionedNodesInRow.Add(Node);
 				}
 			}
 		}
@@ -2051,141 +1979,85 @@ namespace UE::Audio::Insights
 		}
 	}
 
-	void FSignalFlowDashboardViewFactory::CalcXPositionForNode(const TSharedPtr<ISignalFlowNode>& Node, const TArray<FSignalFlowEntryKey>& NodesToPositionRelativeTo, const TOptional<float>& PrevXPos, const EOrientation GraphOrientation)
+	void FSignalFlowDashboardViewFactory::CalcXPositionForNode(const TSharedPtr<ISignalFlowNode>& Node, const TArray<FSignalFlowEntryKey>& NodesToPositionRelativeTo, const int32 NodeIndexInRow, const TArray<TSharedPtr<ISignalFlowNode>>& PrevPositionedNodesInRow)
 	{
-		switch (GraphOrientation)
+		switch (GraphJustification)
 		{
-			case EOrientation::Orient_Horizontal:
-
-				switch (HorizontalGraphJustification)
-				{
-					case ESignalFlowHorizontalJustification::Bottom:
-						AlignNodeTowardsMaxX(Node, NodesToPositionRelativeTo, PrevXPos);
-						break;
-
-					case ESignalFlowHorizontalJustification::Center:
-						AlignNodeTowardsCenterX(Node, NodesToPositionRelativeTo, PrevXPos);
-						break;
-
-					case ESignalFlowHorizontalJustification::Top:
-						AlignNodeTowardsLowestX(Node, NodesToPositionRelativeTo, PrevXPos);
-						break;
-				}
-
-				break;
+		case ESignalFlowJustification::Edge:
+			AlignNodeTowardsEdge(Node, NodeIndexInRow);
+			break;
 
-			case EOrientation::Orient_Vertical:
-
-				switch (VerticalGraphJustification)
-				{
-				case ESignalFlowVerticalJustification::Left:
-					AlignNodeTowardsLowestX(Node, NodesToPositionRelativeTo, PrevXPos);
-					break;
-
-				case ESignalFlowVerticalJustification::Center:
-					AlignNodeTowardsCenterX(Node, NodesToPositionRelativeTo, PrevXPos);
-					break;
-
-				case ESignalFlowVerticalJustification::Right:
-					AlignNodeTowardsMaxX(Node, NodesToPositionRelativeTo, PrevXPos);
-					break;
-				}
-
-				break;
+		case ESignalFlowJustification::Center:
+			AlignNodeTowardsCenter(Node, NodesToPositionRelativeTo, PrevPositionedNodesInRow);
+			break;
 		}
 	}
 
-	void FSignalFlowDashboardViewFactory::AlignNodeTowardsCenterX(const TSharedPtr<ISignalFlowNode>& Node, const TArray<FSignalFlowEntryKey>& NodesToPositionRelativeTo, const TOptional<float>& PrevXPos)
+	void FSignalFlowDashboardViewFactory::AlignNodeTowardsEdge(const TSharedPtr<ISignalFlowNode>& Node, const int32 NodeIndexInRow)
 	{
-		if (!Node.IsValid() || NodesToPositionRelativeTo.IsEmpty())
+		if (Node.IsValid())
 		{
-			return;
-		}
-
-		float TotalXPositions = 0.0f;
-		for (int32 AligningNodeIndex = 0; AligningNodeIndex < NodesToPositionRelativeTo.Num(); ++AligningNodeIndex)
-		{
-			TSharedPtr<ISignalFlowNode>* AliningNode = FilteredEntries.Find(NodesToPositionRelativeTo[AligningNodeIndex]);
-			if (AliningNode && AliningNode->IsValid())
-			{
-				TotalXPositions += (*AliningNode)->XPos;
-			}
-		}
-
-		const float AverageXPosition = TotalXPositions / NodesToPositionRelativeTo.Num();
-
-		if (PrevXPos.IsSet())
-		{
-			Node->XPos = FMath::Max(AverageXPosition, PrevXPos.GetValue() + 1.0f);
-		}
-		else
-		{
-			Node->XPos = AverageXPosition;
+			Node->XPos = NodeIndexInRow;
 		}
 	}
 
-	void FSignalFlowDashboardViewFactory::AlignNodeTowardsLowestX(const TSharedPtr<ISignalFlowNode>& Node, const TArray<FSignalFlowEntryKey>& NodesToPositionRelativeTo, const TOptional<float>& PrevXPos)
+	void FSignalFlowDashboardViewFactory::AlignNodeTowardsCenter(const TSharedPtr<ISignalFlowNode>& Node, const TArray<FSignalFlowEntryKey>& NodesToPositionRelativeTo, const TArray<TSharedPtr<ISignalFlowNode>>& PrevPositionedNodesInRow)
 	{
-		if (!Node.IsValid() || NodesToPositionRelativeTo.IsEmpty())
+		if (!Node.IsValid())
 		{
 			return;
 		}
 
-		TSharedPtr<ISignalFlowNode> NodeToAlignWith = nullptr;
-		float MinXPos = TNumericLimits<float>::Max();
-		for (int32 AligningNodeIndex = 0; AligningNodeIndex < NodesToPositionRelativeTo.Num(); ++AligningNodeIndex)
+		if (NodesToPositionRelativeTo.IsEmpty())
 		{
-			TSharedPtr<ISignalFlowNode>* AliningNode = FilteredEntries.Find(NodesToPositionRelativeTo[AligningNodeIndex]);
-			if (AliningNode && AliningNode->IsValid() && (*AliningNode)->XPos < MinXPos)
+			// If no nodes are present to position this node relative to, it may be that this node has no inputs/outputs
+			// In this case, position this just next to the previously placed node
+			if (PrevPositionedNodesInRow.IsEmpty() || !PrevPositionedNodesInRow.Last().IsValid())
 			{
-				NodeToAlignWith = (*AliningNode);
-				MinXPos = (*AliningNode)->XPos;
-			}
-		}
-
-		if (NodeToAlignWith.IsValid())
-		{
-			if (PrevXPos.IsSet())
-			{
-				Node->XPos = FMath::Max(NodeToAlignWith->XPos, PrevXPos.GetValue() + 1.0f);
+				Node->XPos = 0.0f;
 			}
 			else
 			{
-				Node->XPos = NodeToAlignWith->XPos;
+				Node->XPos = PrevPositionedNodesInRow.Last()->XPos + 1.0f;
 			}
 		}
-	}
-
-	void FSignalFlowDashboardViewFactory::AlignNodeTowardsMaxX(const TSharedPtr<ISignalFlowNode>& Node, const TArray<FSignalFlowEntryKey>& NodesToPositionRelativeTo, const TOptional<float>& PrevXPos)
-	{
-		if (!Node.IsValid() || NodesToPositionRelativeTo.IsEmpty())
-		{
-			return;
-		}
-
-		TSharedPtr<ISignalFlowNode> NodeToAlignWith = nullptr;
-		float MaxXPos = TNumericLimits<float>::Lowest();
-		for (int32 AligningNodeIndex = 0; AligningNodeIndex < NodesToPositionRelativeTo.Num(); ++AligningNodeIndex)
+		else
 		{
-			TSharedPtr<ISignalFlowNode>* AliningNode = FilteredEntries.Find(NodesToPositionRelativeTo[AligningNodeIndex]);
-			if (AliningNode && AliningNode->IsValid() && (*AliningNode)->XPos > MaxXPos)
+			// Position this node in the average position of the nodes we are positioning relative to
+			float TotalXPositions = 0.0f;
+			for (int32 AligningNodeIndex = 0; AligningNodeIndex < NodesToPositionRelativeTo.Num(); ++AligningNodeIndex)
 			{
-				NodeToAlignWith = (*AliningNode);
-				MaxXPos = (*AliningNode)->XPos;
+				TSharedPtr<ISignalFlowNode>* AliningNode = FilteredEntries.Find(NodesToPositionRelativeTo[AligningNodeIndex]);
+				if (AliningNode && AliningNode->IsValid())
+				{
+					TotalXPositions += (*AliningNode)->XPos;
+				}
 			}
-		}
 
-		if (NodeToAlignWith.IsValid())
-		{
-			if (PrevXPos.IsSet())
-			{
-				Node->XPos = FMath::Max(NodeToAlignWith->XPos, PrevXPos.GetValue() + 1.0f);
-			}
-			else
+			// We use Floor(Av X 2) / 2 to snap the node's position on a grid with spacing 0.5 between cells
+			const float AverageXPosition = FMath::Floor((TotalXPositions / NodesToPositionRelativeTo.Num()) * 2.0f) * 0.5f;
+			Node->XPos = AverageXPosition;
+
+			// Now ensure any nodes placed previously are shifted if they overlap the new node.
+			// We want to maintain the order that nodes are drawn in the graph
+			// Run backwards, ensuring at least a distance of 1.0 between nodes, until we don't detect any further overlaps
+			float ShiftXPosition = AverageXPosition - 1.0f;
+			for (int32 RevIndex = PrevPositionedNodesInRow.Num() - 1; RevIndex >= 0; --RevIndex)
 			{
-				Node->XPos = NodeToAlignWith->XPos;
+				const TSharedPtr<ISignalFlowNode> PrevNode = PrevPositionedNodesInRow[RevIndex];
+				if (PrevNode.IsValid() && PrevNode->XPos > ShiftXPosition)
+				{
+					PrevNode->XPos = ShiftXPosition;
+					ShiftXPosition -= 1.0f;
+				}
+				else
+				{
+					// No more overlaps detected, early out
+					break;
+				}
 			}
 		}
+		
 	}
 
 #if WITH_EDITOR
@@ -2203,8 +2075,7 @@ namespace UE::Audio::Insights
 		}
 
 		bPauseGraphOnSelect = InSettings.bPauseGraphOnSelect;
-		HorizontalGraphJustification = InSettings.HorizontalGraphJustification;
-		VerticalGraphJustification = InSettings.VerticalGraphJustification;
+		GraphJustification = InSettings.GraphJustification;
 		bShowNodeDetails = InSettings.bShowNodeDetails;
 		NodeDetailFilters = InSettings.NodeDetailFilters;
 		bDisplayAmpPeakInDb = InSettings.AmplitudeDisplayMode == EAudioAmplitudeDisplayMode::Decibels;
@@ -2225,8 +2096,7 @@ namespace UE::Audio::Insights
 			OutSettings.bHorizontalFlow = bGraphIsHorizontal;
 		}
 
-		OutSettings.HorizontalGraphJustification = HorizontalGraphJustification;
-		OutSettings.VerticalGraphJustification = VerticalGraphJustification;
+		OutSettings.GraphJustification = GraphJustification;
 		OutSettings.bShowNodeDetails = bShowNodeDetails;
 		OutSettings.bPauseGraphOnSelect = bPauseGraphOnSelect;
 		OutSettings.NodeDetailFilters = NodeDetailFilters;
diff --git a/Engine/Plugins/AudioInsights/Source/AudioInsights/Public/Settings/SignalFlowSettings.h b/Engine/Plugins/AudioInsights/Source/AudioInsights/Public/Settings/SignalFlowSettings.h
--- a/Engine/Plugins/AudioInsights/Source/AudioInsights/Public/Settings/SignalFlowSettings.h
+++ b/Engine/Plugins/AudioInsights/Source/AudioInsights/Public/Settings/SignalFlowSettings.h
@@ -12,38 +12,18 @@
 
 #define LOCTEXT_NAMESPACE "AudioInsights"
 
-/** ESignalFlowHorizontalJustification
+/** ESignalFlowJustification
  *
- * Where the nodes in the signal flow graph are aligned when viewed in horizontal mode
+ * Where the nodes in the signal flow graph are aligned
  */
 UENUM(BlueprintType)
-enum class ESignalFlowHorizontalJustification : uint8
+enum class ESignalFlowJustification : uint8
 {
-	/** Position nodes towards the bottom */
-	Bottom		UMETA(DisplayName = "Bottom"),
+	/** Position nodes towards one edge of the graph */
+	Edge		UMETA(DisplayName = "Edge"),
 
-	/** Position nodes around the center */
-	Center		UMETA(DisplayName = "Center"),
-
-	/** Position nodes towards the top */
-	Top			UMETA(DisplayName = "Top")
-};
-
-/** ESignalFlowVerticalJustification
- *
- * Where the nodes in the signal flow graph are aligned when viewed in vertical mode
- */
-UENUM(BlueprintType)
-enum class ESignalFlowVerticalJustification : uint8
-{
-	/** Position nodes towards the left */
-	Left		UMETA(DisplayName = "Left"),
-
-	/** Position nodes around the center */
-	Center		UMETA(DisplayName = "Center"),
-
-	/** Position nodes towards the right */
-	Right		UMETA(DisplayName = "Right")
+	/** Position nodes around the center of the graph */
+	Center		UMETA(DisplayName = "Center")
 };
 
 USTRUCT()
@@ -116,13 +96,9 @@ struct FSignalFlowSettings
 	UPROPERTY(EditAnywhere, config, Category = SignalFlow)
 	bool bPauseGraphOnSelect = false;
 
-	/** Where the nodes in the signal flow graph are aligned when viewed in horizontal mode */
-	UPROPERTY(EditAnywhere, config, Category = SignalFlow, meta = (DisplayName = "Horizontal Graph Justification"))
-	ESignalFlowHorizontalJustification HorizontalGraphJustification = ESignalFlowHorizontalJustification::Top;
-
-	/** Where the nodes in the signal flow graph are aligned when viewed in vertical mode */
-	UPROPERTY(EditAnywhere, config, Category = SignalFlow, meta = (DisplayName = "Vertical Graph Justification"))
-	ESignalFlowVerticalJustification VerticalGraphJustification = ESignalFlowVerticalJustification::Left;
+	/** Where the nodes in the signal flow graph are aligned */
+	UPROPERTY(EditAnywhere, config, Category = SignalFlow, meta = (DisplayName = "Graph Justification"))
+	ESignalFlowJustification GraphJustification = ESignalFlowJustification::Edge;
 
 	/** Toggles visibility of node details. */
 	UPROPERTY(EditAnywhere, config, Category = SignalFlow)
diff --git a/Engine/Plugins/AudioInsights/Source/AudioInsights/Public/Views/SignalFlowDashboardViewFactory.h b/Engine/Plugins/AudioInsights/Source/AudioInsights/Public/Views/SignalFlowDashboardViewFactory.h
--- a/Engine/Plugins/AudioInsights/Source/AudioInsights/Public/Views/SignalFlowDashboardViewFactory.h
+++ b/Engine/Plugins/AudioInsights/Source/AudioInsights/Public/Views/SignalFlowDashboardViewFactory.h
@@ -77,8 +77,7 @@ namespace UE::Audio::Insights
 		TSharedRef<SWidget> CreateToolbarButtonActions();
 
 		TSharedRef<SWidget> OnGetSettingsMenuContent();
-		void BuildHorizontalJustificationMenuContent(FMenuBuilder& OutMenuBuilder);
-		void BuildVerticalJustificationMenuContent(FMenuBuilder& OutMenuBuilder);
+		void BuildJustificationMenuContent(FMenuBuilder& OutMenuBuilder);
 		void BuildNodeDetailFiltersMenuContent(FMenuBuilder& OutMenuBuilder);
 		void CreateAmpDisplayModeToggle(FMenuBuilder& OutMenuBuilder);
 
@@ -88,8 +87,7 @@ namespace UE::Audio::Insights
 		void TogglePauseGraphEnabled();
 		void ToggleEnableAllNodeDetailFilters();
 		void ToggleAnimateWires();
-		void SetHorizontalJustification(const ESignalFlowHorizontalJustification Justification);
-		void SetVerticalJustification(const ESignalFlowVerticalJustification Justification);
+		void SetJustification(const ESignalFlowJustification Justification);
 
 		bool CanPauseTimestamp() const;
 		bool IsTimestampPaused() const;
@@ -136,10 +134,9 @@ namespace UE::Audio::Insights
 		void CalcXPositionsForRows(const TArray<TArray<int32>>& NodeIndexRows, const EOrientation GraphOrientation, int32& OutGraphWidth);
 		void PositionRowEquidistant(const TArray<int32>& NodeIndexRow);
 		
-		void CalcXPositionForNode(const TSharedPtr<ISignalFlowNode>& Node, const TArray<FSignalFlowEntryKey>& NodesToPositionRelativeTo, const TOptional<float>& PrevXPos, const EOrientation GraphOrientation);
-		void AlignNodeTowardsCenterX(const TSharedPtr<ISignalFlowNode>& Node, const TArray<FSignalFlowEntryKey>& NodesToPositionRelativeTo, const TOptional<float>& PrevXPos);
-		void AlignNodeTowardsLowestX(const TSharedPtr<ISignalFlowNode>& Node, const TArray<FSignalFlowEntryKey>& NodesToPositionRelativeTo, const TOptional<float>& PrevXPos);
-		void AlignNodeTowardsMaxX(const TSharedPtr<ISignalFlowNode>& Node, const TArray<FSignalFlowEntryKey>& NodesToPositionRelativeTo, const TOptional<float>& PrevXPos);
+		void CalcXPositionForNode(const TSharedPtr<ISignalFlowNode>& Node, const TArray<FSignalFlowEntryKey>& NodesToPositionRelativeTo, const int32 NodeIndexInRow, const TArray<TSharedPtr<ISignalFlowNode>>& PrevPositionedNodesInRow);
+		void AlignNodeTowardsEdge(const TSharedPtr<ISignalFlowNode>& Node, const int32 NodeIndexInRow);
+		void AlignNodeTowardsCenter(const TSharedPtr<ISignalFlowNode>& Node, const TArray<FSignalFlowEntryKey>& NodesToPositionRelativeTo, const TArray<TSharedPtr<ISignalFlowNode>>& PrevPositionedNodesInRow);
 
 #if WITH_EDITOR
 		void OnReadEditorSettings(const FSignalFlowSettings& InSettings);
@@ -173,8 +170,7 @@ namespace UE::Audio::Insights
 
 		FSignalFlowNodeDetailFilterSettings NodeDetailFilters;
 
-		ESignalFlowHorizontalJustification HorizontalGraphJustification = ESignalFlowHorizontalJustification::Top;
-		ESignalFlowVerticalJustification VerticalGraphJustification = ESignalFlowVerticalJustification::Left;
+		ESignalFlowJustification GraphJustification = ESignalFlowJustification::Edge;
 
 		bool bGraphRequiresRefresh = false;
 		bool bAutoFocusAfterGraphRefresh = false;
