[Audio Insights] Allow users to expand/collapse the left-hand "Active Entries" menus in the Signal Flow tab
* This was an original feature on Robbie's previz, and will become useful when adding new node types as we can set the default expansion setting per list.

Settings
* Persistent settings added to FSignalFlowSettings (which will appear in editor preferences)
* Settings passed between editor preferences and Signal Flow tab inside FSignalFlowDashboardViewFactory::OnReadEditorSettings/OnWriteEditorSettings

Expandable/Collapsible Menus
* Scale up the height props between each menu so the lowest is equal to 1.0, but the ratios remain the same
* We now cache the SVerticalBox ActiveEntryMenusContainer so we can edit to size of the slots
* Each menu caches it's SlotIndex inside the lambda AddSelectableListToSlot when the menu is first added to the vertical box (FSignalFlowDashboardViewFactory::CreateSelectionPanel)
* Add new expander arrow to the left of each menu title (in the header row content), and link this to the settings
* Menu content is now collapsed when the expander arrow is collapsed
* Set the slot size to Auto() with a fixed minimum when the menu is collapsed. This lets us see the header row and a small amount of padding when the menu is collapsed.
* Updated arrow icons to show sideways for collapsed, down for expanded
* Fixed bug where nodes with one send output with no set value would still show the expander arrow

#Synced-CL 51049241

diff --git a/Engine/Plugins/AudioInsights/Source/AudioInsights/Private/Views/SSignalFlowGraph.cpp b/Engine/Plugins/AudioInsights/Source/AudioInsights/Private/Views/SSignalFlowGraph.cpp
--- a/Engine/Plugins/AudioInsights/Source/AudioInsights/Private/Views/SSignalFlowGraph.cpp
+++ b/Engine/Plugins/AudioInsights/Source/AudioInsights/Private/Views/SSignalFlowGraph.cpp
@@ -730,7 +730,7 @@ namespace UE::Audio::Insights
 			.Padding(4.0f, 0.0f, 0.0f, 0.0f)
 			[
 				SNew(SCheckBox)
-				.Style(&FSlateStyle::Get().GetWidgetStyle<FCheckBoxStyle>("AudioInsights.ExpandArrowToggle"))
+				.Style(&FSlateStyle::Get().GetWidgetStyle<FCheckBoxStyle>("AudioInsights.LeftExpandArrowToggle"))
 				.Visibility_Lambda([this]()
 				{ 
 					return NodeDetailsExpanderArrowIsEnabled() ? EVisibility::Visible : EVisibility::Collapsed;
@@ -929,7 +929,13 @@ namespace UE::Audio::Insights
 			{
 				if (NodeDetailFilterSettings->GetParameterIsVisible(Param))
 				{
-					return true;
+					for (const auto& [OutputKey, SendInfo] : Entry->Outputs)
+					{
+						if (SendInfo.SendLevel.IsSet())
+						{
+							return true;
+						}
+					}
 				}
 			}
 			else if (NodeDetailFilterSettings->GetParameterIsVisible(Param) && GetNodeDetailParamValue(Entry, Param).IsSet())
diff --git a/Engine/Plugins/AudioInsights/Source/AudioInsights/Private/Views/SignalFlowDashboardViewFactory.cpp b/Engine/Plugins/AudioInsights/Source/AudioInsights/Private/Views/SignalFlowDashboardViewFactory.cpp
--- a/Engine/Plugins/AudioInsights/Source/AudioInsights/Private/Views/SignalFlowDashboardViewFactory.cpp
+++ b/Engine/Plugins/AudioInsights/Source/AudioInsights/Private/Views/SignalFlowDashboardViewFactory.cpp
@@ -35,9 +35,10 @@ namespace UE::Audio::Insights
 		const FName ButtonTextStyleName = "SmallButtonText";
 
 		constexpr float GraphWidthProp = 0.8f;
-		constexpr float SourcesListHeightProp = 0.5f;
-		constexpr float BusListHeightProp = 0.2f;
-		constexpr float SubmixesListHeightProp = 0.3f;
+		constexpr float SourcesListHeightProp = 2.5f;
+		constexpr float BusListHeightProp = 1.0f;
+		constexpr float SubmixesListHeightProp = 1.5f;
+		constexpr float CollapsedActiveEntryMenuHeight = 28.0f;
 
 		constexpr float ToolbarButtonHeight = 20.0f;
 		constexpr float ToolbarButtonPaddingX = 8.0f;
@@ -885,11 +886,11 @@ namespace UE::Audio::Insights
 	{
 		using namespace FSignalFlowDashboardViewFactoryPrivate;
 
-		ActiveSourcesListView = CreateSelectableList(LOCTEXT("AudioDashboard_SignalFlow_ActiveSourcesTitle", "Active Sources"), ActiveSources);
-		ActiveBusesListView = CreateSelectableList(LOCTEXT("AudioDashboard_SignalFlow_ActiveBusesTitle", "Active Buses"), ActiveBuses);
-		ActiveSubmixesListView = CreateSelectableList(LOCTEXT("AudioDashboard_SignalFlow_ActiveSubmixesTitle", "Active Submixes"), ActiveSubmixes);
+		ActiveSourcesListView = CreateSelectableList(LOCTEXT("AudioDashboard_SignalFlow_ActiveSourcesTitle", "Active Sources"), ActiveSources, ESignalFlowEntryType::SoundSource);
+		ActiveBusesListView = CreateSelectableList(LOCTEXT("AudioDashboard_SignalFlow_ActiveBusesTitle", "Active Buses"), ActiveBuses, ESignalFlowEntryType::AudioBus);
+		ActiveSubmixesListView = CreateSelectableList(LOCTEXT("AudioDashboard_SignalFlow_ActiveSubmixesTitle", "Active Submixes"), ActiveSubmixes, ESignalFlowEntryType::Submix);
 
-		return SNew(SVerticalBox)
+		SAssignNew(ActiveEntryMenusContainer, SVerticalBox)
 
 		+ SVerticalBox::Slot()
 		.AutoHeight()
@@ -901,25 +902,25 @@ namespace UE::Audio::Insights
 			.HintText(LOCTEXT("AudioDashboard_SignalFlow_SearchBoxHintText", "Search"))
 			.MinDesiredWidth(200.0f)
 			.OnTextChanged(this, &FSignalFlowDashboardViewFactory::SetSearchBoxFilterText)
-		]
-
-		+ SVerticalBox::Slot()
-		.FillHeight(SourcesListHeightProp)
-		[
-			ActiveSourcesListView.ToSharedRef()
-		]
-			
-		+ SVerticalBox::Slot()
-		.FillHeight(BusListHeightProp)
-		[
-			ActiveBusesListView.ToSharedRef()
-		]
-		
-		+ SVerticalBox::Slot()
-		.FillHeight(SubmixesListHeightProp)
-		[
-			ActiveSubmixesListView.ToSharedRef()
 		];
+
+		auto AddSelectableListToSlot = [this](const TSharedRef<SListView<TSharedPtr<FSignalFlowDashboardEntry>>> Menu, const float FillHeight, int32& OutSlotIndex)
+		{
+			ActiveEntryMenusContainer->AddSlot()
+			.FillHeight(FillHeight)
+			.MinHeight(CollapsedActiveEntryMenuHeight)
+			[
+				Menu
+			];
+
+			OutSlotIndex = ActiveEntryMenusContainer->NumSlots() - 1;
+		};
+
+		AddSelectableListToSlot(ActiveSourcesListView.ToSharedRef(), SourcesListHeightProp, ActiveSourcesSlotIndex);
+		AddSelectableListToSlot(ActiveBusesListView.ToSharedRef(), BusListHeightProp, ActiveBusesSlotIndex);
+		AddSelectableListToSlot(ActiveSubmixesListView.ToSharedRef(), SubmixesListHeightProp, ActiveSubmixesSlotIndex);
+
+		return ActiveEntryMenusContainer.ToSharedRef();
 	}
 
 	TSharedRef<SWidget> FSignalFlowDashboardViewFactory::CreateGraphSection()
@@ -1066,7 +1067,7 @@ namespace UE::Audio::Insights
 		return GraphWidget;
 	}
 
-	TSharedRef<SListView<TSharedPtr<FSignalFlowDashboardEntry>>> FSignalFlowDashboardViewFactory::CreateSelectableList(const FText& HeadingText, const TArray<TSharedPtr<FSignalFlowDashboardEntry>>& ListSource)
+	TSharedRef<SListView<TSharedPtr<FSignalFlowDashboardEntry>>> FSignalFlowDashboardViewFactory::CreateSelectableList(const FText& HeadingText, const TArray<TSharedPtr<FSignalFlowDashboardEntry>>& ListSource, const ESignalFlowEntryType EntryType)
 	{
 		TSharedRef<SHeaderRow> HeaderRowWidget = SNew(SHeaderRow)
 												 .CanSelectGeneratedColumn(false);
@@ -1079,9 +1080,34 @@ namespace UE::Audio::Insights
 			.HeaderContentPadding(FMargin(4.0f, 4.0f))
 			.HeaderContent()
 			[
-				SNew(STextBlock)
-				.Text(HeadingText)
-				.Font(FCoreStyle::Get().GetFontStyle("BoldFont"))
+				SNew(SHorizontalBox)
+
+				+ SHorizontalBox::Slot()
+				.AutoWidth()
+				.MinSize(12.0f)
+				.MaxSize(12.0f)
+				.Padding(0.0f, 0.0f, 4.0f, 0.0f)
+				[
+					SNew(SCheckBox)
+					.Style(&FSlateStyle::Get().GetWidgetStyle<FCheckBoxStyle>("AudioInsights.RightExpandArrowToggle"))
+					.IsChecked_Lambda([this, EntryType]()
+					{
+						return GetSelectableListIsExpanded(EntryType) ? ECheckBoxState::Checked : ECheckBoxState::Unchecked;
+					})
+					.OnCheckStateChanged_Lambda([this, EntryType](const ECheckBoxState CheckboxState)
+					{
+						const bool bMenuIsExpanded = CheckboxState == ECheckBoxState::Checked;
+						SetSelectableListIsExpanded(EntryType, bMenuIsExpanded);
+					})
+				]
+
+				+ SHorizontalBox::Slot()
+				.FillWidth(1.0f)
+				[
+					SNew(STextBlock)
+					.Text(HeadingText)
+					.Font(FCoreStyle::Get().GetFontStyle("BoldFont"))
+				]
 			];
 
 		HeaderRowWidget->AddColumn(NameColumnArgs);
@@ -1089,10 +1115,14 @@ namespace UE::Audio::Insights
 		return SNew(SListView<TSharedPtr<FSignalFlowDashboardEntry>>)
 			.ListItemsSource(&ListSource)
 			.SelectionMode(ESelectionMode::SingleToggle)
-			.OnGenerateRow_Lambda([this](TSharedPtr<FSignalFlowDashboardEntry> Item, const TSharedRef<STableViewBase>& OwnerTable)
+			.OnGenerateRow_Lambda([this, EntryType](TSharedPtr<FSignalFlowDashboardEntry> Item, const TSharedRef<STableViewBase>& OwnerTable)
 			{
 				return SNew(SRowWidget, OwnerTable, Item)
-					   .Style(GetRowStyle());
+						.Style(GetRowStyle())
+						.Visibility_Lambda([this, EntryType]()
+						{
+							return GetSelectableListIsExpanded(EntryType) ? EVisibility::Visible : EVisibility::Collapsed;
+						});
 			})
 			.OnSelectionChanged(this, &FSignalFlowDashboardViewFactory::OnSelectionChanged)
 #if WITH_EDITOR
@@ -1104,6 +1134,88 @@ namespace UE::Audio::Insights
 			);
 	}
 
+	bool FSignalFlowDashboardViewFactory::GetSelectableListIsExpanded(const ESignalFlowEntryType EntryType) const
+	{
+		switch (EntryType)
+		{
+			case ESignalFlowEntryType::SoundSource:
+				return ActiveEntryMenuExpansionSettings.bShowActiveSources;
+
+			case ESignalFlowEntryType::AudioBus:
+				return ActiveEntryMenuExpansionSettings.bShowActiveBuses;
+
+			case ESignalFlowEntryType::Submix:
+				return ActiveEntryMenuExpansionSettings.bShowActiveSubmixes;
+
+			default:
+				ensureMsgf(false, TEXT("Unsupported ESignalFlowEntryType in FSignalFlowDashboardViewFactory::GetSelectableListIsExpanded"));
+				break;
+		}
+
+		return false;
+	}
+
+	void FSignalFlowDashboardViewFactory::SetSelectableListIsExpanded(const ESignalFlowEntryType EntryType, const bool bIsExpanded)
+	{
+		using namespace FSignalFlowDashboardViewFactoryPrivate;
+
+		switch (EntryType)
+		{
+			case ESignalFlowEntryType::SoundSource:
+				ActiveEntryMenuExpansionSettings.bShowActiveSources = bIsExpanded;
+
+				if (ActiveEntryMenusContainer.IsValid() && ensure(ActiveEntryMenusContainer->IsValidSlotIndex(ActiveSourcesSlotIndex)))
+				{
+					SetSelectableListSlotSize(ActiveEntryMenusContainer->GetSlot(ActiveSourcesSlotIndex), SourcesListHeightProp, bIsExpanded);
+				}
+				
+				break;
+
+			case ESignalFlowEntryType::AudioBus:
+				ActiveEntryMenuExpansionSettings.bShowActiveBuses = bIsExpanded;
+
+				if (ActiveEntryMenusContainer.IsValid() && ensure(ActiveEntryMenusContainer->IsValidSlotIndex(ActiveBusesSlotIndex)))
+				{
+					SetSelectableListSlotSize(ActiveEntryMenusContainer->GetSlot(ActiveBusesSlotIndex), BusListHeightProp, bIsExpanded);
+				}
+				break;
+
+			case ESignalFlowEntryType::Submix:
+				ActiveEntryMenuExpansionSettings.bShowActiveSubmixes = bIsExpanded;
+
+				if (ActiveEntryMenusContainer.IsValid() && ensure(ActiveEntryMenusContainer->IsValidSlotIndex(ActiveSubmixesSlotIndex)))
+				{
+					SetSelectableListSlotSize(ActiveEntryMenusContainer->GetSlot(ActiveSubmixesSlotIndex), SubmixesListHeightProp, bIsExpanded);
+				}
+				break;
+
+			default:
+				ensureMsgf(false, TEXT("Unsupported ESignalFlowEntryType in FSignalFlowDashboardViewFactory::SetSelectableListIsExpanded"));
+				break;
+		}
+
+#if WITH_EDITOR
+		FSignalFlowSettings::OnRequestWriteSettings.Broadcast();
+#endif // WITH_EDITOR	
+	}
+
+	void FSignalFlowDashboardViewFactory::SetSelectableListSlotSize(SVerticalBox::FSlot& OutMenuSlot, const float ExpandedStretchCoeff, const bool bIsExpanded)
+	{
+		using namespace FSignalFlowDashboardViewFactoryPrivate;
+
+		if (bIsExpanded)
+		{
+			OutMenuSlot.SetMaxHeight(0.0f); // Resets the max height to no max
+			OutMenuSlot.SetSizeToStretch(ExpandedStretchCoeff);
+		}
+		else
+		{
+			OutMenuSlot.SetSizeToAuto();
+			OutMenuSlot.SetMaxHeight(CollapsedActiveEntryMenuHeight);
+			OutMenuSlot.SetMinHeight(CollapsedActiveEntryMenuHeight);
+		}
+	}
+
 	void FSignalFlowDashboardViewFactory::SetSearchBoxFilterText(const FText& NewText)
 	{
 		SearchBoxFilterText = NewText;
@@ -2086,6 +2198,12 @@ namespace UE::Audio::Insights
 		{
 			Provider->SetConnectionWiresAreAnimated(bAnimateWires);
 		}
+
+		ActiveEntryMenuExpansionSettings = InSettings.ActiveEntryMenuExpansionSettings;
+
+		SetSelectableListIsExpanded(ESignalFlowEntryType::SoundSource, ActiveEntryMenuExpansionSettings.bShowActiveSources);
+		SetSelectableListIsExpanded(ESignalFlowEntryType::AudioBus, ActiveEntryMenuExpansionSettings.bShowActiveBuses);
+		SetSelectableListIsExpanded(ESignalFlowEntryType::Submix, ActiveEntryMenuExpansionSettings.bShowActiveSubmixes);
 	}
 
 	void FSignalFlowDashboardViewFactory::OnWriteEditorSettings(FSignalFlowSettings& OutSettings)
@@ -2100,6 +2218,7 @@ namespace UE::Audio::Insights
 		OutSettings.bShowNodeDetails = bShowNodeDetails;
 		OutSettings.bPauseGraphOnSelect = bPauseGraphOnSelect;
 		OutSettings.NodeDetailFilters = NodeDetailFilters;
+		OutSettings.ActiveEntryMenuExpansionSettings = ActiveEntryMenuExpansionSettings;
 		OutSettings.AmplitudeDisplayMode = bDisplayAmpPeakInDb ? EAudioAmplitudeDisplayMode::Decibels : EAudioAmplitudeDisplayMode::Linear;
 		OutSettings.bAnimateWires = bAnimateWires;
 
diff --git a/Engine/Plugins/AudioInsights/Source/AudioInsights/Public/AudioInsightsStyle.h b/Engine/Plugins/AudioInsights/Source/AudioInsights/Public/AudioInsightsStyle.h
--- a/Engine/Plugins/AudioInsights/Source/AudioInsights/Public/AudioInsightsStyle.h
+++ b/Engine/Plugins/AudioInsights/Source/AudioInsights/Public/AudioInsightsStyle.h
@@ -307,14 +307,24 @@ namespace UE::Audio::Insights
 				.SetUncheckedHoveredImage(RoundedWhiteAlphaBrush)
 			);
 
-			Set("AudioInsights.ExpandArrowToggle", FCheckBoxStyle(FAppStyle::Get().GetWidgetStyle<FCheckBoxStyle>("TransparentCheckBox"))
+			Set("AudioInsights.LeftExpandArrowToggle", FCheckBoxStyle(FAppStyle::Get().GetWidgetStyle<FCheckBoxStyle>("TransparentCheckBox"))
 				.SetCheckBoxType(ESlateCheckBoxType::ToggleButton)
-				.SetCheckedImage(IMAGE_BRUSH_SVG(TEXT("Icons/expand_arrow_up"), Icon12, FStyleColors::Foreground))
-				.SetCheckedHoveredImage(IMAGE_BRUSH_SVG(TEXT("Icons/expand_arrow_up"), Icon12, FStyleColors::ForegroundHover))
-				.SetCheckedPressedImage(IMAGE_BRUSH_SVG(TEXT("Icons/expand_arrow_up"), Icon12, FStyleColors::ForegroundHover))
-				.SetUncheckedImage(IMAGE_BRUSH_SVG(TEXT("Icons/expand_arrow_down"), Icon12, FStyleColors::Foreground))
-				.SetUncheckedHoveredImage(IMAGE_BRUSH_SVG(TEXT("Icons/expand_arrow_down"), Icon12, FStyleColors::ForegroundHover))
-				.SetUncheckedPressedImage(IMAGE_BRUSH_SVG(TEXT("Icons/expand_arrow_down"), Icon12, FStyleColors::ForegroundHover))
+				.SetCheckedImage(IMAGE_BRUSH_SVG(TEXT("Icons/expand_arrow_down"), Icon12, FStyleColors::Foreground))
+				.SetCheckedHoveredImage(IMAGE_BRUSH_SVG(TEXT("Icons/expand_arrow_down"), Icon12, FStyleColors::ForegroundHover))
+				.SetCheckedPressedImage(IMAGE_BRUSH_SVG(TEXT("Icons/expand_arrow_down"), Icon12, FStyleColors::ForegroundHover))
+				.SetUncheckedImage(IMAGE_BRUSH_SVG(TEXT("Icons/expand_arrow_left"), Icon12, FStyleColors::Foreground))
+				.SetUncheckedHoveredImage(IMAGE_BRUSH_SVG(TEXT("Icons/expand_arrow_left"), Icon12, FStyleColors::ForegroundHover))
+				.SetUncheckedPressedImage(IMAGE_BRUSH_SVG(TEXT("Icons/expand_arrow_left"), Icon12, FStyleColors::ForegroundHover))
+			);
+
+			Set("AudioInsights.RightExpandArrowToggle", FCheckBoxStyle(FAppStyle::Get().GetWidgetStyle<FCheckBoxStyle>("TransparentCheckBox"))
+				.SetCheckBoxType(ESlateCheckBoxType::ToggleButton)
+				.SetCheckedImage(IMAGE_BRUSH_SVG(TEXT("Icons/expand_arrow_down"), Icon12, FStyleColors::Foreground))
+				.SetCheckedHoveredImage(IMAGE_BRUSH_SVG(TEXT("Icons/expand_arrow_down"), Icon12, FStyleColors::ForegroundHover))
+				.SetCheckedPressedImage(IMAGE_BRUSH_SVG(TEXT("Icons/expand_arrow_down"), Icon12, FStyleColors::ForegroundHover))
+				.SetUncheckedImage(IMAGE_BRUSH_SVG(TEXT("Icons/expand_arrow_right"), Icon12, FStyleColors::Foreground))
+				.SetUncheckedHoveredImage(IMAGE_BRUSH_SVG(TEXT("Icons/expand_arrow_right"), Icon12, FStyleColors::ForegroundHover))
+				.SetUncheckedPressedImage(IMAGE_BRUSH_SVG(TEXT("Icons/expand_arrow_right"), Icon12, FStyleColors::ForegroundHover))
 			);
 
 			FSlateStyleRegistry::RegisterSlateStyle(*this);
diff --git a/Engine/Plugins/AudioInsights/Source/AudioInsights/Public/Settings/SignalFlowSettings.h b/Engine/Plugins/AudioInsights/Source/AudioInsights/Public/Settings/SignalFlowSettings.h
--- a/Engine/Plugins/AudioInsights/Source/AudioInsights/Public/Settings/SignalFlowSettings.h
+++ b/Engine/Plugins/AudioInsights/Source/AudioInsights/Public/Settings/SignalFlowSettings.h
@@ -79,6 +79,24 @@ struct FSignalFlowNodeDetailFilterSettings
 	bool bSendOutputVolume = true;
 };
 
+USTRUCT()
+struct FSignalFlowEntryTypeListExpansionSettings
+{
+	GENERATED_BODY()
+
+	/** Toggles expansion of the left hand Active Sources menu. */
+	UPROPERTY(EditAnywhere, config, Category = SignalFlow)
+	bool bShowActiveSources = true;
+
+	/** Toggles expansion of the left hand Active Buses menu. */
+	UPROPERTY(EditAnywhere, config, Category = SignalFlow)
+	bool bShowActiveBuses = true;
+
+	/** Toggles expansion of the left hand Active Submixes menu. */
+	UPROPERTY(EditAnywhere, config, Category = SignalFlow)
+	bool bShowActiveSubmixes = true;
+};
+
 USTRUCT()
 struct FSignalFlowSettings
 {
@@ -108,6 +126,10 @@ struct FSignalFlowSettings
 	UPROPERTY(EditAnywhere, config, Category = SignalFlow)
 	FSignalFlowNodeDetailFilterSettings NodeDetailFilters;
 
+	/** Expansion settings for the Active Entry menus on the left hand side of the Signal Flow Tab */
+	UPROPERTY(EditAnywhere, config, Category = SignalFlow, meta = (DisplayName = "Active Entry Menu Expansion"))
+	FSignalFlowEntryTypeListExpansionSettings ActiveEntryMenuExpansionSettings;
+
 	/** Toggles whether connection wire animations driven by amplitude envelope values are enabled in the Signal Flow graph. */
 	UPROPERTY(EditAnywhere, config, Category = SignalFlow)
 	bool bAnimateWires = true;
diff --git a/Engine/Plugins/AudioInsights/Source/AudioInsights/Public/Views/SignalFlowDashboardViewFactory.h b/Engine/Plugins/AudioInsights/Source/AudioInsights/Public/Views/SignalFlowDashboardViewFactory.h
--- a/Engine/Plugins/AudioInsights/Source/AudioInsights/Public/Views/SignalFlowDashboardViewFactory.h
+++ b/Engine/Plugins/AudioInsights/Source/AudioInsights/Public/Views/SignalFlowDashboardViewFactory.h
@@ -97,7 +97,11 @@ namespace UE::Audio::Insights
 
 		TSharedRef<SWidget> CreateSelectionPanel();
 		TSharedRef<SWidget> CreateGraphSection();
-		TSharedRef<SListView<TSharedPtr<FSignalFlowDashboardEntry>>> CreateSelectableList(const FText& HeadingText, const TArray<TSharedPtr<FSignalFlowDashboardEntry>>& ListSource);
+
+		TSharedRef<SListView<TSharedPtr<FSignalFlowDashboardEntry>>> CreateSelectableList(const FText& HeadingText, const TArray<TSharedPtr<FSignalFlowDashboardEntry>>& ListSource, const ESignalFlowEntryType EntryType);
+		bool GetSelectableListIsExpanded(const ESignalFlowEntryType EntryType) const;
+		void SetSelectableListIsExpanded(const ESignalFlowEntryType EntryType, const bool bIsExpanded);
+		void SetSelectableListSlotSize(SVerticalBox::FSlot& OutMenuSlot, const float ExpandedStretchCoeff, const bool bIsExpanded);
 
 		void SetSearchBoxFilterText(const FText& NewText);
 		void OnSelectionChanged(TSharedPtr<FSignalFlowDashboardEntry> SelectedItem, ESelectInfo::Type SelectInfo);
@@ -147,14 +151,19 @@ namespace UE::Audio::Insights
 		TSharedPtr<FSignalFlowTraceProvider> SignalFlowProvider;
 		TSharedPtr<SSignalFlowDashboard> DashboardWidget;
 
+		TSharedPtr<SVerticalBox> ActiveEntryMenusContainer;
+
 		TSharedPtr<SListView<TSharedPtr<FSignalFlowDashboardEntry>>> ActiveSourcesListView;
 		TArray<TSharedPtr<FSignalFlowDashboardEntry>> ActiveSources;
+		int32 ActiveSourcesSlotIndex = INDEX_NONE;
 
 		TSharedPtr<SListView<TSharedPtr<FSignalFlowDashboardEntry>>> ActiveBusesListView;
 		TArray<TSharedPtr<FSignalFlowDashboardEntry>> ActiveBuses;
+		int32 ActiveBusesSlotIndex = INDEX_NONE;
 
 		TSharedPtr<SListView<TSharedPtr<FSignalFlowDashboardEntry>>> ActiveSubmixesListView;
 		TArray<TSharedPtr<FSignalFlowDashboardEntry>> ActiveSubmixes;
+		int32 ActiveSubmixesSlotIndex = INDEX_NONE;
 
 		TSortedMap<FSignalFlowEntryKey, TSharedPtr<ISignalFlowNode>> FilteredEntries;
 		TArray<TSharedPtr<ISignalFlowNode>> FilteredGraphNodes;
@@ -169,6 +178,7 @@ namespace UE::Audio::Insights
 		FText SearchBoxFilterText;
 
 		FSignalFlowNodeDetailFilterSettings NodeDetailFilters;
+		FSignalFlowEntryTypeListExpansionSettings ActiveEntryMenuExpansionSettings;
 
 		ESignalFlowJustification GraphJustification = ESignalFlowJustification::Edge;
 
