[Audio Insights] Output Meter: added trace message for loudness values + adjusted dashboard so Loudness meters are available in standalone.

#Synced-CL 51090472

diff --git a/Engine/Plugins/AudioInsights/Source/AudioInsights/Private/Messages/OutputMeterTraceMessages.cpp b/Engine/Plugins/AudioInsights/Source/AudioInsights/Private/Messages/OutputMeterTraceMessages.cpp
--- a/Engine/Plugins/AudioInsights/Source/AudioInsights/Private/Messages/OutputMeterTraceMessages.cpp
+++ b/Engine/Plugins/AudioInsights/Source/AudioInsights/Private/Messages/OutputMeterTraceMessages.cpp
@@ -7,6 +7,7 @@ namespace UE::Audio::Insights
 	{
 		const FName MainSubmixLoaded = "OutputMeterMainSubmixLoaded";
 		const FName TruePeak = "OutputMeterTruePeak";
+		const FName LKFSValues = "OutputMeterLKFSValues";
 	};
 
 	// FOutputMeterMessageBase
@@ -45,4 +46,23 @@ namespace UE::Audio::Insights
 		// Returning fixed size of the struct (no dynamically sized members in it)
 		return sizeof(FOutputMeterTruePeakMessage);
 	}
+
+	// FOutputMeterLKFSValuesMessage
+	FOutputMeterLKFSValuesMessage::FOutputMeterLKFSValuesMessage(const Trace::IAnalyzer::FOnEventContext& InContext)
+		: FOutputMeterMessageBase(InContext)
+	{
+		const Trace::IAnalyzer::FEventData& EventData = InContext.EventData;
+
+		LongTermLoudness  = EventData.GetValue<float>("LongTermLoudness");
+		ShortTermLoudness = EventData.GetValue<float>("ShortTermLoudness");
+		MomentaryLoudness = EventData.GetValue<float>("MomentaryLoudness");
+		LoudnessRangeLowerBound = EventData.GetValue<float>("LoudnessRangeLowerBound");
+		LoudnessRangeUpperBound = EventData.GetValue<float>("LoudnessRangeUpperBound");
+	}
+
+	uint32 FOutputMeterLKFSValuesMessage::GetSizeOf() const
+	{
+		// Returning fixed size of the struct (no dynamically sized members in it)
+		return sizeof(FOutputMeterLKFSValuesMessage);
+	}
 } // namespace UE::Audio::Insights
diff --git a/Engine/Plugins/AudioInsights/Source/AudioInsights/Private/Providers/OutputMeterTraceProvider.cpp b/Engine/Plugins/AudioInsights/Source/AudioInsights/Private/Providers/OutputMeterTraceProvider.cpp
--- a/Engine/Plugins/AudioInsights/Source/AudioInsights/Private/Providers/OutputMeterTraceProvider.cpp
+++ b/Engine/Plugins/AudioInsights/Source/AudioInsights/Private/Providers/OutputMeterTraceProvider.cpp
@@ -33,6 +33,7 @@ namespace UE::Audio::Insights
 				
 				Builder.RouteEvent(RouteId_MainSubmixLoaded, "Audio", "MainSubmixLoaded");
 				Builder.RouteEvent(RouteId_TruePeakMaxValue, "Audio", "MainSubmixTruePeakMaxValue");
+				Builder.RouteEvent(RouteId_LKFSValues,       "Audio", "MainSubmixLKFSValues");
 			}
 
 			virtual bool OnHandleEvent(uint16 RouteId, EStyle Style, const FOnEventContext& Context) override
@@ -55,6 +56,12 @@ namespace UE::Audio::Insights
 						break;
 					}
 
+					case RouteId_LKFSValues:
+					{
+						CacheMessage<FOutputMeterLKFSValuesMessage>(Context, Messages.LKFSValuesMessages);
+						break;
+					}
+
 					default:
 					{
 						return OnEventFailure(RouteId, Style, Context);
@@ -75,7 +82,8 @@ namespace UE::Audio::Insights
 			enum : uint16
 			{
 				RouteId_MainSubmixLoaded,
-				RouteId_TruePeakMaxValue
+				RouteId_TruePeakMaxValue,
+				RouteId_LKFSValues
 			};
 
 			TraceServices::IAnalysisSession& Session;
@@ -125,6 +133,20 @@ namespace UE::Audio::Insights
 			}
 		});
 
+		ProcessMessageQueue<FOutputMeterLKFSValuesMessage>(TraceMessages.LKFSValuesMessages, GetEntryFunc,
+		[this](const FOutputMeterLKFSValuesMessage& Msg, TSharedPtr<FOutputMeterDashboardEntry>* OutEntry)
+		{
+			if (OutEntry)
+			{
+				FOutputMeterDashboardEntry& EntryRef = *OutEntry->Get();
+				EntryRef.LongTermLoudness  = Msg.LongTermLoudness;
+				EntryRef.ShortTermLoudness = Msg.ShortTermLoudness;
+				EntryRef.MomentaryLoudness = Msg.MomentaryLoudness;
+				EntryRef.LoudnessRangeLowerBound = Msg.LoudnessRangeLowerBound;
+				EntryRef.LoudnessRangeUpperBound = Msg.LoudnessRangeUpperBound;
+			}
+		});
+
 		return true;
 	}
 
@@ -175,6 +197,7 @@ namespace UE::Audio::Insights
 		struct CachedEntryInfo
 		{
 			TOptional<FOutputMeterTruePeakMessage> TruePeakMessage;
+			TOptional<FOutputMeterLKFSValuesMessage> LKFSValuesMessage;
 		};
 
 		TArray<uint32> MainSubmixIdArray;
@@ -195,6 +218,13 @@ namespace UE::Audio::Insights
 			{
 				CachedEntryInfos[Index].TruePeakMessage = *FoundTruePeakCachedMessage;
 			}
+
+			// LKFSValues
+			const FOutputMeterLKFSValuesMessage* FoundLKFSValuesCachedMessage = CacheManager.FindClosestMessage<FOutputMeterLKFSValuesMessage>(OutputMeterMessageNames::LKFSValues, InTimeMarker, MainSubmixId);
+			if (FoundLKFSValuesCachedMessage)
+			{
+				CachedEntryInfos[Index].LKFSValuesMessage = *FoundLKFSValuesCachedMessage;
+			}
 		});
 
 		// Update the device entries with the collected info
@@ -214,6 +244,25 @@ namespace UE::Audio::Insights
 					}
 				});
 			}
+
+			// LKFSValues
+			if (CachedEntryInfo.LKFSValuesMessage.IsSet())
+			{
+				const FOutputMeterLKFSValuesMessage& OutputMeterLKFSValuesMessage = CachedEntryInfo.LKFSValuesMessage.GetValue();
+
+				UpdateDeviceEntry(OutputMeterLKFSValuesMessage.DeviceId, OutputMeterLKFSValuesMessage.SubmixId, [&OutputMeterLKFSValuesMessage](TSharedPtr<FOutputMeterDashboardEntry>& Entry)
+				{
+					if (ensure(Entry.IsValid()))
+					{
+						Entry->Timestamp = OutputMeterLKFSValuesMessage.Timestamp;
+						Entry->LongTermLoudness  = OutputMeterLKFSValuesMessage.LongTermLoudness;
+						Entry->ShortTermLoudness = OutputMeterLKFSValuesMessage.ShortTermLoudness;
+						Entry->MomentaryLoudness = OutputMeterLKFSValuesMessage.MomentaryLoudness;
+						Entry->LoudnessRangeLowerBound = OutputMeterLKFSValuesMessage.LoudnessRangeLowerBound;
+						Entry->LoudnessRangeUpperBound = OutputMeterLKFSValuesMessage.LoudnessRangeUpperBound;
+					}
+				});
+			}
 		}
 	}
 } // namespace UE::Audio::Insights
diff --git a/Engine/Plugins/AudioInsights/Source/AudioInsights/Private/Views/OutputMeterDashboardViewFactory.cpp b/Engine/Plugins/AudioInsights/Source/AudioInsights/Private/Views/OutputMeterDashboardViewFactory.cpp
--- a/Engine/Plugins/AudioInsights/Source/AudioInsights/Private/Views/OutputMeterDashboardViewFactory.cpp
+++ b/Engine/Plugins/AudioInsights/Source/AudioInsights/Private/Views/OutputMeterDashboardViewFactory.cpp
@@ -63,6 +63,40 @@ namespace UE::Audio::Insights
 	{
 		using namespace AudioWidgetsCore;
 
+		// LoudnessScaleParams
+		FLoudnessMeterWidgetView::FLoudnessScaleParams LoudnessScaleParams;
+
+		LoudnessScaleParams.Range.BindLambda([this]() { return LoudnessMeterSettings.LoudnessScaleRange; });
+		LoudnessScaleParams.OnRangeValueChanged.BindLambda([this](int32 Value) { LoudnessMeterSettings.LoudnessScaleRange = Value; });
+		LoudnessScaleParams.OnRangeValueCommitted.BindLambda([this](int32 Value, ETextCommit::Type CommitType)
+		{
+			LoudnessMeterSettings.LoudnessScaleRange = Value;
+		});
+
+		LoudnessScaleParams.Offset.BindLambda([this]() { return LoudnessMeterSettings.LoudnessScaleOffset; });
+		LoudnessScaleParams.OnOffsetValueChanged.BindLambda([this](int32 Value) { LoudnessMeterSettings.LoudnessScaleOffset = Value; });
+		LoudnessScaleParams.OnOffsetValueCommitted.BindLambda([this](int32 Value, ETextCommit::Type CommitType)
+		{
+			LoudnessMeterSettings.LoudnessScaleOffset = Value;
+		});
+
+		LoudnessScaleParams.Target.BindLambda([this]() { return LoudnessMeterSettings.LoudnessScaleTarget; });
+		LoudnessScaleParams.OnTargetValueChanged.BindLambda([this](int32 Value) { LoudnessMeterSettings.LoudnessScaleTarget = Value; });
+		LoudnessScaleParams.OnTargetValueCommitted.BindLambda([this](int32 Value, ETextCommit::Type CommitType)
+		{
+			LoudnessMeterSettings.LoudnessScaleTarget = Value;
+		});
+
+		LoudnessScaleParams.TargetColor.BindLambda([this]() { return LoudnessMeterSettings.TargetColor; });
+		LoudnessScaleParams.OnTargetColorChanged.BindLambda([this](FLinearColor NewColor) {	LoudnessMeterSettings.TargetColor = NewColor; });
+		LoudnessScaleParams.OnResetTargetColorRequested.BindLambda([this]() { LoudnessMeterSettings.TargetColor = FLoudnessMeterWidgetView::GetDefaultTargetColor(); });
+
+		LoudnessScaleParams.bRelativeScale.BindLambda([this]() { return LoudnessMeterSettings.bUseRelativeLoudnessScale; });
+		LoudnessScaleParams.OnRelativeScaleToggleRequested.BindLambda([this]()
+		{
+			LoudnessMeterSettings.bUseRelativeLoudnessScale = !LoudnessMeterSettings.bUseRelativeLoudnessScale;
+		});
+
 		// ValuesDisplayOrderParams
 		FLoudnessMeterWidgetView::FDisplayOrderParams ValuesDisplayOrderParams;
 
@@ -82,11 +116,40 @@ namespace UE::Audio::Insights
 		});
 
 		// LoudnessMetrics
-		LoudnessMeterSettings.TruePeak.bShowValue = false;
-		LoudnessMeterSettings.TruePeak.bShowMeter = true;
-
 		FLoudnessMeterWidgetView::FLoudnessMetric LoudnessMetrics[] =
 		{
+			{
+				.Name = GET_MEMBER_NAME_CHECKED(FLoudnessMeterSettings, LongTermLoudness),
+				.DisplayName = LOCTEXT("LongTermLoudnessDisplayName", "Long Term"),
+				.MeterMetric = EAudioMeterMetric::Loudness,
+				.Value = TAttribute<TOptional<float>>::CreateSPLambda(this, [this]() { return LongTermLoudness; }),
+				.bShowValue = true,
+				.bShowMeter = false,
+			},
+			{
+				.Name = GET_MEMBER_NAME_CHECKED(FLoudnessMeterSettings, ShortTermLoudness),
+				.DisplayName = LOCTEXT("ShortTermLoudnessDisplayName", "Short Term"),
+				.MeterMetric = EAudioMeterMetric::Loudness,
+				.Value = TAttribute<TOptional<float>>::CreateSPLambda(this, [this]() { return ShortTermLoudness; }),
+				.bShowValue = false,
+				.bShowMeter = true,
+			},
+			{
+				.Name = GET_MEMBER_NAME_CHECKED(FLoudnessMeterSettings, MomentaryLoudness),
+				.DisplayName = LOCTEXT("MomentaryLoudnessDisplayName", "Momentary"),
+				.MeterMetric = EAudioMeterMetric::Loudness,
+				.Value = TAttribute<TOptional<float>>::CreateSPLambda(this, [this]() { return MomentaryLoudness; }),
+				.bShowValue = false,
+				.bShowMeter = true,
+			},
+			{
+				.Name = GET_MEMBER_NAME_CHECKED(FLoudnessMeterSettings, LoudnessRange),
+				.DisplayName = LOCTEXT("LoudnessRangeDisplayName", "Range"),
+				.MeterMetric = EAudioMeterMetric::Loudness,
+				.Range = TAttribute<TOptional<FFloatInterval>>::CreateSPLambda(this, [this]() { return FFloatInterval{ LoudnessRangeLowerBound, LoudnessRangeUpperBound }; }),
+				.bShowValue = true,
+				.bShowMeter = false,
+			},
 			{
 				.Name = GET_MEMBER_NAME_CHECKED(FLoudnessMeterSettings, TruePeak),
 				.DisplayName = LOCTEXT("TruePeakDisplayName", "True Peak"),
@@ -150,6 +213,7 @@ namespace UE::Audio::Insights
 			}
 		}
 
+		LoudnessMeterWidgetView.InitLoudnessScale(LoudnessScaleParams);
 		LoudnessMeterWidgetView.InitValuesDisplayOrder(ValuesDisplayOrderParams);
 		LoudnessMeterWidgetView.InitMetersDisplayOrder(MetersDisplayOrderParams);
 
@@ -168,7 +232,6 @@ namespace UE::Audio::Insights
 		}
 
 		return SNew(SOverlay)
-			.Clipping(EWidgetClipping::ClipToBounds)
 			+ SOverlay::Slot()
 			[
 				SNew(SColorBlock)
@@ -179,16 +242,15 @@ namespace UE::Audio::Insights
 				SNew(SSplitter)
 				.Orientation(Orient_Vertical)
 				.ResizeMode(ESplitterResizeMode::Fill)
-				.Clipping(EWidgetClipping::ClipToBounds)
 				.PhysicalSplitterHandleSize(4.0f)
 				.HitDetectionSplitterHandleSize(6.0f)
 				+ SSplitter::Slot()
-				.Value(0.5f)
+				.Value(0.6f)
 				[
 					LoudnessMeterWidgetView.MakeWidget()
 				]
 				+ SSplitter::Slot()
-				.Value(0.5f)
+				.Value(0.4f)
 				[
 					SNew(SVerticalBox)
 					+ SVerticalBox::Slot()
@@ -271,6 +333,11 @@ namespace UE::Audio::Insights
 		{
 			const FOutputMeterDashboardEntry& OutputMeterEntry = FOutputMeterDashboardViewFactoryPrivate::CastEntry(*DataViewEntries[0].Get());
 			TruePeakMaxValueDb = FMath::Max<float>(TruePeakMaxValueDb, OutputMeterEntry.TruePeakMaxValueDb);
+			LongTermLoudness  = OutputMeterEntry.LongTermLoudness;
+			ShortTermLoudness = OutputMeterEntry.ShortTermLoudness;
+			MomentaryLoudness = OutputMeterEntry.MomentaryLoudness;
+			LoudnessRangeLowerBound = OutputMeterEntry.LoudnessRangeLowerBound;
+			LoudnessRangeUpperBound = OutputMeterEntry.LoudnessRangeUpperBound;
 		}
 	}
 
diff --git a/Engine/Plugins/AudioInsights/Source/AudioInsights/Private/Views/OutputMeterDashboardViewFactory.h b/Engine/Plugins/AudioInsights/Source/AudioInsights/Private/Views/OutputMeterDashboardViewFactory.h
--- a/Engine/Plugins/AudioInsights/Source/AudioInsights/Private/Views/OutputMeterDashboardViewFactory.h
+++ b/Engine/Plugins/AudioInsights/Source/AudioInsights/Private/Views/OutputMeterDashboardViewFactory.h
@@ -44,8 +44,14 @@ namespace UE::Audio::Insights
 		// Loudness meters related members
 		FLoudnessMeterSettings LoudnessMeterSettings;
 		AudioWidgetsCore::FLoudnessMeterWidgetView LoudnessMeterWidgetView;
-		float TruePeakMaxValueDb = MIN_VOLUME_DECIBELS;
 		bool bIsLoudnessMeterInitialized = false;
+
+		float TruePeakMaxValueDb = MIN_VOLUME_DECIBELS;
+		float LongTermLoudness  = MIN_VOLUME_DECIBELS;
+		float ShortTermLoudness = MIN_VOLUME_DECIBELS;
+		float MomentaryLoudness = MIN_VOLUME_DECIBELS;
+		float LoudnessRangeLowerBound = 0.0f;
+		float LoudnessRangeUpperBound = 0.0f;
 	};
 } // namespace UE::Audio::Insights
 
diff --git a/Engine/Plugins/AudioInsights/Source/AudioInsights/Public/Messages/OutputMeterTraceMessages.h b/Engine/Plugins/AudioInsights/Source/AudioInsights/Public/Messages/OutputMeterTraceMessages.h
--- a/Engine/Plugins/AudioInsights/Source/AudioInsights/Public/Messages/OutputMeterTraceMessages.h
+++ b/Engine/Plugins/AudioInsights/Source/AudioInsights/Public/Messages/OutputMeterTraceMessages.h
@@ -12,6 +12,7 @@ namespace UE::Audio::Insights
 	{
 		extern const FName MainSubmixLoaded;
 		extern const FName TruePeak;
+		extern const FName LKFSValues;
 	};
 
 	struct FOutputMeterMessageBase : public IAudioCachedMessage
@@ -45,6 +46,21 @@ namespace UE::Audio::Insights
 		float TruePeakMaxValueDb = MIN_VOLUME_DECIBELS;
 	};
 
+	struct FOutputMeterLKFSValuesMessage : public FOutputMeterMessageBase
+	{
+		FOutputMeterLKFSValuesMessage() = default;
+		FOutputMeterLKFSValuesMessage(const Trace::IAnalyzer::FOnEventContext& InContext);
+
+		virtual const FName GetMessageName() const override { return OutputMeterMessageNames::LKFSValues; }
+		virtual uint32 GetSizeOf() const override;
+
+		float LongTermLoudness  = MIN_VOLUME_DECIBELS;
+		float ShortTermLoudness = MIN_VOLUME_DECIBELS;
+		float MomentaryLoudness = MIN_VOLUME_DECIBELS;
+		float LoudnessRangeLowerBound = 0.0f;
+		float LoudnessRangeUpperBound = 0.0f;
+	};
+
 	class FOutputMeterDashboardEntry : public IDashboardDataViewEntry
 	{
 	public:
@@ -61,12 +77,19 @@ namespace UE::Audio::Insights
 		uint32 SubmixId = static_cast<uint32>(INDEX_NONE);
 
 		float TruePeakMaxValueDb = MIN_VOLUME_DECIBELS;
+
+		float LongTermLoudness  = MIN_VOLUME_DECIBELS;
+		float ShortTermLoudness = MIN_VOLUME_DECIBELS;
+		float MomentaryLoudness = MIN_VOLUME_DECIBELS;
+		float LoudnessRangeLowerBound = 0.0f;
+		float LoudnessRangeUpperBound = 0.0f;
 	};
 
 	class FOutputMeterMessages
 	{
 		TAnalyzerMessageQueue<FOutputMeterMainSubmixLoadedMessage> MainSubmixLoadedMessages;
 		TAnalyzerMessageQueue<FOutputMeterTruePeakMessage> TruePeakMessages;
+		TAnalyzerMessageQueue<FOutputMeterLKFSValuesMessage> LKFSValuesMessages;
 
 		friend class FOutputMeterTraceProvider;
 	};
