[AudioWidgets] [Audio Insights] Loudness Meters: added "hold max value" option to Short Term Loudness, Momentary Loudness and TruePeak meters.

#Synced-CL 51090472

diff --git a/Engine/Plugins/AudioInsights/Source/AudioInsights/Private/Views/OutputMeterDashboardViewFactory.cpp b/Engine/Plugins/AudioInsights/Source/AudioInsights/Private/Views/OutputMeterDashboardViewFactory.cpp
--- a/Engine/Plugins/AudioInsights/Source/AudioInsights/Private/Views/OutputMeterDashboardViewFactory.cpp
+++ b/Engine/Plugins/AudioInsights/Source/AudioInsights/Private/Views/OutputMeterDashboardViewFactory.cpp
@@ -122,7 +122,7 @@ namespace UE::Audio::Insights
 				.Name = GET_MEMBER_NAME_CHECKED(FLoudnessMeterSettings, LongTermLoudness),
 				.DisplayName = LOCTEXT("LongTermLoudnessDisplayName", "Long Term"),
 				.MeterMetric = EAudioMeterMetric::Loudness,
-				.Value = TAttribute<TOptional<float>>::CreateSPLambda(this, [this]() { return LongTermLoudness; }),
+				.Value = TAttribute<TOptional<float>>::CreateSPLambda(this, [this]() { return LoudnessMeterResults.LongTermLoudness; }),
 				.bShowValue = true,
 				.bShowMeter = false,
 			},
@@ -130,23 +130,31 @@ namespace UE::Audio::Insights
 				.Name = GET_MEMBER_NAME_CHECKED(FLoudnessMeterSettings, ShortTermLoudness),
 				.DisplayName = LOCTEXT("ShortTermLoudnessDisplayName", "Short Term"),
 				.MeterMetric = EAudioMeterMetric::Loudness,
-				.Value = TAttribute<TOptional<float>>::CreateSPLambda(this, [this]() { return ShortTermLoudness; }),
+				.Value = TAttribute<TOptional<float>>::CreateSPLambda(this, [this]() { return LoudnessMeterResults.ShortTermLoudness; }),
+				.MaxValue = TAttribute<TOptional<float>>::CreateSPLambda(this, [this]() { return LoudnessMeterMaxResults.ShortTermLoudness; }),
 				.bShowValue = false,
 				.bShowMeter = true,
+				.bHoldMaxForValue = false,
+				.bHoldMaxForMeter = false,
+				.OnClicked = FSimpleDelegate::CreateSPLambda(this, [this]() { LoudnessMeterMaxResults.ShortTermLoudness = MIN_VOLUME_DECIBELS; }),
 			},
 			{
 				.Name = GET_MEMBER_NAME_CHECKED(FLoudnessMeterSettings, MomentaryLoudness),
 				.DisplayName = LOCTEXT("MomentaryLoudnessDisplayName", "Momentary"),
 				.MeterMetric = EAudioMeterMetric::Loudness,
-				.Value = TAttribute<TOptional<float>>::CreateSPLambda(this, [this]() { return MomentaryLoudness; }),
+				.Value = TAttribute<TOptional<float>>::CreateSPLambda(this, [this]() { return LoudnessMeterResults.MomentaryLoudness; }),
+				.MaxValue = TAttribute<TOptional<float>>::CreateSPLambda(this, [this]() { return LoudnessMeterMaxResults.MomentaryLoudness; }),
 				.bShowValue = false,
 				.bShowMeter = true,
+				.bHoldMaxForValue = false,
+				.bHoldMaxForMeter = false,
+				.OnClicked = FSimpleDelegate::CreateSPLambda(this, [this]() { LoudnessMeterMaxResults.MomentaryLoudness = MIN_VOLUME_DECIBELS; }),
 			},
 			{
 				.Name = GET_MEMBER_NAME_CHECKED(FLoudnessMeterSettings, LoudnessRange),
 				.DisplayName = LOCTEXT("LoudnessRangeDisplayName", "Range"),
 				.MeterMetric = EAudioMeterMetric::Loudness,
-				.Range = TAttribute<TOptional<FFloatInterval>>::CreateSPLambda(this, [this]() { return FFloatInterval{ LoudnessRangeLowerBound, LoudnessRangeUpperBound }; }),
+				.Range = TAttribute<TOptional<FFloatInterval>>::CreateSPLambda(this, [this]() { return FFloatInterval{ LoudnessMeterResults.LoudnessRangeLowerBound, LoudnessMeterResults.LoudnessRangeUpperBound }; }),
 				.bShowValue = true,
 				.bShowMeter = false,
 			},
@@ -154,11 +162,13 @@ namespace UE::Audio::Insights
 				.Name = GET_MEMBER_NAME_CHECKED(FLoudnessMeterSettings, TruePeak),
 				.DisplayName = LOCTEXT("TruePeakDisplayName", "True Peak"),
 				.MeterMetric = EAudioMeterMetric::Decibels,
-				.Value = TAttribute<TOptional<float>>::CreateSPLambda(this, [this]()
-				{
-					return TruePeakMaxValueDb;
-				}),
-				.OnClicked = FSimpleDelegate::CreateSPLambda(this, [this]() { TruePeakMaxValueDb = MIN_VOLUME_DECIBELS; })
+				.Value = TAttribute<TOptional<float>>::CreateSPLambda(this, [this]() { return LoudnessMeterResults.TruePeakDb; }),
+				.MaxValue = TAttribute<TOptional<float>>::CreateSPLambda(this, [this]() { return LoudnessMeterMaxResults.TruePeakDb; }),
+				.bShowValue = true,
+				.bShowMeter = false,
+				.bHoldMaxForValue = true,
+				.bHoldMaxForMeter = true,
+				.OnClicked = FSimpleDelegate::CreateSPLambda(this, [this]() { LoudnessMeterMaxResults.TruePeakDb = MIN_VOLUME_DECIBELS; })
 			}
 		};
 
@@ -176,7 +186,22 @@ namespace UE::Audio::Insights
 				LoudnessMetric.OnShowValueToggleRequested.BindLambda([this, DisplayOptionsProperty]()
 				{
 					FLoudnessMeterDisplayOptions* DisplayOptions = DisplayOptionsProperty->ContainerPtrToValuePtr<FLoudnessMeterDisplayOptions>(&LoudnessMeterSettings);
-					DisplayOptions->bShowValue = (!DisplayOptions->bShowValue) ? true : false;
+					DisplayOptions->bShowValue = !DisplayOptions->bShowValue;
+
+					LoudnessMeterWidgetView.RefreshVisibleLoudnessMetrics();
+				});
+
+				LoudnessMetric.OnShowValueToggleFromDragDropRequested.BindLambda([this, DisplayOptionsProperty]()
+				{
+					FLoudnessMeterDisplayOptions* DisplayOptions = DisplayOptionsProperty->ContainerPtrToValuePtr<FLoudnessMeterDisplayOptions>(&LoudnessMeterSettings);
+					DisplayOptions->bShowValue = !DisplayOptions->bShowValue;
+
+					if (DisplayOptions->bShowValue)
+					{
+						DisplayOptions->bHoldMaxForValue = DisplayOptions->bHoldMaxForMeter;
+					}
+
+					LoudnessMeterWidgetView.RefreshVisibleLoudnessMetrics();
 				});
 
 				// Show Meter
@@ -189,7 +214,48 @@ namespace UE::Audio::Insights
 				LoudnessMetric.OnShowMeterToggleRequested.BindLambda([this, DisplayOptionsProperty]()
 				{
 					FLoudnessMeterDisplayOptions* DisplayOptions = DisplayOptionsProperty->ContainerPtrToValuePtr<FLoudnessMeterDisplayOptions>(&LoudnessMeterSettings);
-					DisplayOptions->bShowMeter = (!DisplayOptions->bShowMeter) ? true : false;
+					DisplayOptions->bShowMeter = !DisplayOptions->bShowMeter;
+
+					LoudnessMeterWidgetView.RefreshVisibleLoudnessMetrics();
+				});
+
+				LoudnessMetric.OnShowMeterToggleFromDragDropRequested.BindLambda([this, DisplayOptionsProperty]()
+				{
+					FLoudnessMeterDisplayOptions* DisplayOptions = DisplayOptionsProperty->ContainerPtrToValuePtr<FLoudnessMeterDisplayOptions>(&LoudnessMeterSettings);
+					DisplayOptions->bShowMeter = !DisplayOptions->bShowMeter;
+
+					if (DisplayOptions->bShowMeter)
+					{
+						DisplayOptions->bHoldMaxForMeter = DisplayOptions->bHoldMaxForValue;
+					}
+
+					LoudnessMeterWidgetView.RefreshVisibleLoudnessMetrics();
+				});
+
+				// Hold Max for Value
+				LoudnessMetric.bHoldMaxForValue.BindLambda([this, DisplayOptionsProperty]()
+				{
+					const FLoudnessMeterDisplayOptions* DisplayOptions = DisplayOptionsProperty->ContainerPtrToValuePtr<FLoudnessMeterDisplayOptions>(&LoudnessMeterSettings);
+					return DisplayOptions->bHoldMaxForValue;
+				});
+
+				LoudnessMetric.OnHoldMaxForValueToggleRequested.BindLambda([this, DisplayOptionsProperty, LoudnessMetricName = LoudnessMetric.Name]()
+				{
+					FLoudnessMeterDisplayOptions* DisplayOptions = DisplayOptionsProperty->ContainerPtrToValuePtr<FLoudnessMeterDisplayOptions>(&LoudnessMeterSettings);
+					DisplayOptions->bHoldMaxForValue = !DisplayOptions->bHoldMaxForValue;
+				});
+
+				// Hold Max for Meter
+				LoudnessMetric.bHoldMaxForMeter.BindLambda([this, DisplayOptionsProperty]()
+				{
+					const FLoudnessMeterDisplayOptions* DisplayOptions = DisplayOptionsProperty->ContainerPtrToValuePtr<FLoudnessMeterDisplayOptions>(&LoudnessMeterSettings);
+					return DisplayOptions->bHoldMaxForMeter;
+				});
+
+				LoudnessMetric.OnHoldMaxForMeterToggleRequested.BindLambda([this, DisplayOptionsProperty, LoudnessMetricName = LoudnessMetric.Name]()
+				{
+					FLoudnessMeterDisplayOptions* DisplayOptions = DisplayOptionsProperty->ContainerPtrToValuePtr<FLoudnessMeterDisplayOptions>(&LoudnessMeterSettings);
+					DisplayOptions->bHoldMaxForMeter = !DisplayOptions->bHoldMaxForMeter;
 				});
 
 				// Color
@@ -332,12 +398,16 @@ namespace UE::Audio::Insights
 		if (!DataViewEntries.IsEmpty())
 		{
 			const FOutputMeterDashboardEntry& OutputMeterEntry = FOutputMeterDashboardViewFactoryPrivate::CastEntry(*DataViewEntries[0].Get());
-			TruePeakMaxValueDb = FMath::Max<float>(TruePeakMaxValueDb, OutputMeterEntry.TruePeakMaxValueDb);
-			LongTermLoudness  = OutputMeterEntry.LongTermLoudness;
-			ShortTermLoudness = OutputMeterEntry.ShortTermLoudness;
-			MomentaryLoudness = OutputMeterEntry.MomentaryLoudness;
-			LoudnessRangeLowerBound = OutputMeterEntry.LoudnessRangeLowerBound;
-			LoudnessRangeUpperBound = OutputMeterEntry.LoudnessRangeUpperBound;
+			LoudnessMeterResults.LongTermLoudness  = OutputMeterEntry.LongTermLoudness;
+			LoudnessMeterResults.ShortTermLoudness = OutputMeterEntry.ShortTermLoudness;
+			LoudnessMeterResults.MomentaryLoudness = OutputMeterEntry.MomentaryLoudness;
+			LoudnessMeterResults.LoudnessRangeLowerBound = OutputMeterEntry.LoudnessRangeLowerBound;
+			LoudnessMeterResults.LoudnessRangeUpperBound = OutputMeterEntry.LoudnessRangeUpperBound;
+			LoudnessMeterResults.TruePeakDb = OutputMeterEntry.TruePeakMaxValueDb;
+
+			LoudnessMeterMaxResults.ShortTermLoudness = FMath::Max<float>(LoudnessMeterMaxResults.ShortTermLoudness, LoudnessMeterResults.ShortTermLoudness);
+			LoudnessMeterMaxResults.MomentaryLoudness = FMath::Max<float>(LoudnessMeterMaxResults.MomentaryLoudness, LoudnessMeterResults.MomentaryLoudness);
+			LoudnessMeterMaxResults.TruePeakDb        = FMath::Max<float>(LoudnessMeterMaxResults.TruePeakDb, LoudnessMeterResults.TruePeakDb);
 		}
 	}
 
diff --git a/Engine/Plugins/AudioInsights/Source/AudioInsights/Private/Views/OutputMeterDashboardViewFactory.h b/Engine/Plugins/AudioInsights/Source/AudioInsights/Private/Views/OutputMeterDashboardViewFactory.h
--- a/Engine/Plugins/AudioInsights/Source/AudioInsights/Private/Views/OutputMeterDashboardViewFactory.h
+++ b/Engine/Plugins/AudioInsights/Source/AudioInsights/Private/Views/OutputMeterDashboardViewFactory.h
@@ -23,6 +23,25 @@ namespace UE::Audio::Insights
 		virtual TSharedRef<SWidget> MakeWidget(TSharedRef<SDockTab> OwnerTab, const FSpawnTabArgs& SpawnTabArgs) override;
 
 	private:
+		struct FLoudnessMeterResults
+		{
+			float LongTermLoudness  = MIN_VOLUME_DECIBELS;
+			float ShortTermLoudness = MIN_VOLUME_DECIBELS;
+			float MomentaryLoudness = MIN_VOLUME_DECIBELS;
+			float LoudnessRangeLowerBound = 0.0f;
+			float LoudnessRangeUpperBound = 0.0f;
+
+			float TruePeakDb = MIN_VOLUME_DECIBELS;
+		};
+
+		struct FLoudnessMeterMaxResults
+		{
+			float ShortTermLoudness = MIN_VOLUME_DECIBELS;
+			float MomentaryLoudness = MIN_VOLUME_DECIBELS;
+
+			float TruePeakDb = MIN_VOLUME_DECIBELS;
+		};
+
 		void InitLoudnessMeterWidgetView();
 
 		void UpdateDataViewEntries();
@@ -46,12 +65,8 @@ namespace UE::Audio::Insights
 		AudioWidgetsCore::FLoudnessMeterWidgetView LoudnessMeterWidgetView;
 		bool bIsLoudnessMeterInitialized = false;
 
-		float TruePeakMaxValueDb = MIN_VOLUME_DECIBELS;
-		float LongTermLoudness  = MIN_VOLUME_DECIBELS;
-		float ShortTermLoudness = MIN_VOLUME_DECIBELS;
-		float MomentaryLoudness = MIN_VOLUME_DECIBELS;
-		float LoudnessRangeLowerBound = 0.0f;
-		float LoudnessRangeUpperBound = 0.0f;
+		FLoudnessMeterResults LoudnessMeterResults;
+		FLoudnessMeterMaxResults LoudnessMeterMaxResults;
 	};
 } // namespace UE::Audio::Insights
 
diff --git a/Engine/Plugins/Runtime/AudioWidgets/Source/AudioWidgets/Private/AudioLoudnessMeter.cpp b/Engine/Plugins/Runtime/AudioWidgets/Source/AudioWidgets/Private/AudioLoudnessMeter.cpp
--- a/Engine/Plugins/Runtime/AudioWidgets/Source/AudioWidgets/Private/AudioLoudnessMeter.cpp
+++ b/Engine/Plugins/Runtime/AudioWidgets/Source/AudioWidgets/Private/AudioLoudnessMeter.cpp
@@ -188,6 +188,9 @@ namespace AudioWidgets
 
 		LatestLoudnessResults.Reset();
 
+		LoudnessMeterMaxResults.ShortTermLoudness = MIN_VOLUME_DECIBELS;
+		LoudnessMeterMaxResults.MomentaryLoudness = MIN_VOLUME_DECIBELS;
+
 		// Recreate loudness analyzer
 		if (AudioDeviceId != FAudioBusInfo::InvalidAudioDeviceId)
 		{
@@ -228,6 +231,8 @@ namespace AudioWidgets
 
 		LatestTruePeakResults.Reset();
 
+		LoudnessMeterMaxResults.TruePeakDb = MIN_VOLUME_DECIBELS;
+
 		// Recreate true peak analyzer
 		if (AudioDeviceId != FAudioBusInfo::InvalidAudioDeviceId)
 		{
@@ -268,16 +273,24 @@ namespace AudioWidgets
 				.DisplayName = LOCTEXT("ShortTermLoudnessDisplayName", "Short Term"),
 				.MeterMetric = EAudioMeterMetric::Loudness,
 				.Value = TAttribute<TOptional<float>>::CreateSPLambda(this, [this]() { return (LatestLoudnessResults.IsSet()) ? TOptional<float>(LatestLoudnessResults->ShortTermLoudness) : NullOpt; }),
+				.MaxValue = TAttribute<TOptional<float>>::CreateSPLambda(this, [this]()	{ return LoudnessMeterMaxResults.ShortTermLoudness;	}),
 				.bShowValue = false,
 				.bShowMeter = true,
+				.bHoldMaxForValue = false,
+				.bHoldMaxForMeter = false,
+				.OnClicked = FSimpleDelegate::CreateSPLambda(this, [this]() { LoudnessMeterMaxResults.ShortTermLoudness = MIN_VOLUME_DECIBELS; }),
 			},
 			{
 				.Name = GET_MEMBER_NAME_CHECKED(FLoudnessMeterRackUnitSettings, MomentaryLoudness),
 				.DisplayName = LOCTEXT("MomentaryLoudnessDisplayName", "Momentary"),
 				.MeterMetric = EAudioMeterMetric::Loudness,
 				.Value = TAttribute<TOptional<float>>::CreateSPLambda(this, [this]() { return (LatestLoudnessResults.IsSet()) ? TOptional<float>(LatestLoudnessResults->Loudness) : NullOpt; }),
+				.MaxValue = TAttribute<TOptional<float>>::CreateSPLambda(this, [this]() { return LoudnessMeterMaxResults.MomentaryLoudness; }),
 				.bShowValue = false,
 				.bShowMeter = true,
+				.bHoldMaxForValue = false,
+				.bHoldMaxForMeter = false,
+				.OnClicked = FSimpleDelegate::CreateSPLambda(this, [this]() { LoudnessMeterMaxResults.MomentaryLoudness = MIN_VOLUME_DECIBELS; }),
 			},
 			{
 				.Name = GET_MEMBER_NAME_CHECKED(FLoudnessMeterRackUnitSettings, LoudnessRange),
@@ -292,9 +305,12 @@ namespace AudioWidgets
 				.DisplayName = LOCTEXT("TruePeakDisplayName", "True Peak"),
 				.MeterMetric = EAudioMeterMetric::Decibels,
 				.Value = TAttribute<TOptional<float>>::CreateSPLambda(this, [this]() { return (LatestTruePeakResults.IsSet()) ? TOptional<float>(LatestTruePeakResults->MeterValue) : NullOpt; }),
+				.MaxValue = TAttribute<TOptional<float>>::CreateSPLambda(this, [this]() { return LoudnessMeterMaxResults.TruePeakDb; }),
 				.bShowValue = true,
 				.bShowMeter = false,
-				.OnClicked = FSimpleDelegate::CreateSPLambda(this, [this]() { ResetTruePeakAnalyzer(); })
+				.bHoldMaxForValue = true,
+				.bHoldMaxForMeter = true,
+				.OnClicked = FSimpleDelegate::CreateSPLambda(this, [this]() { LoudnessMeterMaxResults.TruePeakDb = MIN_VOLUME_DECIBELS; })
 			}
 		};
 
@@ -393,6 +409,18 @@ namespace AudioWidgets
 								return DisplayOptions->bShowMeter;
 							});
 
+						LoudnessMetric.bHoldMaxForValue.BindLambda([=]()
+							{
+								const FLoudnessMetricDisplayOptions* DisplayOptions = GetLoudnessMetricDisplayOptions(SettingsHelper.GetRackUnitSettings());
+								return DisplayOptions->bHoldMaxForValue;
+							});
+
+						LoudnessMetric.bHoldMaxForMeter.BindLambda([=]()
+							{
+								const FLoudnessMetricDisplayOptions* DisplayOptions = GetLoudnessMetricDisplayOptions(SettingsHelper.GetRackUnitSettings());
+								return DisplayOptions->bHoldMaxForMeter;
+							});
+
 						LoudnessMetric.OnShowValueToggleRequested.BindLambda([=]()
 							{
 								FLoudnessMetricDisplayOptions* DisplayOptions = GetLoudnessMetricDisplayOptions(SettingsHelper.GetRackUnitSettings());
@@ -407,6 +435,46 @@ namespace AudioWidgets
 								SettingsHelper.SaveConfig();
 							});
 
+						LoudnessMetric.OnShowValueToggleFromDragDropRequested.BindLambda([=]()
+							{
+								FLoudnessMetricDisplayOptions* DisplayOptions = GetLoudnessMetricDisplayOptions(SettingsHelper.GetRackUnitSettings());
+								DisplayOptions->bShowValue = !DisplayOptions->bShowValue;
+
+								if (DisplayOptions->bShowValue)
+								{
+									DisplayOptions->bHoldMaxForValue = DisplayOptions->bHoldMaxForMeter;
+								}
+
+								SettingsHelper.SaveConfig();
+							});
+
+						LoudnessMetric.OnShowMeterToggleFromDragDropRequested.BindLambda([=]()
+							{
+								FLoudnessMetricDisplayOptions* DisplayOptions = GetLoudnessMetricDisplayOptions(SettingsHelper.GetRackUnitSettings());
+								DisplayOptions->bShowMeter = !DisplayOptions->bShowMeter;
+
+								if (DisplayOptions->bShowMeter)
+								{
+									DisplayOptions->bHoldMaxForMeter = DisplayOptions->bHoldMaxForValue;
+								}
+
+								SettingsHelper.SaveConfig();
+							});
+
+						LoudnessMetric.OnHoldMaxForValueToggleRequested.BindLambda([=, this]()
+							{
+								FLoudnessMetricDisplayOptions* DisplayOptions = GetLoudnessMetricDisplayOptions(SettingsHelper.GetRackUnitSettings());
+								DisplayOptions->bHoldMaxForValue = !DisplayOptions->bHoldMaxForValue;
+								SettingsHelper.SaveConfig();
+							});
+
+						LoudnessMetric.OnHoldMaxForMeterToggleRequested.BindLambda([=, this]()
+							{
+								FLoudnessMetricDisplayOptions* DisplayOptions = GetLoudnessMetricDisplayOptions(SettingsHelper.GetRackUnitSettings());
+								DisplayOptions->bHoldMaxForMeter = !DisplayOptions->bHoldMaxForMeter;
+								SettingsHelper.SaveConfig();
+							});
+
 						LoudnessMetric.Color.BindLambda([=]()
 							{
 								const FLoudnessMetricDisplayOptions* DisplayOptions = GetLoudnessMetricDisplayOptions(SettingsHelper.GetRackUnitSettings());
@@ -454,6 +522,9 @@ namespace AudioWidgets
 		if (InLKFSAnalyzer == LKFSAnalyzer.Get())
 		{
 			LatestLoudnessResults = InLoudnessResults;
+
+			LoudnessMeterMaxResults.ShortTermLoudness = FMath::Max<float>(LoudnessMeterMaxResults.ShortTermLoudness, InLoudnessResults.ShortTermLoudness);
+			LoudnessMeterMaxResults.MomentaryLoudness = FMath::Max<float>(LoudnessMeterMaxResults.MomentaryLoudness, InLoudnessResults.Loudness);
 		}
 	}
 
@@ -461,18 +532,9 @@ namespace AudioWidgets
 	{
 		if (InMeterAnalyzer == TruePeakAnalyzer.Get())
 		{
-			if (!LatestTruePeakResults.IsSet())
-			{
-				LatestTruePeakResults = InMeterResults;
-			}
-			else
-			{
-				// Keep always the max true peak value (unless reset)
-				const float PreviousTruePeakValue = LatestTruePeakResults->MeterValue;
-
-				LatestTruePeakResults = InMeterResults;
-				LatestTruePeakResults->MeterValue = FMath::Max<float>(PreviousTruePeakValue, InMeterResults.MeterValue);
-			}
+			LatestTruePeakResults = InMeterResults;
+			
+			LoudnessMeterMaxResults.TruePeakDb = FMath::Max<float>(LoudnessMeterMaxResults.TruePeakDb, InMeterResults.MeterValue);
 		}
 	}
 
diff --git a/Engine/Plugins/Runtime/AudioWidgets/Source/AudioWidgets/Private/AudioLoudnessMeter.h b/Engine/Plugins/Runtime/AudioWidgets/Source/AudioWidgets/Private/AudioLoudnessMeter.h
--- a/Engine/Plugins/Runtime/AudioWidgets/Source/AudioWidgets/Private/AudioLoudnessMeter.h
+++ b/Engine/Plugins/Runtime/AudioWidgets/Source/AudioWidgets/Private/AudioLoudnessMeter.h
@@ -37,6 +37,13 @@ namespace AudioWidgets
 		// End IAudioAnalyzerRackUnit overrides.
 
 	private:
+		struct FLoudnessMeterMaxResults
+		{
+			float ShortTermLoudness = MIN_VOLUME_DECIBELS;
+			float MomentaryLoudness = MIN_VOLUME_DECIBELS;
+			float TruePeakDb        = MIN_VOLUME_DECIBELS;
+		};
+
 		static TSharedRef<IAudioAnalyzerRackUnit> MakeRackUnit(const FAudioAnalyzerRackUnitConstructParams& Params);
 
 		void Init(int32 InNumChannels, Audio::FDeviceId InAudioDeviceId, TObjectPtr<UAudioBus> InExternalAudioBus);
@@ -85,6 +92,9 @@ namespace AudioWidgets
 		/** Most recent results from the true peak analyzer. */
 		TOptional<FMeterResults> LatestTruePeakResults;
 
+		/** Loudness meter max results (only for short term loudness, momentary loudness and true peak). */
+		FLoudnessMeterMaxResults LoudnessMeterMaxResults;
+
 		Audio::FDeviceId AudioDeviceId = INDEX_NONE;
 
 		FTSTicker::FDelegateHandle WidgetRefreshTicker;
diff --git a/Engine/Plugins/Runtime/AudioWidgets/Source/AudioWidgetsCore/Private/LoudnessMeterWidgetView.cpp b/Engine/Plugins/Runtime/AudioWidgets/Source/AudioWidgetsCore/Private/LoudnessMeterWidgetView.cpp
--- a/Engine/Plugins/Runtime/AudioWidgets/Source/AudioWidgetsCore/Private/LoudnessMeterWidgetView.cpp
+++ b/Engine/Plugins/Runtime/AudioWidgets/Source/AudioWidgetsCore/Private/LoudnessMeterWidgetView.cpp
@@ -9,6 +9,7 @@
 #include "Brushes/SlateRoundedBoxBrush.h"
 #include "Framework/Application/SlateApplication.h"
 #include "Framework/MultiBox/MultiBoxBuilder.h"
+#include "LoudnessMeterSettings.h"
 #include "Permutations.h"
 #include "SAudioMeterWidget.h"
 #include "SDragReorderableTileView.h"
@@ -37,6 +38,13 @@ namespace AudioWidgetsCore
 
 		const FText NegativeInfinityText = LOCTEXT("NegativeInfinity", "-∞");
 
+		bool SupportsMaxValueDisplay(const FName& InLoudnessMeterName)
+		{
+			return InLoudnessMeterName == GET_MEMBER_NAME_CHECKED(FLoudnessMeterSettings, ShortTermLoudness) ||
+				   InLoudnessMeterName == GET_MEMBER_NAME_CHECKED(FLoudnessMeterSettings, MomentaryLoudness) ||
+			       InLoudnessMeterName == GET_MEMBER_NAME_CHECKED(FLoudnessMeterSettings, TruePeak);
+		}
+
 		/**
 		 * Derived meter widget type adds bindable attributes for scale DisplayRange and Offset.
 		 */
@@ -158,10 +166,37 @@ namespace AudioWidgetsCore
 			TAttribute<int32> Target;
 		};
 
+		/** 
+		 * Creates a widget that displays a max value checkbox.
+		 */
+		TSharedRef<SWidget> CreateMaxValueCheckboxWidget(FLoudnessMetricRef InItem, TAttribute<bool> bInHoldMax, const FSimpleDelegate& InOnToggleRequested)
+		{
+			return SNew(SCheckBox)
+				.ToolTipText(LOCTEXT("MaxValueCheckboxWidget_Tooltip", "Enable to hold the max value for this loudness metric"))
+				.Style(&FAppStyle::GetWidgetStyle<FCheckBoxStyle>("RadioButton"))
+				.IsChecked_Lambda([InItem, bInHoldMax]()
+				{
+					return bInHoldMax.Get() ? ECheckBoxState::Checked : ECheckBoxState::Unchecked;
+				})
+				.OnCheckStateChanged_Lambda([InOnToggleRequested](ECheckBoxState NewState)
+				{
+					InOnToggleRequested.ExecuteIfBound();
+				})
+				[
+					SNew(STextBlock)
+					.ColorAndOpacity_Lambda([InItem, bInHoldMax]()
+					{
+						return bInHoldMax.Get() ? FColor::Yellow : FColor::Silver;
+					})
+					.Font(FStyleDefaults::GetFontInfo(7))
+					.Text(LOCTEXT("HoldMax", "Max"))
+				];
+		}
+
 		/**
 		 * Creates a widget that displays the numeric value of the given loudness metric, with units.
 		 */
-		TSharedRef<SWidget> MakeValueReadoutWidget(FLoudnessMetricRef InItem, const FLoudnessScaleParams& LoudnessScaleParams)
+		TSharedRef<SWidget> MakeValueReadoutWidget(FLoudnessMetricRef InItem, const FLoudnessScaleParams& LoudnessScaleParams, TAttribute<bool> bInHoldMax)
 		{
 			const FLoudnessValueConverter LoudnessValueConverter(LoudnessScaleParams);
 
@@ -176,13 +211,13 @@ namespace AudioWidgetsCore
 				+ SOverlay::Slot()
 				[
 					SNew(STextBlock)
-						.Text_Lambda([InItem, LoudnessValueConverter]()
+						.Text_Lambda([InItem, LoudnessValueConverter, bInHoldMax]()
 							{
 								switch (InItem->MeterMetric.Get())
 								{
 									case EAudioMeterMetric::Decibels:
 									{
-										const TOptional<float> ValueOpt = InItem->Value.Get();
+										const TOptional<float>& ValueOpt = bInHoldMax.Get() ? InItem->MaxValue.Get() : InItem->Value.Get();
 										if (!ValueOpt.IsSet())
 										{
 											return FText::GetEmpty();
@@ -199,7 +234,7 @@ namespace AudioWidgetsCore
 									{
 										return InItem->IsRanged()
 											? LoudnessValueConverter.FormatRangeText(InItem->Range.Get())
-											: LoudnessValueConverter.FormatNumericValueText(InItem->Value.Get());
+											: LoudnessValueConverter.FormatNumericValueText(bInHoldMax.Get() ? InItem->MaxValue.Get() : InItem->Value.Get());
 									}
 
 									default:
@@ -255,7 +290,7 @@ namespace AudioWidgetsCore
 				.AutoHeight()
 				.HAlign(EHorizontalAlignment::HAlign_Center)
 				[
-					MakeValueReadoutWidget(InItem, LoudnessScaleParams)
+					MakeValueReadoutWidget(InItem, LoudnessScaleParams, InItem->bHoldMaxForValue)
 				]
 				+ SVerticalBox::Slot()
 				.MinHeight(24)
@@ -268,6 +303,13 @@ namespace AudioWidgetsCore
 						.Font(FStyleDefaults::GetFontInfo(8))
 						.AutoWrapText(true)
 						.Justification(ETextJustify::Center)
+				]
+				+ SVerticalBox::Slot()
+				.AutoHeight()
+				.HAlign(EHorizontalAlignment::HAlign_Right)
+				.VAlign(EVerticalAlignment::VAlign_Center)
+				[
+					SupportsMaxValueDisplay(InItem->Name) ? CreateMaxValueCheckboxWidget(InItem, InItem->bHoldMaxForValue, InItem->OnHoldMaxForValueToggleRequested) : SNullWidget::NullWidget
 				];
 		}
 
@@ -283,7 +325,7 @@ namespace AudioWidgetsCore
 				.AutoHeight()
 				.HAlign(EHorizontalAlignment::HAlign_Center)
 				[
-					MakeValueReadoutWidget(InItem, LoudnessScaleParams)
+					MakeValueReadoutWidget(InItem, LoudnessScaleParams, InItem->bHoldMaxForMeter)
 				]
 				+ SVerticalBox::Slot()
 				.HAlign(EHorizontalAlignment::HAlign_Center)
@@ -316,11 +358,12 @@ namespace AudioWidgetsCore
 
 								if (InItem->MeterMetric.Get() == EAudioMeterMetric::Decibels)
 								{
+									const TOptional<float>& ValueOpt = InItem->bHoldMaxForMeter.Get() ? InItem->MaxValue.Get() : InItem->Value.Get();
 									return
 									{
 										{
-											.MeterValue = InItem->Value.Get().Get(DefaultInvalidValue),
-											.PeakValue = DefaultInvalidValue
+											.MeterValue = ValueOpt.Get(DefaultInvalidValue),
+											.PeakValue  = DefaultInvalidValue
 										}
 									};
 								}
@@ -341,11 +384,13 @@ namespace AudioWidgetsCore
 								}
 
 								// Display standard loudness meter bar:
-								const TOptional<float> ConvertedValue = LoudnessValueConverter.ConvertValue(InItem->Value.Get());
+								const TOptional<float> ValueOpt = InItem->bHoldMaxForMeter.Get() ? InItem->MaxValue.Get() : InItem->Value.Get();
+								const TOptional<float> ConvertedValueOpt = LoudnessValueConverter.ConvertValue(ValueOpt);
+								
 								return
 								{
 									{
-										.MeterValue = ConvertedValue.Get(DefaultInvalidValue),
+										.MeterValue = ConvertedValueOpt.Get(DefaultInvalidValue),
 										.PeakValue  = bRelativeScale.Get() ? 0.0f : static_cast<float>(Target.Get())
 									}
 								};
@@ -368,10 +413,18 @@ namespace AudioWidgetsCore
 				.VAlign(EVerticalAlignment::VAlign_Center)
 				[
 					SNew(STextBlock)
-						.Text_Lambda([InItem]() { return InItem->DisplayName; })
+						.Text(InItem->DisplayName)
 						.Font(FStyleDefaults::GetFontInfo(8))
 						.AutoWrapText(true)
 						.Justification(ETextJustify::Center)
+				]
+				+ SVerticalBox::Slot()
+				.AutoHeight()
+				.MinHeight(20)
+				.HAlign(EHorizontalAlignment::HAlign_Right)
+				.VAlign(EVerticalAlignment::VAlign_Center)
+				[
+					SupportsMaxValueDisplay(InItem->Name) ? CreateMaxValueCheckboxWidget(InItem, InItem->bHoldMaxForMeter, InItem->OnHoldMaxForMeterToggleRequested) : SNullWidget::NullWidget
 				];
 		}
 
@@ -442,7 +495,7 @@ namespace AudioWidgetsCore
 		{
 			if (LoudnessMetric->bShowValue.Get() != bShowValue)
 			{
-				LoudnessMetric->OnShowValueToggleRequested.ExecuteIfBound();
+				LoudnessMetric->OnShowValueToggleFromDragDropRequested.ExecuteIfBound();
 			}
 		}
 
@@ -453,7 +506,7 @@ namespace AudioWidgetsCore
 		{
 			if (LoudnessMetric->bShowMeter.Get() != bShowMeter)
 			{
-				LoudnessMetric->OnShowMeterToggleRequested.ExecuteIfBound();
+				LoudnessMetric->OnShowMeterToggleFromDragDropRequested.ExecuteIfBound();
 			}
 		}
 
@@ -568,7 +621,8 @@ namespace AudioWidgetsCore
 				+ SOverlay::Slot()
 				[
 					SNew(SColorBlock)
-						.Color(TileWidget, &SLoudnessMetricTile::GetOverlayColor)
+					.Visibility(EVisibility::HitTestInvisible)
+					.Color(TileWidget, &SLoudnessMetricTile::GetOverlayColor)
 				]);
 
 			return TileWidget;
@@ -614,7 +668,8 @@ namespace AudioWidgetsCore
 				+ SOverlay::Slot()
 				[
 					SNew(SColorBlock)
-						.Color(TileWidget, &SLoudnessMetricTile::GetOverlayColor)
+					.Visibility(EVisibility::HitTestInvisible)
+					.Color(TileWidget, &SLoudnessMetricTile::GetOverlayColor)
 				]);
 
 			return TileWidget;
@@ -902,7 +957,10 @@ namespace AudioWidgetsCore
 			.ScrollbarVisibility(EVisibility::Collapsed)
 			.OnMouseButtonClick_Lambda([](FLoudnessMetricRef InItem)
 			{
-				InItem->OnClicked.ExecuteIfBound();
+				if (InItem->bHoldMaxForMeter.Get())
+				{
+					InItem->OnClicked.ExecuteIfBound();
+				}
 			});
 
 		MeterTileView->SetItemHeight(TAttribute<float>::CreateLambda([&TileView = *MeterTileView]()
@@ -1006,13 +1064,16 @@ namespace AudioWidgetsCore
 							.Visibility(GetNumericValuesPanelVisibility)
 							.OnGenerateTile_Static(&MakeNumericValueTileWidget)
 							.ListItemsSource(ValuesCollectionView->GetVisibleItems())
-							.ItemHeight(60.0f)
-							.ItemWidth(60.0f)
+							.ItemHeight(82.0f)
+							.ItemWidth(68.0f)
 							.ItemAlignment(EListItemAlignment::CenterAligned)
 							.SelectionMode(ESelectionMode::Single)
 							.OnMouseButtonClick_Lambda([](FLoudnessMetricRef InItem)
 							{
-								InItem->OnClicked.ExecuteIfBound();
+								if (InItem->bHoldMaxForValue.Get())
+								{
+									InItem->OnClicked.ExecuteIfBound();
+								}
 							})
 					]
 					+ SVerticalBox::Slot()
diff --git a/Engine/Plugins/Runtime/AudioWidgets/Source/AudioWidgetsCore/Public/LoudnessMeterSettings.h b/Engine/Plugins/Runtime/AudioWidgets/Source/AudioWidgetsCore/Public/LoudnessMeterSettings.h
--- a/Engine/Plugins/Runtime/AudioWidgets/Source/AudioWidgetsCore/Public/LoudnessMeterSettings.h
+++ b/Engine/Plugins/Runtime/AudioWidgets/Source/AudioWidgetsCore/Public/LoudnessMeterSettings.h
@@ -17,6 +17,12 @@ struct FLoudnessMeterDisplayOptions
 	UPROPERTY(EditAnywhere, config, Category = LoudnessMeter)
 	bool bShowMeter = false;
 
+	UPROPERTY(EditAnywhere, config, Category = LoudnessMeter) 
+	bool bHoldMaxForValue = false;
+
+	UPROPERTY(EditAnywhere, config, Category = LoudnessMeter)
+	bool bHoldMaxForMeter = false;
+
 	UPROPERTY(EditAnywhere, config, Category = LoudnessMeter)
 	FLinearColor Color = FStyleColors::AccentGreen.GetSpecifiedColor().CopyWithNewOpacity(0.5f);
 };
@@ -57,7 +63,7 @@ struct FLoudnessMeterSettings
 	FLoudnessMeterDisplayOptions LoudnessRange = { .bShowValue = true, .bShowMeter = false };
 
 	UPROPERTY(EditAnywhere, config, Category = LoudnessMeter)
-	FLoudnessMeterDisplayOptions TruePeak = { .bShowValue = true, .bShowMeter = false };
+	FLoudnessMeterDisplayOptions TruePeak = { .bShowValue = true, .bShowMeter = false, .bHoldMaxForValue = true, .bHoldMaxForMeter = true };
 
 	UPROPERTY(EditAnywhere, config, Category = LoudnessMeter, meta = (ClampMin = "0", ClampMax = "119", UIMin = "0", UIMax = "119"))
 	uint64 ValuesOrderingPermutation = 0; // Max == 5! - 1
diff --git a/Engine/Plugins/Runtime/AudioWidgets/Source/AudioWidgetsCore/Public/LoudnessMeterWidgetView.h b/Engine/Plugins/Runtime/AudioWidgets/Source/AudioWidgetsCore/Public/LoudnessMeterWidgetView.h
--- a/Engine/Plugins/Runtime/AudioWidgets/Source/AudioWidgetsCore/Public/LoudnessMeterWidgetView.h
+++ b/Engine/Plugins/Runtime/AudioWidgets/Source/AudioWidgetsCore/Public/LoudnessMeterWidgetView.h
@@ -32,11 +32,18 @@ namespace AudioWidgetsCore
 			FText DisplayName;
 			TAttribute<EAudioMeterMetric> MeterMetric;
 			TAttribute<TOptional<float>> Value;
+			TAttribute<TOptional<float>> MaxValue;
 			TAttribute<TOptional<FFloatInterval>> Range;
 			TAttribute<bool> bShowValue = false;
 			TAttribute<bool> bShowMeter = false;
+			TAttribute<bool> bHoldMaxForValue = false;
+			TAttribute<bool> bHoldMaxForMeter = false;
 			FSimpleDelegate OnShowValueToggleRequested;
 			FSimpleDelegate OnShowMeterToggleRequested;
+			FSimpleDelegate OnShowValueToggleFromDragDropRequested;
+			FSimpleDelegate OnShowMeterToggleFromDragDropRequested;
+			FSimpleDelegate OnHoldMaxForValueToggleRequested;
+			FSimpleDelegate OnHoldMaxForMeterToggleRequested;
 			TAttribute<FLinearColor> Color = GetDefaultMeterColor();
 			FOnLinearColorValueChanged OnColorChanged;
 			FOnWindowClosed OnColorPickerWindowClosed;
