[Audio Insights] Part 01: reduce jitteriness in signal flow graph when the graph size is larger than the viewport

MaintainFocusedNodeXPos
* Pick a node (either the selected node or the audio device/root graph node if none is selected) to maintain at X position 0
* After the graph has finished picking positions for all nodes, shift all node X positions so the chosen node is back at position 0

SSignalFlowScrollBox
* Added small helper class to disable scroll bar animation. bAnimateScroll is a protected member that is set when calling ScrollDescendantIntoView. However, it remains on after this, triggering lerping scrolls set directly from the signal flow graph.
* Resetting bAnimateScroll before setting the scroll offet directly inside SSignalFlowGraph::MaintainOldViewportPosition stops this lerping behaviour

MaintainOldViewportPosition
* Only apply viewport repositioning on centered graphs - Edge justified graphs do not require repositioning to maintain the viewport

CalculateGraphSizeIncreaseAtRight and CalculateGraphSizeIncreaseAtBottom
* Only take into account the difference at the left-most side of the graph (min positions)
* Account for difference in node slot size when viewed in horizontral vs vertical mode

#Synced-CL 51538744

diff --git a/Engine/Plugins/AudioInsights/Source/AudioInsights/Private/Views/SSignalFlowGraph.cpp b/Engine/Plugins/AudioInsights/Source/AudioInsights/Private/Views/SSignalFlowGraph.cpp
--- a/Engine/Plugins/AudioInsights/Source/AudioInsights/Private/Views/SSignalFlowGraph.cpp
+++ b/Engine/Plugins/AudioInsights/Source/AudioInsights/Private/Views/SSignalFlowGraph.cpp
@@ -14,7 +14,6 @@
 #include "Widgets/Input/SCheckBox.h"
 #include "Widgets/Layout/SBorder.h"
 #include "Widgets/Layout/SScaleBox.h"
-#include "Widgets/Layout/SScrollBox.h"
 #include "Widgets/SBoxPanel.h"
 #include "Widgets/Text/STextBlock.h"
 
@@ -48,6 +47,8 @@ namespace UE::Audio::Insights
 		constexpr float NodeBorderPadding = 12.0f;
 		constexpr float NodeDetailVerticalPadding = 8.0f;
 		constexpr float NodeDetailIndentationPadding = 20.0f;
+		constexpr float LargeBetweenNodesPadding = 48.0f;
+		constexpr float SmallBetweenNodePadding = 8.0f;
 
 		constexpr float SendLabelVerticalWidth = 20.0f;
 		constexpr float SendLabelPadding = 8.0f;
@@ -65,12 +66,6 @@ namespace UE::Audio::Insights
 			return RoundedBoxBrush;
 		}
 
-		float GetNodeGraphSlotWidth(const EOrientation GraphOrientation)
-		{
-			const float XPaddingBetweenNodes = GraphOrientation == EOrientation::Orient_Horizontal ? 48.0f : 8.0f;
-			return NodeMaxWidth + XPaddingBetweenNodes;
-		}
-
 		float GetNodeMaxHeight(const int32 NumParamsVisible, const bool bShowNodeDetails)
 		{
 			constexpr float NodeMaxHeight = 50.0f;
@@ -79,12 +74,6 @@ namespace UE::Audio::Insights
 			return bShowNodeDetails ? NodeMaxHeight + (NodeDetailsRowHeight * NumParamsVisible) : NodeMaxHeight;
 		}
 
-		float GetNodeGraphSlotHeight(const EOrientation GraphOrientation, const int32 NumParamsVisible, const bool bShowNodeDetails)
-		{
-			const float YPaddingBetweenNodes = GraphOrientation == EOrientation::Orient_Horizontal ? 8.0f : 48.0f;
-			return GetNodeMaxHeight(NumParamsVisible, bShowNodeDetails) + YPaddingBetweenNodes;
-		}
-
 		FSlateColor GetNodeAccentColor(const ESignalFlowEntryType NodeType)
 		{
 			switch (NodeType)
@@ -1130,6 +1119,7 @@ namespace UE::Audio::Insights
 		ActiveAudioDeviceEntry = InArgs._ActiveAudioDeviceEntry;
 		HorizontalScrollBox = InArgs._HorizontalScrollBox;
 		VerticalScrollBox = InArgs._VerticalScrollBox;
+		GraphJustification = InArgs._GraphJustification;
 		ShowNodeDetails = InArgs._ShowNodeDetails;
 		DisplayAmpPeakInDb = InArgs._DisplayAmpPeakInDb;
 		AnimateWires = InArgs._AnimateWires;
@@ -1188,7 +1178,11 @@ namespace UE::Audio::Insights
 		CreateGraphNodes(RemainingNodeKeys, NewGraphWidthLimits, NewTreeDepthOrder);
 		CreateSendLabels();
 		CleanUpOldWidgets(RemainingNodeKeys);
-		MaintainOldViewportPosition(OriginalSize, NewGraphWidthLimits, NewTreeDepthOrder);
+
+		if (GraphJustification.Get() == ESignalFlowJustification::Center)
+		{
+			MaintainOldViewportPosition(OriginalSize, NewGraphWidthLimits, NewTreeDepthOrder);
+		}
 
 		CachedGraphWidthLimits = NewGraphWidthLimits;
 	}
@@ -1640,16 +1634,8 @@ namespace UE::Audio::Insights
 		float GraphWidth = 0.0f;
 		float GraphHeight = 0.0f;
 
-		const TSet<ESignalFlowNodeDetailParam> IgnoredParams{ ESignalFlowNodeDetailParam::SendOutputVolume };
-		const int32 NumParamsVisible = NodeDetailFilterSettings ? NodeDetailFilterSettings->GetNumVisibleParams(IgnoredParams) : 0;
-
-		const bool bSendLabelsAreVisible = ShowNodeDetails.Get() && NodeDetailFilterSettings && NodeDetailFilterSettings->GetParameterIsVisible(ESignalFlowNodeDetailParam::SendOutputVolume);
-		const int32 MaxNumVisibleSendLabels = bSendLabelsAreVisible ? MaxNumNodeOutputsInGraph : 0;
-
-		const int32 MaxNumParamRows = FMath::Max(NumParamsVisible, MaxNumVisibleSendLabels);
-
-		const float NodeGraphSlotWidth = GetNodeGraphSlotWidth(GetGraphOrientation());
-		const float NodeGraphSlotHeight = GetNodeGraphSlotHeight(GetGraphOrientation(), MaxNumParamRows, ShowNodeDetails.Get());
+		const float NodeGraphSlotWidth = GetNodeGraphSlotWidth();
+		const float NodeGraphSlotHeight = GetNodeGraphSlotHeight();
 
 		for (const auto& [Key, NodeWidget] : NodeWidgets)
 		{
@@ -1697,6 +1683,31 @@ namespace UE::Audio::Insights
 		return GraphWidthLimits;
 	}
 
+	float SSignalFlowGraph::GetNodeGraphSlotWidth() const
+	{
+		using namespace SSignalFlowGraphStyle;
+
+		const float XPaddingBetweenNodes = Orientation == EOrientation::Orient_Horizontal ? LargeBetweenNodesPadding : SmallBetweenNodePadding;
+		return NodeMaxWidth + XPaddingBetweenNodes;
+	}
+
+	float SSignalFlowGraph::GetNodeGraphSlotHeight() const
+	{
+		using namespace SSignalFlowGraphStyle;
+
+		// The max height of a node in the graph is driven by either the number of visible node details, or the number of send label outputs
+		const TSet<ESignalFlowNodeDetailParam> IgnoredParams{ ESignalFlowNodeDetailParam::SendOutputVolume };
+		const int32 NumParamsVisible = NodeDetailFilterSettings ? NodeDetailFilterSettings->GetNumVisibleParams(IgnoredParams) : 0;
+
+		const bool bSendLabelsAreVisible = ShowNodeDetails.Get() && NodeDetailFilterSettings && NodeDetailFilterSettings->GetParameterIsVisible(ESignalFlowNodeDetailParam::SendOutputVolume);
+		const int32 MaxNumVisibleSendLabels = bSendLabelsAreVisible ? MaxNumNodeOutputsInGraph : 0;
+
+		const int32 MaxNumParamRows = FMath::Max(NumParamsVisible, MaxNumVisibleSendLabels);
+
+		const float YPaddingBetweenNodes = Orientation == EOrientation::Orient_Horizontal ? SmallBetweenNodePadding : LargeBetweenNodesPadding;
+		return GetNodeMaxHeight(MaxNumParamRows, ShowNodeDetails.Get()) + YPaddingBetweenNodes;
+	}
+
 	void SSignalFlowGraph::CreateGraphNodes(TArray<FSignalFlowEntryKey>& OutRemainingNodeKeys, const FGraphLimits& GraphWidthLimits, TArray<TreeDepthPair>& OutTreeDepthStructure)
 	{
 		TreeDepthPair CurrentTreeDepth{ ESignalFlowEntryType::InvalidLow, TNumericLimits<int32>::Lowest() };
@@ -1805,8 +1816,8 @@ namespace UE::Audio::Insights
 		// New Position = Old Postion + Total Size Increase - Size Increase Beyond Old Postion
 		// (i.e. Old Postion + Size Increase Before Old Postion)
 
-		const TSharedPtr<SScrollBox> HorizontalScrollBoxPinned = HorizontalScrollBox.Get().Pin();
-		const TSharedPtr<SScrollBox> VerticalScrollBoxPinned = VerticalScrollBox.Get().Pin();
+		const TSharedPtr<SSignalFlowScrollBox> HorizontalScrollBoxPinned = HorizontalScrollBox.Get().Pin();
+		const TSharedPtr<SSignalFlowScrollBox> VerticalScrollBoxPinned = VerticalScrollBox.Get().Pin();
 
 		if (HorizontalScrollBoxPinned.IsValid() && VerticalScrollBoxPinned.IsValid() && CachedGraphWidthLimits.IsValid())
 		{
@@ -1819,7 +1830,7 @@ namespace UE::Audio::Insights
 			// Get the old position (top left corner) of the viewport
 			const FVector2D OldViewportPosInGraph(HorizontalScrollBoxPinned->GetScrollOffset(), VerticalScrollBoxPinned->GetScrollOffset());
 			
-			const float IncreaseAtRight = CalculateGraphSizeIncreaseAtRight(NewGraphWidthLimits);
+			const float IncreaseAtLeft = CalculateGraphSizeIncreaseAtLeft(NewGraphWidthLimits);
 
 			const float IncreaseAtBottom = CalculateGraphSizeIncreaseAtBottom
 			(
@@ -1828,13 +1839,16 @@ namespace UE::Audio::Insights
 				Orientation == EOrientation::Orient_Vertical ? OldViewportPosInGraph.Y : OldViewportPosInGraph.X
 			);
 
-			const float XOffset = Orientation == EOrientation::Orient_Vertical ? TotalSizeDiff.X - IncreaseAtRight : TotalSizeDiff.X - IncreaseAtBottom;
-			const float YOffset = Orientation == EOrientation::Orient_Vertical ? TotalSizeDiff.Y - IncreaseAtBottom : TotalSizeDiff.Y - IncreaseAtRight;
+			const float XOffset = Orientation == EOrientation::Orient_Vertical ? IncreaseAtLeft : TotalSizeDiff.X - IncreaseAtBottom;
+			const float YOffset = Orientation == EOrientation::Orient_Vertical ? TotalSizeDiff.Y - IncreaseAtBottom : IncreaseAtLeft;
 
 			const FVector2D NewDesiredPos(OldViewportPosInGraph.X + XOffset, OldViewportPosInGraph.Y + YOffset);
 
-			const FVector2D CurrentScrollOffset(HorizontalScrollBoxPinned->GetScrollOffsetOfEnd(), VerticalScrollBoxPinned->GetScrollOffsetOfEnd());
-			const FVector2D MaxScrollOffset = NewDesiredSize - (OriginalSize - CurrentScrollOffset);
+			const FVector2D CurrentScrollOffsetOfEnd(HorizontalScrollBoxPinned->GetScrollOffsetOfEnd(), VerticalScrollBoxPinned->GetScrollOffsetOfEnd());
+			const FVector2D MaxScrollOffset = NewDesiredSize - (OriginalSize - CurrentScrollOffsetOfEnd);
+
+			HorizontalScrollBoxPinned->ResetAnimatingScroll();
+			VerticalScrollBoxPinned->ResetAnimatingScroll();
 
 			HorizontalScrollBoxPinned->SetScrollOffset(FMath::Clamp(NewDesiredPos.X, 0.0f, FMath::Max(MaxScrollOffset.X, 0.0f)));
 			VerticalScrollBoxPinned->SetScrollOffset(FMath::Clamp(NewDesiredPos.Y, 0.0f, FMath::Max(MaxScrollOffset.Y, 0.0f)));
@@ -1843,17 +1857,28 @@ namespace UE::Audio::Insights
 		CachedTreeDepthPairs = NewTreeDepthOrder;
 	}
 
-	float SSignalFlowGraph::CalculateGraphSizeIncreaseAtRight(const FGraphLimits& NewGraphWidthLimits) const
+	float SSignalFlowGraph::CalculateGraphSizeIncreaseAtLeft(const FGraphLimits& NewGraphWidthLimits) const
 	{
 		using namespace SSignalFlowGraphStyle;
 
-		const float NodeGraphSlotWidth = GetNodeGraphSlotWidth(GetGraphOrientation());
+		if (!CachedGraphWidthLimits.IsValid() || !NewGraphWidthLimits.IsValid())
+		{
+			return 0.0f;
+		}
 
-		// Work out in pixels how much the graph has increased in size beyond the original viewport position
-		const float OldMaxXPosScaled = CachedGraphWidthLimits.IsValid() ? NodeGraphSlotWidth * CachedGraphWidthLimits.MaxPos : 0.0f;
-		const float NewMaxXPosScaled = NewGraphWidthLimits.IsValid() ? NodeGraphSlotWidth * NewGraphWidthLimits.MaxPos : 0.0f;
+		const float MinXPosDiff = NewGraphWidthLimits.MinPos - CachedGraphWidthLimits.MinPos;
+		const float TotalXShift = MinXPosDiff * -1.0f;
 
-		return NewMaxXPosScaled - OldMaxXPosScaled;
+		if (GetGraphOrientation() == EOrientation::Orient_Horizontal)
+		{
+			const float NodeGraphSlotHeight = GetNodeGraphSlotHeight();
+			return TotalXShift * NodeGraphSlotHeight * ZoomScaleFactor;
+		}
+		else
+		{
+			const float NodeGraphSlotWidth = GetNodeGraphSlotWidth();
+			return TotalXShift * NodeGraphSlotWidth * ZoomScaleFactor;
+		}
 	}
 
 	float SSignalFlowGraph::CalculateGraphSizeIncreaseAtBottom(const TArray<TreeDepthPair>& NewTreeDepthOrder, const float OriginalGraphHeight, const float OriginalViewportPosY) const
@@ -1885,18 +1910,16 @@ namespace UE::Audio::Insights
 		// Calculate how much the graph has increased below the original viewport (positive increase)
 		const int32 YPosDiffAtBottom = FMath::Max(NewYPos - OriginalYPos, 0);
 
-		// The max height of a node in the graph is driven by either the number of visible node details, or the number of send label outputs
-		const TSet<ESignalFlowNodeDetailParam> IgnoredParams{ ESignalFlowNodeDetailParam::SendOutputVolume };
-		const int32 NumParamsVisible = NodeDetailFilterSettings ? NodeDetailFilterSettings->GetNumVisibleParams(IgnoredParams) : 0;
-
-		const bool bSendLabelsAreVisible = ShowNodeDetails.Get() && NodeDetailFilterSettings && NodeDetailFilterSettings->GetParameterIsVisible(ESignalFlowNodeDetailParam::SendOutputVolume);
-		const int32 MaxNumVisibleSendLabels = bSendLabelsAreVisible ? MaxNumNodeOutputsInGraph : 0;
-
-		const int32 MaxNumParamRows = FMath::Max(NumParamsVisible, MaxNumVisibleSendLabels);
-
-		const float NodeGraphSlotHeight = GetNodeGraphSlotHeight(GetGraphOrientation(), MaxNumParamRows, ShowNodeDetails.Get());
-
-		return YPosDiffAtBottom * NodeGraphSlotHeight;
+		if (GetGraphOrientation() == EOrientation::Orient_Horizontal)
+		{
+			const float NodeGraphSlotWidth = GetNodeGraphSlotWidth();
+			return YPosDiffAtBottom * NodeGraphSlotWidth * ZoomScaleFactor;
+		}
+		else
+		{
+			const float NodeGraphSlotHeight = GetNodeGraphSlotHeight();
+			return YPosDiffAtBottom * NodeGraphSlotHeight * ZoomScaleFactor;
+		}
 	}
 
 	void SSignalFlowGraph::CreateGraphNode(const TSharedPtr<ISignalFlowNode>& Node, const float MinXPos, const float MaxXPos, const int32 DepthTracker, TArray<FSignalFlowEntryKey>& OutRemainingNodeKeys)
@@ -1965,17 +1988,9 @@ namespace UE::Audio::Insights
 				const TSharedPtr<SSignalFlowGraphNode>* Node = NodeWidgets.Find(EntryKey);
 				if (Node && Node->IsValid())
 				{
-					const TSet<ESignalFlowNodeDetailParam> IgnoredParams{ ESignalFlowNodeDetailParam::SendOutputVolume };
-					const int32 NumParamsVisible = NodeDetailFilterSettings ? NodeDetailFilterSettings->GetNumVisibleParams(IgnoredParams) : 0;
-
-					const bool bSendLabelsAreVisible = ShowNodeDetails.Get() && NodeDetailFilterSettings && NodeDetailFilterSettings->GetParameterIsVisible(ESignalFlowNodeDetailParam::SendOutputVolume);
-					const int32 MaxNumVisibleSendLabels = bSendLabelsAreVisible ? MaxNumNodeOutputsInGraph : 0;
-
-					const int32 MaxNumParamRows = FMath::Max(NumParamsVisible, MaxNumVisibleSendLabels);
-
 					const FVector2D NodeSize = (*Node)->GetDesiredSize();
-					const float NodeGraphSlotWidth = GetNodeGraphSlotWidth(GetGraphOrientation());
-					const float NodeGraphSlotHeight = GetNodeGraphSlotHeight(GetGraphOrientation(), MaxNumParamRows, ShowNodeDetails.Get());
+					const float NodeGraphSlotWidth = GetNodeGraphSlotWidth();
+					const float NodeGraphSlotHeight = GetNodeGraphSlotHeight();
 
 					const float XOffset = (NodeGraphSlotWidth - NodeSize.X) / 2.0f;
 					const float YOffset = (NodeGraphSlotHeight - NodeSize.Y) / 2.0f;
diff --git a/Engine/Plugins/AudioInsights/Source/AudioInsights/Private/Views/SignalFlowDashboardViewFactory.cpp b/Engine/Plugins/AudioInsights/Source/AudioInsights/Private/Views/SignalFlowDashboardViewFactory.cpp
--- a/Engine/Plugins/AudioInsights/Source/AudioInsights/Private/Views/SignalFlowDashboardViewFactory.cpp
+++ b/Engine/Plugins/AudioInsights/Source/AudioInsights/Private/Views/SignalFlowDashboardViewFactory.cpp
@@ -931,7 +931,7 @@ namespace UE::Audio::Insights
 
 		+ SOverlay::Slot()
 		[
-			SAssignNew(GraphVerticalScrollBox, SScrollBox)
+			SAssignNew(GraphVerticalScrollBox, SSignalFlowScrollBox)
 			.Orientation(EOrientation::Orient_Vertical)
 			.ScrollBarAlwaysVisible(true)
 
@@ -940,7 +940,7 @@ namespace UE::Audio::Insights
 			.HAlign(HAlign_Fill)
 			.VAlign(VAlign_Fill)
 			[
-				SAssignNew(GraphHorizontalScrollBox, SScrollBox)
+				SAssignNew(GraphHorizontalScrollBox, SSignalFlowScrollBox)
 				.Orientation(EOrientation::Orient_Horizontal)
 				.ScrollBarAlwaysVisible(true)
 
@@ -1022,6 +1022,10 @@ namespace UE::Audio::Insights
 
 							OnSelectionChanged(SelectedNodeEntry, ESelectInfo::OnMouseClick);
 						})
+						.GraphJustification_Lambda([this]()
+						{
+							return GraphJustification;
+						})
 						.ShowNodeDetails_Lambda([this]()
 						{
 							return bShowNodeDetails;
@@ -1397,6 +1401,7 @@ namespace UE::Audio::Insights
 
 		if (SelectedItem.IsValid())
 		{
+			XPosFocusNodeEntryKey = SelectedItem->GetSignalFlowEntryKey();
 			CreateFilteredNodeGraph(*Provider, SelectedItem);
 		}
 		else
@@ -1470,7 +1475,7 @@ namespace UE::Audio::Insights
 
 			// Assign X positions based off the new order
 			// X positions given in units of 1.0 that will be scaled up by SSignalFlowGraph later
-			CalcXPositions();
+			CalcXPositions(RootNodeSharedPtr);
 		}
 	}
 
@@ -1485,6 +1490,7 @@ namespace UE::Audio::Insights
 		const TSharedPtr<FSignalFlowDashboardEntry> DeviceEntry = GetFilteredAudioDeviceEntry();
 		if (DeviceEntry.IsValid())
 		{
+			XPosFocusNodeEntryKey = DeviceEntry->GetSignalFlowEntryKey();
 			CreateFilteredNodeGraph(*Provider, DeviceEntry);
 		}
 	}
@@ -1969,7 +1975,7 @@ namespace UE::Audio::Insights
 		FilteredGraphNodes.Sort(SortSignalNodes);
 	}
 
-	void FSignalFlowDashboardViewFactory::CalcXPositions()
+	void FSignalFlowDashboardViewFactory::CalcXPositions(const TSharedPtr<ISignalFlowNode>& GraphRootNode)
 	{
 		if (!GraphWidgetContents.IsValid())
 		{
@@ -2019,6 +2025,10 @@ namespace UE::Audio::Insights
 
 		// Do one final positioning process for the final row
 		CalcXPositionsForRows(ProcessedNodeRows, GraphOrientation, GraphWidth);
+
+		// Now shift the whole graph so the root focused node remains at X position 0
+		MaintainFocusedNodeXPos(GraphRootNode);
+		
 	}
 
 	void FSignalFlowDashboardViewFactory::CalcXPositionsForRows(const TArray<TArray<int32>>& NodeIndexRows, const EOrientation GraphOrientation, int32& OutGraphWidth)
@@ -2074,6 +2084,44 @@ namespace UE::Audio::Insights
 		}
 	}
 
+	void FSignalFlowDashboardViewFactory::MaintainFocusedNodeXPos(const TSharedPtr<ISignalFlowNode>& FallbackNode)
+	{
+		TOptional<float> XShift;
+
+		if (XPosFocusNodeEntryKey.IsSet())
+		{
+			const FSignalFlowEntryKey EntryKey = XPosFocusNodeEntryKey.GetValue();
+			TSharedPtr<ISignalFlowNode>* FilterFromNode = FilteredEntries.Find(EntryKey);
+
+			if (FilterFromNode && FilterFromNode->IsValid())
+			{
+				XShift = (*FilterFromNode)->XPos;
+			}
+		}
+
+		if (!XShift.IsSet() && FallbackNode.IsValid())
+		{
+			XShift = FallbackNode->XPos;
+		}
+
+		if (!XShift.IsSet())
+		{
+			return;
+		}
+
+		const float XShiftValue = XShift.GetValue();
+		if (!FMath::IsNearlyZero(XShiftValue))
+		{
+			for (const TSharedPtr<ISignalFlowNode>& FilteredNode : FilteredGraphNodes)
+			{
+				if (FilteredNode.IsValid())
+				{
+					FilteredNode->XPos -= XShiftValue;
+				}
+			}
+		}
+	}
+
 	void FSignalFlowDashboardViewFactory::PositionRowEquidistant(const TArray<int32>& NodeIndexRow)
 	{
 		// Space a row of nodes with a X distance of 1.0
diff --git a/Engine/Plugins/AudioInsights/Source/AudioInsights/Public/Views/SSignalFlowGraph.h b/Engine/Plugins/AudioInsights/Source/AudioInsights/Public/Views/SSignalFlowGraph.h
--- a/Engine/Plugins/AudioInsights/Source/AudioInsights/Public/Views/SSignalFlowGraph.h
+++ b/Engine/Plugins/AudioInsights/Source/AudioInsights/Public/Views/SSignalFlowGraph.h
@@ -3,6 +3,8 @@
 
 #include "Delegates/Delegate.h"
 #include "Messages/SignalFlowEntryKey.h"
+#include "Settings/SignalFlowSettings.h"
+#include "Widgets/Layout/SScrollBox.h"
 #include "Widgets/SCanvas.h"
 #include "Widgets/SCompoundWidget.h"
 
@@ -42,6 +44,18 @@ namespace UE::Audio::Insights
 		UE_API FSlateColor GetNodeAccentColor(const ESignalFlowEntryType NodeType);
 	}
 
+	class SSignalFlowScrollBox final : public SScrollBox
+	{
+	public:
+		void ResetAnimatingScroll() 
+		{ 
+			if (!bIsScrolling)
+			{
+				bAnimateScroll = false;
+			}
+		}
+	};
+
 	class SSignalFlowGraphNode : public SCompoundWidget
 	{
 	public:
@@ -204,6 +218,7 @@ namespace UE::Audio::Insights
 		SLATE_BEGIN_ARGS(SSignalFlowGraph)
 			: _ListSource(nullptr)
 			, _NodeDetailFilterSettings(nullptr)
+			, _GraphJustification(ESignalFlowJustification::Edge)
 			, _ShowNodeDetails(false)
 			, _DisplayAmpPeakInDb(true)
 			, _AnimateWires(true)
@@ -215,8 +230,9 @@ namespace UE::Audio::Insights
 
 		SLATE_ATTRIBUTE(TSharedPtr<FSignalFlowDashboardEntry>, SelectedItem)
 		SLATE_ATTRIBUTE(TSharedPtr<FSignalFlowDashboardEntry>, ActiveAudioDeviceEntry)
-		SLATE_ATTRIBUTE(TWeakPtr<SScrollBox>, HorizontalScrollBox)
-		SLATE_ATTRIBUTE(TWeakPtr<SScrollBox>, VerticalScrollBox)
+		SLATE_ATTRIBUTE(TWeakPtr<SSignalFlowScrollBox>, HorizontalScrollBox)
+		SLATE_ATTRIBUTE(TWeakPtr<SSignalFlowScrollBox>, VerticalScrollBox)
+		SLATE_ATTRIBUTE(ESignalFlowJustification, GraphJustification)
 		SLATE_ATTRIBUTE(bool, ShowNodeDetails)
 		SLATE_ATTRIBUTE(bool, DisplayAmpPeakInDb)
 		SLATE_ATTRIBUTE(bool, AnimateWires)
@@ -226,8 +242,8 @@ namespace UE::Audio::Insights
 		SLATE_END_ARGS()
 
 		UE_API void Construct(const FArguments& InArgs);
-		void SetHorizontalScrollBox(const TWeakPtr<SScrollBox> ScrollBox) { HorizontalScrollBox = ScrollBox; }
-		void SetVerticalScrollBox(const TWeakPtr<SScrollBox> ScrollBox) { VerticalScrollBox = ScrollBox; }
+		void SetHorizontalScrollBox(const TWeakPtr<SSignalFlowScrollBox> ScrollBox) { HorizontalScrollBox = ScrollBox; }
+		void SetVerticalScrollBox(const TWeakPtr<SSignalFlowScrollBox> ScrollBox) { VerticalScrollBox = ScrollBox; }
 
 		UE_API void RefreshGraph();
 		UE_API void ResetGraph();
@@ -286,12 +302,15 @@ namespace UE::Audio::Insights
 		};
 
 		FGraphLimits GetGraphWidthLimits() const;
+		float GetNodeGraphSlotWidth() const;
+		float GetNodeGraphSlotHeight() const;
+
 		void CreateGraphNodes(TArray<FSignalFlowEntryKey>& OutRemainingNodeKeys, const FGraphLimits& GraphWidthLimits, TArray<struct TreeDepthPair>& OutTreeDepthStructure);
 		void CreateSendLabels();
 		void CleanUpOldWidgets(const TArray<FSignalFlowEntryKey>& NodeKeysToRemove);
 		
 		void MaintainOldViewportPosition(const FVector2D& OriginalSize, const FGraphLimits& NewGraphWidthLimits, const TArray<struct TreeDepthPair>& NewTreeDepthOrder);
-		float CalculateGraphSizeIncreaseAtRight(const FGraphLimits& NewGraphWidthLimits) const;
+		float CalculateGraphSizeIncreaseAtLeft(const FGraphLimits& NewGraphWidthLimits) const;
 		float CalculateGraphSizeIncreaseAtBottom(const TArray<TreeDepthPair>& NewTreeDepthOrder, const float OriginalGraphHeight, const float OriginalViewportPosY) const;
 
 		void CreateGraphNode(const TSharedPtr<ISignalFlowNode>& Node, const float MinXPos, const float MaxXPos, const int32 DepthTracker, TArray<FSignalFlowEntryKey>& OutRemainingNodeKeys);
@@ -334,12 +353,13 @@ namespace UE::Audio::Insights
 
 		TAttribute<TSharedPtr<FSignalFlowDashboardEntry>> SelectedItem;
 		TAttribute<TSharedPtr<FSignalFlowDashboardEntry>> ActiveAudioDeviceEntry;
+		TAttribute<ESignalFlowJustification> GraphJustification;
 		TAttribute<bool> ShowNodeDetails;
 		TAttribute<bool> DisplayAmpPeakInDb;
 		TAttribute<bool> AnimateWires;
 
-		TAttribute<TWeakPtr<SScrollBox>> HorizontalScrollBox;
-		TAttribute<TWeakPtr<SScrollBox>> VerticalScrollBox;
+		TAttribute<TWeakPtr<SSignalFlowScrollBox>> HorizontalScrollBox;
+		TAttribute<TWeakPtr<SSignalFlowScrollBox>> VerticalScrollBox;
 
 		FOnSignalFlowNodeSelected OnNodeSelected;
 
diff --git a/Engine/Plugins/AudioInsights/Source/AudioInsights/Public/Views/SignalFlowDashboardViewFactory.h b/Engine/Plugins/AudioInsights/Source/AudioInsights/Public/Views/SignalFlowDashboardViewFactory.h
--- a/Engine/Plugins/AudioInsights/Source/AudioInsights/Public/Views/SignalFlowDashboardViewFactory.h
+++ b/Engine/Plugins/AudioInsights/Source/AudioInsights/Public/Views/SignalFlowDashboardViewFactory.h
@@ -134,8 +134,9 @@ namespace UE::Audio::Insights
 		void FixUnassignedNodeOrders();
 		void SortNodeOrder();
 
-		void CalcXPositions();
+		void CalcXPositions(const TSharedPtr<ISignalFlowNode>& GraphRootNode);
 		void CalcXPositionsForRows(const TArray<TArray<int32>>& NodeIndexRows, const EOrientation GraphOrientation, int32& OutGraphWidth);
+		void MaintainFocusedNodeXPos(const TSharedPtr<ISignalFlowNode>& FallbackNode);
 		void PositionRowEquidistant(const TArray<int32>& NodeIndexRow);
 		
 		void CalcXPositionForNode(const TSharedPtr<ISignalFlowNode>& Node, const TArray<FSignalFlowEntryKey>& NodesToPositionRelativeTo, const int32 NodeIndexInRow, const TArray<TSharedPtr<ISignalFlowNode>>& PrevPositionedNodesInRow);
@@ -169,8 +170,8 @@ namespace UE::Audio::Insights
 		TArray<TSharedPtr<ISignalFlowNode>> FilteredGraphNodes;
 		TSharedPtr<SSignalFlowGraph> GraphWidgetContents;
 
-		TSharedPtr<SScrollBox> GraphHorizontalScrollBox;
-		TSharedPtr<SScrollBox> GraphVerticalScrollBox;
+		TSharedPtr<SSignalFlowScrollBox> GraphHorizontalScrollBox;
+		TSharedPtr<SSignalFlowScrollBox> GraphVerticalScrollBox;
 
 		TSharedPtr<FUICommandList> CommandList;
 
@@ -181,6 +182,7 @@ namespace UE::Audio::Insights
 		FSignalFlowEntryTypeListExpansionSettings ActiveEntryMenuExpansionSettings;
 
 		ESignalFlowJustification GraphJustification = ESignalFlowJustification::Edge;
+		TOptional<FSignalFlowEntryKey> XPosFocusNodeEntryKey;
 
 		bool bGraphRequiresRefresh = false;
 		bool bAutoFocusAfterGraphRefresh = false;
