[Audio Insights] Fix Signal Flow AddAllNodesToGraph missing out nodes with no output

AddAllNodesToGraph
* This funtion follows the same process as the usual CreateFilteredNodeGraph, but it uses the Audio Device as the "FilterByEntry", and recursively moves up it's inputs to add the connected nodes into the graph
* When a node has no direct output path connecting to the Audio Device (such as some Audio Buses), they are missed by this creation process.
* This change adds an additional step to try and connect these node paths back into the graph inside AddMissingNodesToGraph

AddMissingNodesToGraph
* Identifies entries that are missing from the graph, and either tries to connect them via their inputs or outputs, depending on if they have any inputs/outputs
* Nodes that are not connected to anything are ignored (such as unconnected buses)

AddMissingNodeByOutputs/Inputs
* These gather a list of the nearest connecting nodes currently in the graph (GetNextOutputs/InputsInGraphRecursive), and then try to create a connecting path to the missing node by walking along the input/output paths

#Synced-CL 51626790

diff --git a/Engine/Plugins/AudioInsights/Source/AudioInsights/Private/Views/SignalFlowDashboardViewFactory.cpp b/Engine/Plugins/AudioInsights/Source/AudioInsights/Private/Views/SignalFlowDashboardViewFactory.cpp
--- a/Engine/Plugins/AudioInsights/Source/AudioInsights/Private/Views/SignalFlowDashboardViewFactory.cpp
+++ b/Engine/Plugins/AudioInsights/Source/AudioInsights/Private/Views/SignalFlowDashboardViewFactory.cpp
@@ -1402,7 +1402,7 @@ namespace UE::Audio::Insights
 		if (SelectedItem.IsValid())
 		{
 			XPosFocusNodeEntryKey = SelectedItem->GetSignalFlowEntryKey();
-			CreateFilteredNodeGraph(*Provider, SelectedItem);
+			CreateFilteredNodeGraph(*Provider, SelectedItem, false /*bAddAllNodes*/);
 		}
 		else
 		{
@@ -1410,7 +1410,7 @@ namespace UE::Audio::Insights
 		}
 	}
 
-	void FSignalFlowDashboardViewFactory::CreateFilteredNodeGraph(const FSignalFlowTraceProvider& Provider, TSharedPtr<FSignalFlowDashboardEntry> FilterFromEntry)
+	void FSignalFlowDashboardViewFactory::CreateFilteredNodeGraph(const FSignalFlowTraceProvider& Provider, TSharedPtr<FSignalFlowDashboardEntry> FilterFromEntry, const bool bAddAllNodes)
 	{
 		if (!FilterFromEntry.IsValid())
 		{
@@ -1425,6 +1425,8 @@ namespace UE::Audio::Insights
 
 		bGraphRequiresRefresh = true;
 
+		// Make a note of the existing entries in the graph before we start the creation process
+		// Any remaining existing nodes by the end of the process will be removed
 		TSet<FSignalFlowEntryKey> ExistingEntries;
 		ExistingEntries.Reserve(FilteredEntries.Num());
 		for (const auto& [EntryKey, Entry] : FilteredEntries)
@@ -1432,8 +1434,18 @@ namespace UE::Audio::Insights
 			ExistingEntries.Add(EntryKey);
 		}
 
+		// Add nodes recursively from the FilterFromEntry, connecting direct input and output paths
 		AddFilteredInputsRecursive(*DeviceData, FilterFromEntry, nullptr, 0, ExistingEntries);
 		AddFilteredOutputs(*DeviceData, FilterFromEntry, 1, ExistingEntries);
+
+		if (bAddAllNodes)
+		{
+			// Only try to locate missing nodes when the graph is not filtered
+			// When filtered, we only want to show the directly connected paths to the filtering node, making this step unnecessary
+			AddMissingNodesToGraph(*DeviceData, ExistingEntries);
+		}
+		
+		// Remove any remaining nodes in the graph that were not discovered in the creation process
 		CleanUpFilteredEntries(ExistingEntries);
 
 		// Gather stats (such as root entry and tree depth structure)
@@ -1491,7 +1503,7 @@ namespace UE::Audio::Insights
 		if (DeviceEntry.IsValid())
 		{
 			XPosFocusNodeEntryKey = DeviceEntry->GetSignalFlowEntryKey();
-			CreateFilteredNodeGraph(*Provider, DeviceEntry);
+			CreateFilteredNodeGraph(*Provider, DeviceEntry, true /*bAddAllNodes*/);
 		}
 	}
 
@@ -1669,6 +1681,196 @@ namespace UE::Audio::Insights
 		}
 	}
 
+	void FSignalFlowDashboardViewFactory::AddMissingNodesToGraph(const TSortedMap<FSignalFlowEntryKey, TSharedPtr<FSignalFlowDashboardEntry>>& DeviceData, TSet<FSignalFlowEntryKey>& OutEntriesMarkedToDelete)
+	{
+		// If a node has no inputs/outputs, it can get missed by the recursive node graph creation
+		// Run through all device data entries and add any missing entries to the graph
+		// Note: If a node has both no inputs and no outputs (like an unconnected audio bus), we ignore it entirely
+		for (const auto& [EntryKey, Entry] : DeviceData)
+		{
+			if (!Entry.IsValid())
+			{
+				continue;
+			}
+
+			// Ignore entries that are not connected to anything
+			if (Entry->Inputs.IsEmpty() && Entry->Outputs.IsEmpty())
+			{
+				continue;
+			}
+
+			// Ignore entries that have already been entered into the graph
+			if (EntryIsInGraph(EntryKey, OutEntriesMarkedToDelete))
+			{
+				continue;
+			}
+
+			if (Entry->Inputs.IsEmpty())
+			{
+				AddMissingNodeByOutputs(*Entry, DeviceData, OutEntriesMarkedToDelete);
+			}
+			else if (Entry->Outputs.IsEmpty())
+			{
+				AddMissingNodeByInputs(*Entry, DeviceData, OutEntriesMarkedToDelete);
+			}
+		}
+	}
+
+	void FSignalFlowDashboardViewFactory::AddMissingNodeByOutputs(const FSignalFlowDashboardEntry& Entry, const TSortedMap<FSignalFlowEntryKey, TSharedPtr<FSignalFlowDashboardEntry>>& DeviceData, TSet<FSignalFlowEntryKey>& OutEntriesMarkedToDelete)
+	{
+		// Gather the closest output nodes to this entry that are already in the graph
+		// This may require running down the branch recursively until we find the next output in the graph
+		TArray<TSharedPtr<FSignalFlowDashboardEntry>> NextOutputs;
+		GetNextOutputsInGraphRecursive(Entry, DeviceData, OutEntriesMarkedToDelete, NextOutputs);
+
+		// Run through all closest added outputs and add their inputs recursively, creating a connection path to the missing node
+		for (const TSharedPtr<FSignalFlowDashboardEntry>& NextOutputEntryInGraph : NextOutputs)
+		{
+			if (!NextOutputEntryInGraph.IsValid())
+			{
+				continue;
+			}
+
+			const FSignalFlowEntryKey OutputKey = NextOutputEntryInGraph->GetSignalFlowEntryKey();
+			const TSharedPtr<ISignalFlowNode>* NextOutputEntryInGraphNode = FilteredEntries.Find(OutputKey);
+			if (NextOutputEntryInGraphNode == nullptr || !NextOutputEntryInGraphNode->IsValid())
+			{
+				continue;
+			}
+
+			const int32 NextOutputEntryInGraphTreeDepth = (*NextOutputEntryInGraphNode)->GetTreeDepth().TreeDepth;
+			for (const FSignalFlowEntryKey& InputKey : NextOutputEntryInGraph->Inputs)
+			{
+				const TSharedPtr<FSignalFlowDashboardEntry>* InputEntry = DeviceData.Find(InputKey);
+				if (InputEntry == nullptr || !InputEntry->IsValid())
+				{
+					continue;
+				}
+
+				(*NextOutputEntryInGraphNode)->FilteredInputs.AddUnique(InputKey);
+
+				const TSharedPtr<ISignalFlowNode>* InputNode = FilteredEntries.Find(InputKey);
+				if (InputNode == nullptr || OutEntriesMarkedToDelete.Contains(InputKey))
+				{
+					bool bResetTreeDepth = (*InputEntry)->EntryType != NextOutputEntryInGraph->EntryType;
+					AddFilteredInputsRecursive(DeviceData, *InputEntry, NextOutputEntryInGraph, bResetTreeDepth ? 0 : NextOutputEntryInGraphTreeDepth - 1, OutEntriesMarkedToDelete);
+				}
+				else if (InputNode && InputNode->IsValid())
+				{
+					(*InputNode)->FilteredOutputs.Add(OutputKey);
+				}
+			}
+		}
+	}
+
+	void FSignalFlowDashboardViewFactory::AddMissingNodeByInputs(const FSignalFlowDashboardEntry& Entry, const TSortedMap<FSignalFlowEntryKey, TSharedPtr<FSignalFlowDashboardEntry>>& DeviceData, TSet<FSignalFlowEntryKey>& OutEntriesMarkedToDelete)
+	{
+		// Gather the closest input nodes to this entry that are already in the graph
+		// This may require running up the branch recursively until we find the next input in the graph
+		TArray<TSharedPtr<FSignalFlowDashboardEntry>> NextInputs;
+		GetNextInputsInGraphRecursive(Entry, DeviceData, OutEntriesMarkedToDelete, NextInputs);
+
+		// Run through all closest added inputs and add their outputs recursively, creating a connection path to the missing node
+		for (const TSharedPtr<FSignalFlowDashboardEntry>& NextInputEntryInGraph : NextInputs)
+		{
+			if (!NextInputEntryInGraph.IsValid())
+			{
+				continue;
+			}
+
+			const FSignalFlowEntryKey InputKey = NextInputEntryInGraph->GetSignalFlowEntryKey();
+			const TSharedPtr<ISignalFlowNode>* NextInputEntryInGraphNode = FilteredEntries.Find(InputKey);
+			if (NextInputEntryInGraphNode == nullptr || !NextInputEntryInGraphNode->IsValid())
+			{
+				continue;
+			}
+
+			const int32 NextInputEntryInGraphTreeDepth = (*NextInputEntryInGraphNode)->GetTreeDepth().TreeDepth;
+			for (const auto& [OutputKey, OutputInfo] : NextInputEntryInGraph->Outputs)
+			{
+				const TSharedPtr<FSignalFlowDashboardEntry>* OutputEntry = DeviceData.Find(OutputKey);
+				if (OutputEntry && OutputEntry->IsValid())
+				{
+					(*NextInputEntryInGraphNode)->FilteredOutputs.AddUnique(OutputKey);
+
+					const TSharedPtr<ISignalFlowNode>* OutputNode = FilteredEntries.Find(OutputKey);
+					if ((OutputNode == nullptr || OutEntriesMarkedToDelete.Contains(OutputKey)))
+					{
+						bool bResetTreeDepth = (*OutputEntry)->EntryType != NextInputEntryInGraph->EntryType;
+						AddFilteredOutputsRecursive(DeviceData, *OutputEntry, NextInputEntryInGraph, bResetTreeDepth ? 0 : NextInputEntryInGraphTreeDepth + 1, OutEntriesMarkedToDelete);
+					}
+					else if (OutputNode && OutputNode->IsValid())
+					{
+						(*OutputNode)->FilteredInputs.AddUnique(InputKey);
+					}
+				}
+			}
+		}
+	}
+
+	void FSignalFlowDashboardViewFactory::GetNextInputsInGraphRecursive(const FSignalFlowDashboardEntry& Entry, const TSortedMap<FSignalFlowEntryKey, TSharedPtr<FSignalFlowDashboardEntry>>& DeviceData, const TSet<FSignalFlowEntryKey>& EntriesMarkedToDelete, TArray<TSharedPtr<FSignalFlowDashboardEntry>>& OutNextInputs) const
+	{
+		// Iterate over this node's inputs - follow the branches until we reach the nearest nodes that are present in the current graph
+		// Gather all of these together and add to OutNextInputs
+		if (!Entry.IsValid())
+		{
+			return;
+		}
+
+		for (const FSignalFlowEntryKey& InputKey : Entry.Inputs)
+		{
+			if (const TSharedPtr<FSignalFlowDashboardEntry>* DeviceDataEntry = DeviceData.Find(InputKey))
+			{
+				if (!DeviceDataEntry->IsValid())
+				{
+					continue;
+				}
+
+				if (EntryIsInGraph(InputKey, EntriesMarkedToDelete))
+				{
+					OutNextInputs.Add(*DeviceDataEntry);
+				}
+				else
+				{
+					// This node is also not currently in the graph - we need to check it's inputs
+					GetNextInputsInGraphRecursive(**DeviceDataEntry, DeviceData, EntriesMarkedToDelete, OutNextInputs);
+				}
+			}
+		}
+	}
+
+	void FSignalFlowDashboardViewFactory::GetNextOutputsInGraphRecursive(const FSignalFlowDashboardEntry& Entry, const TSortedMap<FSignalFlowEntryKey, TSharedPtr<FSignalFlowDashboardEntry>>& DeviceData, const TSet<FSignalFlowEntryKey>& EntriesMarkedToDelete, TArray<TSharedPtr<FSignalFlowDashboardEntry>>& OutNextOutputs) const
+	{
+		// Iterate over this node's outputs - follow the branches until we reach the nearest nodes that are present in the current graph
+		// Gather all of these together and add to OutNextOutputs
+		for (const auto& [OutputKey, OutputData] : Entry.Outputs)
+		{
+			if (const TSharedPtr<FSignalFlowDashboardEntry>* DeviceDataEntry = DeviceData.Find(OutputKey))
+			{
+				if (!DeviceDataEntry->IsValid())
+				{
+					continue;
+				}
+
+				if (EntryIsInGraph(OutputKey, EntriesMarkedToDelete))
+				{
+					OutNextOutputs.Add(*DeviceDataEntry);
+				}
+				else
+				{
+					// This node is also not currently in the graph - we need to check it's outputs
+					GetNextOutputsInGraphRecursive(**DeviceDataEntry, DeviceData, EntriesMarkedToDelete, OutNextOutputs);
+				}
+			}
+		}
+	}
+
+	bool FSignalFlowDashboardViewFactory::EntryIsInGraph(const FSignalFlowEntryKey& EntryKey, const TSet<FSignalFlowEntryKey>& EntriesMarkedToDelete) const
+	{
+		// An entry is considered in the graph if a corresponding node exists inside FilteredEntries and it hasn't been marked for deletion
+		return FilteredEntries.Contains(EntryKey) && !EntriesMarkedToDelete.Contains(EntryKey);
+	}
+
 	void FSignalFlowDashboardViewFactory::CalculateFilteredGraphDepthStructure(TOptional<FSignalFlowEntryKey>& OutRootEntryKey, TArray<TreeDepthPair>& OutSortedTreeDepths) const
 	{
 		OutSortedTreeDepths.Reset();
diff --git a/Engine/Plugins/AudioInsights/Source/AudioInsights/Public/Views/SignalFlowDashboardViewFactory.h b/Engine/Plugins/AudioInsights/Source/AudioInsights/Public/Views/SignalFlowDashboardViewFactory.h
--- a/Engine/Plugins/AudioInsights/Source/AudioInsights/Public/Views/SignalFlowDashboardViewFactory.h
+++ b/Engine/Plugins/AudioInsights/Source/AudioInsights/Public/Views/SignalFlowDashboardViewFactory.h
@@ -112,7 +112,7 @@ namespace UE::Audio::Insights
 
 		virtual void RefreshFilteredEntriesListView() override;
 		void OnRequestGraphRefresh();
-		void CreateFilteredNodeGraph(const FSignalFlowTraceProvider& Provider, TSharedPtr<FSignalFlowDashboardEntry> FilterFromEntry);
+		void CreateFilteredNodeGraph(const FSignalFlowTraceProvider& Provider, TSharedPtr<FSignalFlowDashboardEntry> FilterFromEntry, const bool bAddAllNodes);
 		
 		void AddAllNodesToGraph();
 		void AddFilteredInputsRecursive(const TSortedMap<FSignalFlowEntryKey, TSharedPtr<FSignalFlowDashboardEntry>>& DeviceData, TSharedPtr<FSignalFlowDashboardEntry> Entry, TSharedPtr<FSignalFlowDashboardEntry> ChildEntry, int32 TreeDepth, TSet<FSignalFlowEntryKey>& OutExistingEntryKeys);
@@ -120,6 +120,13 @@ namespace UE::Audio::Insights
 		void AddFilteredOutputsRecursive(const TSortedMap<FSignalFlowEntryKey, TSharedPtr<FSignalFlowDashboardEntry>>& DeviceData, TSharedPtr<FSignalFlowDashboardEntry> Entry, TSharedPtr<FSignalFlowDashboardEntry> ParentEntry, int32 TreeDepth, TSet<FSignalFlowEntryKey>& OutExistingEntryKeys);
 		void CleanUpFilteredEntries(const TSet<FSignalFlowEntryKey>& EntriesToRemove);
 
+		void AddMissingNodesToGraph(const TSortedMap<FSignalFlowEntryKey, TSharedPtr<FSignalFlowDashboardEntry>>& DeviceData, TSet<FSignalFlowEntryKey>& OutEntriesMarkedToDelete);
+		void AddMissingNodeByOutputs(const FSignalFlowDashboardEntry& Entry, const TSortedMap<FSignalFlowEntryKey, TSharedPtr<FSignalFlowDashboardEntry>>& DeviceData, TSet<FSignalFlowEntryKey>& OutEntriesMarkedToDelete);
+		void AddMissingNodeByInputs(const FSignalFlowDashboardEntry& Entry, const TSortedMap<FSignalFlowEntryKey, TSharedPtr<FSignalFlowDashboardEntry>>& DeviceData, TSet<FSignalFlowEntryKey>& OutEntriesMarkedToDelete);
+		void GetNextInputsInGraphRecursive(const FSignalFlowDashboardEntry& Entry, const TSortedMap<FSignalFlowEntryKey, TSharedPtr<FSignalFlowDashboardEntry>>& DeviceData, const TSet<FSignalFlowEntryKey>& EntriesMarkedToDelete, TArray<TSharedPtr<FSignalFlowDashboardEntry>>& OutNextInputs) const;
+		void GetNextOutputsInGraphRecursive(const FSignalFlowDashboardEntry& Entry, const TSortedMap<FSignalFlowEntryKey, TSharedPtr<FSignalFlowDashboardEntry>>& DeviceData, const TSet<FSignalFlowEntryKey>& EntriesMarkedToDelete, TArray<TSharedPtr<FSignalFlowDashboardEntry>>& OutNextOutputs) const;
+		bool EntryIsInGraph(const FSignalFlowEntryKey& EntryKey, const TSet<FSignalFlowEntryKey>& EntriesMarkedToDelete) const;
+
 		void CalculateFilteredGraphDepthStructure(TOptional<FSignalFlowEntryKey>& OutRootEntryKey, TArray<TreeDepthPair>& OutSortedTreeDepths) const;
 
 		void CreateDummyInputs(const TArray<TreeDepthPair>& TreeDepthStructure);
