[Audio Insights] Support for source buses in signal flow part 01: Provider trace data support

* This change provides the data required for supporting source buses inside the signal flow dashboard and graph
* Source buses already have two distinct entries in the graph - a sound source and a bus entry. This change manages the link between the two in the provider device data map

Trace messages
* Add a new SourceBusID field to the SoundWaveStart and IsAlivePing messages. Send from AudioMixerSource
* FSoundWaveStartMessage provides a default value (INDEX_NONE) for this field when trying to retrieve it from the message EventData in case the field does not exist, which maintains backwards compatibility

FSignalFlowDashboardEntry
* Two new members added to store the linked source bus if this is a sound source (one-to-one relationship) or the linked sound sources if this is a source bus (one-to-many relationship)

FSignalFlowTraceProvider
* When a new sound source is created (SoundWaveStart), we now link this to a source bus if an ID exists. If the bus is already present, we add the new sound source to it's LinkedSoundSources.
* Because we cannot guarentee the bus entry will be available when the sound source entry is created, we piggyback the loop that cleans up inputs TidyConnections (formerly RemoveInputs) to also add/remove these source bus -> sound source links

#Synced-CL 51626790

diff --git a/Engine/Plugins/AudioInsights/Source/AudioInsights/Private/Providers/SignalFlowTraceProvider.cpp b/Engine/Plugins/AudioInsights/Source/AudioInsights/Private/Providers/SignalFlowTraceProvider.cpp
--- a/Engine/Plugins/AudioInsights/Source/AudioInsights/Private/Providers/SignalFlowTraceProvider.cpp
+++ b/Engine/Plugins/AudioInsights/Source/AudioInsights/Private/Providers/SignalFlowTraceProvider.cpp
@@ -360,15 +360,16 @@ namespace UE::Audio::Insights
 		};
 
 		TSet<FSignalFlowEntryKey> RemovedKeys;
+		TSet<FSignalFlowEntryKey> NewBuses;
 
 		ProcessMessageQueue<FAudioBusStartMessage>(TraceMessages.BusStartMessages, 
 		[this](const FAudioBusStartMessage& Msg) 
 		{ 
 			return CreateAudioBusEntry(Msg); 
 		},
-		[this](const FAudioBusStartMessage& Msg, TSharedPtr<FSignalFlowDashboardEntry>* OutEntry)
+		[this, &NewBuses](const FAudioBusStartMessage& Msg, TSharedPtr<FSignalFlowDashboardEntry>* OutEntry)
 		{
-			ProcessNewBusEntry(Msg, OutEntry);
+			ProcessNewBusEntry(Msg, OutEntry, NewBuses);
 		});
 
 		ProcessMessageQueue<FAudioBusStartMessage>(TraceMessages.BusIsAlivePingMessages, 
@@ -376,9 +377,9 @@ namespace UE::Audio::Insights
 		{
 			return CreateAudioBusEntry(Msg);
 		},
-		[this](const FAudioBusStartMessage& Msg, TSharedPtr<FSignalFlowDashboardEntry>* OutEntry)
+		[this, &NewBuses](const FAudioBusStartMessage& Msg, TSharedPtr<FSignalFlowDashboardEntry>* OutEntry)
 		{
-			ProcessNewBusEntry(Msg, OutEntry);
+			ProcessNewBusEntry(Msg, OutEntry, NewBuses);
 		});
 
 		ProcessMessageQueue<FAudioBusStopMessage>(TraceMessages.BusStopMessages, GetBusEntry,
@@ -501,9 +502,9 @@ namespace UE::Audio::Insights
 			}
 		}
 
-		if (!RemovedKeys.IsEmpty())
+		if (!RemovedKeys.IsEmpty() || !NewBuses.IsEmpty())
 		{
-			RemoveInputs(RemovedKeys);
+			TidyConnections(RemovedKeys, NewBuses);
 			bForceGraphRefresh = true;
 		}
 
@@ -614,6 +615,7 @@ namespace UE::Audio::Insights
 	{
 		TSet<FSignalFlowEntryKey> ActiveEntries;
 		TSet<FSignalFlowEntryKey> RemovedKeys;
+		TSet<FSignalFlowEntryKey> NewBuses;
 
 		const FAudioInsightsCacheManager& CacheManager = IAudioInsightsModule::GetChecked().GetCacheManager();
 
@@ -625,12 +627,12 @@ namespace UE::Audio::Insights
 		/*
 		 * Audio Buses
 		*/
-		CacheManager.IterateOverRange<FAudioBusStartMessage>(AudioBusMessageNames::Start, RangeCheckStart, TimeMarker, [this, &ActiveEntries](const FAudioBusStartMessage& Msg)
+		CacheManager.IterateOverRange<FAudioBusStartMessage>(AudioBusMessageNames::Start, RangeCheckStart, TimeMarker, [this, &ActiveEntries, &NewBuses](const FAudioBusStartMessage& Msg)
 		{
 			TSharedPtr<FSignalFlowDashboardEntry>* Entry = CreateAudioBusEntry(Msg, &ActiveEntries);
 			if (Entry && Entry->IsValid())
 			{
-				ProcessNewBusEntry(Msg, Entry);
+				ProcessNewBusEntry(Msg, Entry, NewBuses);
 			}
 		});
 
@@ -758,9 +760,9 @@ namespace UE::Audio::Insights
 			}
 		}
 
-		if (!RemovedKeys.IsEmpty())
+		if (!RemovedKeys.IsEmpty() || !NewBuses.IsEmpty())
 		{
-			RemoveInputs(RemovedKeys);
+			TidyConnections(RemovedKeys, NewBuses);
 			bForceGraphRefresh = true;
 		}
 
@@ -1099,7 +1101,7 @@ namespace UE::Audio::Insights
 		}
 	}
 
-	void FSignalFlowTraceProvider::ProcessNewBusEntry(const FAudioBusStartMessage& Msg, TSharedPtr<FSignalFlowDashboardEntry>* OutEntry)
+	void FSignalFlowTraceProvider::ProcessNewBusEntry(const FAudioBusStartMessage& Msg, TSharedPtr<FSignalFlowDashboardEntry>* OutEntry, TSet<FSignalFlowEntryKey>& OutNewBuses)
 	{
 		if (OutEntry && OutEntry->IsValid())
 		{
@@ -1107,6 +1109,8 @@ namespace UE::Audio::Insights
 			EntryRef.Name = Msg.Name;
 			EntryRef.IconName = FSignalFlowTraceProviderPrivate::AudioBusIconName;
 			EntryRef.Timestamp = Msg.Timestamp;
+
+			OutNewBuses.Add(EntryRef.GetSignalFlowEntryKey());
 		}
 	}
 
@@ -1137,6 +1141,17 @@ namespace UE::Audio::Insights
 			EntryRef.Name = Msg.Name;
 			EntryRef.IconName = FSignalFlowTraceProviderPrivate::GetSourceIconName(Msg.EntryType);
 
+			if (Msg.SourceBusID != INDEX_NONE)
+			{
+				EntryRef.LinkedSourceBus = FSignalFlowEntryKey(Msg.DeviceId, ESignalFlowEntryType::AudioBus, Msg.SourceBusID);
+
+				const TSharedPtr<FSignalFlowDashboardEntry>* FoundBusEntry = FindDeviceEntry(Msg.DeviceId, EntryRef.LinkedSourceBus.GetValue());
+				if (FoundBusEntry && FoundBusEntry->IsValid())
+				{
+					(*FoundBusEntry)->LinkedSoundSources.Add(EntryRef.GetSignalFlowEntryKey());
+				}
+			}
+
 			TSharedPtr<FSignalFlowDashboardEntry>* OwnerObjectEntry = CreateOwnerEntry(Msg, OutActiveEntriesTracker);
 			if (OwnerObjectEntry && OwnerObjectEntry->IsValid())
 			{
@@ -1200,8 +1215,9 @@ namespace UE::Audio::Insights
 		}
 	}
 
-	void FSignalFlowTraceProvider::RemoveInputs(const TSet<FSignalFlowEntryKey>& RemovedKeys)
+	void FSignalFlowTraceProvider::TidyConnections(const TSet<FSignalFlowEntryKey>& RemovedKeys, const TSet<FSignalFlowEntryKey>& NewBuses)
 	{
+		// Iterate over the device data map - remove old inputs and ensure we have connections between source buses and soung sources
 		for (auto& [DeviceID, Entries] : DeviceDataMap)
 		{
 			for (auto& [EntryKey, Entry] : Entries)
@@ -1211,6 +1227,7 @@ namespace UE::Audio::Insights
 					continue;
 				}
 
+				// Clean up old inputs
 				for (auto Itr = Entry->Inputs.CreateIterator(); Itr; ++Itr)
 				{
 					if (RemovedKeys.Contains(*Itr))
@@ -1218,6 +1235,33 @@ namespace UE::Audio::Insights
 						Itr.RemoveCurrent();
 					}
 				}
+
+				// Ensure source buses are linked
+				if (Entry->LinkedSourceBus.IsSet())
+				{
+					const FSignalFlowEntryKey SourceBusEntryKey = Entry->LinkedSourceBus.GetValue();
+					if (RemovedKeys.Contains(SourceBusEntryKey))
+					{
+						Entry->LinkedSourceBus.Reset();
+					}
+					else if (NewBuses.Contains(SourceBusEntryKey))
+					{
+						const TSharedPtr<FSignalFlowDashboardEntry>* SourceBusEntry = FindDeviceEntry(DeviceID, SourceBusEntryKey);
+						if (SourceBusEntry && SourceBusEntry->IsValid())
+						{
+							(*SourceBusEntry)->LinkedSoundSources.Add(EntryKey);
+						}
+					}
+				}
+
+				// Ensure stale linked sound sources are removed
+				for (auto Itr = Entry->LinkedSoundSources.CreateIterator(); Itr; ++Itr)
+				{
+					if (RemovedKeys.Contains(*Itr))
+					{
+						Itr.RemoveCurrent();
+					}
+				}
 			}
 		}
 	}
diff --git a/Engine/Plugins/AudioInsights/Source/AudioInsights/Private/Providers/SignalFlowTraceProvider.h b/Engine/Plugins/AudioInsights/Source/AudioInsights/Private/Providers/SignalFlowTraceProvider.h
--- a/Engine/Plugins/AudioInsights/Source/AudioInsights/Private/Providers/SignalFlowTraceProvider.h
+++ b/Engine/Plugins/AudioInsights/Source/AudioInsights/Private/Providers/SignalFlowTraceProvider.h
@@ -42,7 +42,7 @@ namespace UE::Audio::Insights
 		TSharedPtr<FSignalFlowDashboardEntry>* CreateAudioBusEntry(const FAudioBusStartMessage& Msg, TSet<FSignalFlowEntryKey>* OutActiveEntriesTracker = nullptr);
 		TSharedPtr<FSignalFlowDashboardEntry>* CreateSubmixEntry(const FSubmixMessageBase& Msg, TSet<FSignalFlowEntryKey>* OutActiveEntriesTracker = nullptr);
 
-		void ProcessNewBusEntry(const FAudioBusStartMessage& Msg, TSharedPtr<FSignalFlowDashboardEntry>* OutEntry);
+		void ProcessNewBusEntry(const FAudioBusStartMessage& Msg, TSharedPtr<FSignalFlowDashboardEntry>* OutEntry, TSet<FSignalFlowEntryKey>& OutNewBuses);
 		void ProcessNewSubmixEntry(const FSubmixLoadedMessage& Msg, TSharedPtr<FSignalFlowDashboardEntry>* OutEntry);
 		void ProcessNewSoundwaveEntry(const FSoundWaveStartMessage& Msg, TSharedPtr<FSignalFlowDashboardEntry>* OutEntry, TSet<FSignalFlowEntryKey>* OutActiveEntriesTracker = nullptr);
 		void ProcessStopMessage(const ::Audio::FDeviceId DeviceID, const double MsgTimestamp, TSharedPtr<FSignalFlowDashboardEntry>* OutEntry, TSet<FSignalFlowEntryKey>& OutRemovedKeys);
@@ -51,7 +51,7 @@ namespace UE::Audio::Insights
 		void ProcessMultiChannelEnvelopeMessage(TSharedPtr<FSignalFlowDashboardEntry>* OutEntry, const TArray<float>& EnvelopeValues);
 
 		void LinkInput(const FSignalFlowDashboardEntry& Entry, const FSignalFlowEntryKey& OutputKey, const ::Audio::FDeviceId DeviceID);
-		void RemoveInputs(const TSet<FSignalFlowEntryKey>& RemovedKeys);
+		void TidyConnections(const TSet<FSignalFlowEntryKey>& RemovedKeys, const TSet<FSignalFlowEntryKey>& NewBuses);
 
 		FSignalFlowMessages TraceMessages;
 
diff --git a/Engine/Plugins/AudioInsights/Source/AudioInsights/Public/Messages/SignalFlowTraceMessages.h b/Engine/Plugins/AudioInsights/Source/AudioInsights/Public/Messages/SignalFlowTraceMessages.h
--- a/Engine/Plugins/AudioInsights/Source/AudioInsights/Public/Messages/SignalFlowTraceMessages.h
+++ b/Engine/Plugins/AudioInsights/Source/AudioInsights/Public/Messages/SignalFlowTraceMessages.h
@@ -152,6 +152,9 @@ namespace UE::Audio::Insights
 		TSet<FSignalFlowEntryKey> Inputs;
 		TSortedMap<FSignalFlowEntryKey, FSignalFlowOutputData> Outputs;
 
+		TOptional<FSignalFlowEntryKey> LinkedSourceBus;
+		TSet<FSignalFlowEntryKey> LinkedSoundSources;
+
 		TCircularBuffer<float> AmplitudeDataRange;
 		int32 AmplitudeWriteIndex = 0;
 
diff --git a/Engine/Plugins/AudioInsights/Source/AudioInsights/Public/Messages/SoundTraceMessages.h b/Engine/Plugins/AudioInsights/Source/AudioInsights/Public/Messages/SoundTraceMessages.h
--- a/Engine/Plugins/AudioInsights/Source/AudioInsights/Public/Messages/SoundTraceMessages.h
+++ b/Engine/Plugins/AudioInsights/Source/AudioInsights/Public/Messages/SoundTraceMessages.h
@@ -125,6 +125,7 @@ namespace UE::Audio::Insights
 		{
 			WaveInstancePlayOrder = InContext.EventData.GetValue<uint32>("PlayOrder");
 			ActiveSoundPlayOrder = InContext.EventData.GetValue<uint32>("ActiveSoundPlayOrder");
+			SourceBusID = InContext.EventData.GetValue<int32>("SourceBusID", INDEX_NONE);
 
 			InContext.EventData.GetString("ActorIconName", ActorIconName);
 		}
@@ -135,6 +136,7 @@ namespace UE::Audio::Insights
 
 		uint32 WaveInstancePlayOrder = INDEX_NONE;
 		FString ActorIconName;
+		int32 SourceBusID = INDEX_NONE;
 	};
 
 	struct FSoundWaveIsAlivePingMessage : public FSoundWaveStartMessage
diff --git a/Engine/Source/Runtime/AudioMixer/Private/AudioMixerSource.cpp b/Engine/Source/Runtime/AudioMixer/Private/AudioMixerSource.cpp
--- a/Engine/Source/Runtime/AudioMixer/Private/AudioMixerSource.cpp
+++ b/Engine/Source/Runtime/AudioMixer/Private/AudioMixerSource.cpp
@@ -34,6 +34,7 @@ UE_TRACE_EVENT_BEGIN(Audio, SoundWaveStart)
 	UE_TRACE_EVENT_FIELD(UE::Trace::WideString, SoundClassName)
 	UE_TRACE_EVENT_FIELD(UE::Trace::WideString, ActorLabel)
 	UE_TRACE_EVENT_FIELD(UE::Trace::WideString, ActorIconName)
+	UE_TRACE_EVENT_FIELD(int32, SourceBusID)
 UE_TRACE_EVENT_END()
 
 UE_TRACE_EVENT_BEGIN(Audio, SoundWaveIsAlivePing)
@@ -45,6 +46,7 @@ UE_TRACE_EVENT_BEGIN(Audio, SoundWaveIsAlivePing)
 	UE_TRACE_EVENT_FIELD(UE::Trace::WideString, SoundClassName)
 	UE_TRACE_EVENT_FIELD(UE::Trace::WideString, ActorLabel)
 	UE_TRACE_EVENT_FIELD(UE::Trace::WideString, ActorIconName)
+	UE_TRACE_EVENT_FIELD(int32, SourceBusID)
 UE_TRACE_EVENT_END()
 
 UE_TRACE_EVENT_BEGIN(Audio, SoundWaveStop)
@@ -122,6 +124,16 @@ namespace Audio
 			}
 		}
 		
+		int32 GetSourceBusID(const TObjectPtr<USoundWave> WaveData)
+		{
+			if (WaveData && WaveData->bIsSourceBus)
+			{
+				const TObjectPtr<USoundSourceBus> SoundSourceBus = CastChecked<USoundSourceBus>(WaveData);
+				return SoundSourceBus->AudioBus ? SoundSourceBus->AudioBus->GetUniqueID() : WaveData->GetUniqueID();
+			}
+
+			return INDEX_NONE;
+		}
 	} // namespace MixerSourcePrivate
 
 	namespace ModulationUtils
@@ -1906,6 +1918,8 @@ namespace Audio
 #if UE_AUDIO_PROFILERTRACE_ENABLED
 	void FMixerSource::SendSoundWaveStartTraceMessage() const
 	{
+		using namespace MixerSourcePrivate;
+
 		const bool bChannelEnabled = UE_TRACE_CHANNELEXPR_IS_ENABLED(AudioChannel);
 		if (bChannelEnabled && WaveInstance && WaveInstance->WaveData && WaveInstance->ActiveSound && WaveInstance->SoundClass)
 		{
@@ -1914,6 +1928,7 @@ namespace Audio
 			{
 				const FString SoundWaveBaseAssetName = Audio::Trace::Util::GetSoundBaseAssetName(SoundWaveObjectPtr->GetClass());
 				const FString OwnerActorLabel = Audio::Trace::Util::GetOwnerActorLabel(*WaveInstance->ActiveSound);
+				const int32 SourceBusID = GetSourceBusID(WaveInstance->WaveData);
 
 				UE_TRACE_LOG(Audio, SoundWaveStart, AudioChannel)
 					<< SoundWaveStart.DeviceId(AudioDevice->DeviceID)
@@ -1923,7 +1938,8 @@ namespace Audio
 					<< SoundWaveStart.AssetPath(*SoundWaveObjectPtr->GetPathName())
 					<< SoundWaveStart.SoundClassName(*SoundWaveBaseAssetName)
 					<< SoundWaveStart.ActorLabel(*OwnerActorLabel)
-					<< SoundWaveStart.ActorIconName(*Audio::Trace::Util::GetOwnerActorIconName(*WaveInstance->ActiveSound).ToString());
+					<< SoundWaveStart.ActorIconName(*Audio::Trace::Util::GetOwnerActorIconName(*WaveInstance->ActiveSound).ToString())
+					<< SoundWaveStart.SourceBusID(SourceBusID);
 			}
 		}
 	}
@@ -1943,6 +1959,8 @@ namespace Audio
 
 	void FMixerSource::UpdateIsAlivePingTraceMessage() const
 	{
+		using namespace MixerSourcePrivate;
+
 		const bool bChannelEnabled = UE_TRACE_CHANNELEXPR_IS_ENABLED(AudioChannel);
 
 		if (bChannelEnabled && WaveInstance && WaveInstance->IsPlaying() && WaveInstance->WaveData && WaveInstance->ActiveSound && WaveInstance->SoundClass)
@@ -1956,6 +1974,7 @@ namespace Audio
 				{
 					const FString SoundWaveBaseAssetName = Audio::Trace::Util::GetSoundBaseAssetName(SoundWaveObjectPtr->GetClass());
 					const FString OwnerActorLabel = Audio::Trace::Util::GetOwnerActorLabel(*WaveInstance->ActiveSound);
+					const int32 SourceBusID = GetSourceBusID(WaveInstance->WaveData);
 
 					UE_TRACE_LOG(Audio, SoundWaveIsAlivePing, AudioChannel)
 						<< SoundWaveIsAlivePing.DeviceId(AudioDevice->DeviceID)
@@ -1965,7 +1984,8 @@ namespace Audio
 						<< SoundWaveIsAlivePing.AssetPath(*SoundWaveObjectPtr->GetPathName())
 						<< SoundWaveIsAlivePing.SoundClassName(*SoundWaveBaseAssetName)
 						<< SoundWaveIsAlivePing.ActorLabel(*OwnerActorLabel)
-						<< SoundWaveIsAlivePing.ActorIconName(*Audio::Trace::Util::GetOwnerActorIconName(*WaveInstance->ActiveSound).ToString());
+						<< SoundWaveIsAlivePing.ActorIconName(*Audio::Trace::Util::GetOwnerActorIconName(*WaveInstance->ActiveSound).ToString())
+						<< SoundWaveIsAlivePing.SourceBusID(SourceBusID);
 				}
 			}
 		}
