[AudioInsights] - Event log: Minor code clean up

#Synced-CL 51664040

diff --git a/Engine/Plugins/AudioInsights/Source/AudioInsights/Private/Views/AudioEventLogDashboardViewFactory.cpp b/Engine/Plugins/AudioInsights/Source/AudioInsights/Private/Views/AudioEventLogDashboardViewFactory.cpp
--- a/Engine/Plugins/AudioInsights/Source/AudioInsights/Private/Views/AudioEventLogDashboardViewFactory.cpp
+++ b/Engine/Plugins/AudioInsights/Source/AudioInsights/Private/Views/AudioEventLogDashboardViewFactory.cpp
@@ -1267,14 +1267,12 @@ namespace UE::Audio::Insights
 		if (AutoStopCachingMode == EAutoStopCachingMode::OnInspect)
 		{
 			AudioInsightsModule.GetTimingViewExtender().StopCachingAndProcessingNewMessages();
-			bAutoScroll = false;
 		}
 		else if (AutoStopCachingMode == EAutoStopCachingMode::WhenInLastChunk)
 		{
 			if (TimestampIsInChunkMarkedToDelete(InTimeMarker))
 			{
 				AudioInsightsModule.GetTimingViewExtender().StopCachingAndProcessingNewMessages();
-				bAutoScroll = false;
 			}
 			else if (AudioInsightsModule.GetTimingViewExtender().GetMessageCacheAndProcessingStatus() != ECacheAndProcess::None)
 			{
