[Audio Insights] Support for source buses in signal flow part 02: Dashboard and graph
* Connect up source bus sound source and audio bus nodes in the signal flow graph
* Make sure that both the sound source and audio bus paths are visible when the graph is filtered on a node that includes a source bus path.

ISignalFlowNode
* New fields for storing the outputs from a source bus and the source bus input in a sound source
* These are different from the field in the dashboard entry as they only contain the entry keys for filtered nodes
* Source bus connections have been kept separate from the filtered inputs/outputs as we don't want them to have an effect on the positioning and sorting algorithms, and the separation helps us manage the connections differently (for example, it helps combat creating infinite loops when filtering the graph).

Dashboard Factory
* Added additional steps when filtering the graph to check the device data for linked source bus/sound source entries, maintaining connections when found
* The first time a source bus/linked sound source entry is discovered when filtering the graph, we effectively create a new filter path from that node (i.e. adding both inputs and outputs to the filtered entries). Filtering is usually one directional, so doing this allows the full context of the source bus connection to be displayed on the graph
* Make sure we clean up any removed source bus/linked sound source connections

Signal Flow Graph
* Source bus -> Linked Sound Source connections drawn in a different color, with a lower alpha than normal connections
* Increase the spline radius for source bus connections - this is because they are drawn flowing backwards in the graph, so we want a larger curve from the output pin
* Support additional source bus output pins from the source bus node for linked sound sources. These are drawn below any normal output pins (i.e. index = linked source index + number of normal outputs)

#Synced-CL 51626790

diff --git a/Engine/Plugins/AudioInsights/Source/AudioInsights/Private/Views/SSignalFlowGraph.cpp b/Engine/Plugins/AudioInsights/Source/AudioInsights/Private/Views/SSignalFlowGraph.cpp
--- a/Engine/Plugins/AudioInsights/Source/AudioInsights/Private/Views/SSignalFlowGraph.cpp
+++ b/Engine/Plugins/AudioInsights/Source/AudioInsights/Private/Views/SSignalFlowGraph.cpp
@@ -34,7 +34,8 @@ namespace UE::Audio::Insights
 
 		constexpr float NodeRoundedBrushRadius = 12.0f;
 
-		constexpr float SplineRadius = 80.0f;
+		constexpr float SplineRadiusFlowingForwards = 80.0f;
+		constexpr float SplineRadiusFlowingBackwards = 600.0f;
 		constexpr float SplineThicknessScalarMin = 1.0f;
 		constexpr float SplineThicknessDefault = 1.5f;
 		constexpr float SplineAnimatedWireThickness = 3.0f;
@@ -91,18 +92,26 @@ namespace UE::Audio::Insights
 			}
 		}
 
-		FLinearColor GetConnectionColor(const float AmplitudeValue, const float WireSplineAmplitudePowerFactor)
+		FLinearColor GetConnectionColor(const float AmplitudeValue, const float WireSplineAmplitudePowerFactor, const bool bIsSourceBusConnection)
 		{
+			if (bIsSourceBusConnection)
+			{
+				constexpr float MaxAlpha = 0.5f;
+				return FLinearColor::LerpUsingHSV(FStyleColors::White25.GetSpecifiedColor(), GetNodeAccentColor(ESignalFlowEntryType::SoundSource).GetSpecifiedColor().CopyWithNewOpacity(MaxAlpha), FMath::Pow(AmplitudeValue, WireSplineAmplitudePowerFactor));
+			}
+
 			return FLinearColor::LerpUsingHSV(FStyleColors::Hover.GetSpecifiedColor(), FStyleColors::AccentWhite.GetSpecifiedColor(), FMath::Pow(AmplitudeValue, WireSplineAmplitudePowerFactor));
 		}
 
-		FVector2d GetSplineStartDirection(const EOrientation Orientation)
+		FVector2d GetSplineStartDirection(const EOrientation Orientation, const bool bIsSourceBusConnection)
 		{
+			const float SplineRadius = bIsSourceBusConnection ? SplineRadiusFlowingBackwards : SplineRadiusFlowingForwards;
 			return Orientation == EOrientation::Orient_Horizontal ? FVector2d{ SplineRadius, 0.0f } : FVector2d{ 0.0f, SplineRadius };
 		}
 
-		FVector2d GetSplineEndDirection(const EOrientation Orientation)
+		FVector2d GetSplineEndDirection(const EOrientation Orientation, const bool bIsSourceBusConnection)
 		{
+			const float SplineRadius = bIsSourceBusConnection ?  SplineRadiusFlowingBackwards : SplineRadiusFlowingForwards;
 			return Orientation == EOrientation::Orient_Horizontal ? FVector2d{ SplineRadius, 0.0f } : FVector2d{ 0.0f, SplineRadius };
 		}
 
@@ -633,6 +642,11 @@ namespace UE::Audio::Insights
 					{
 						return SendInfo->SendLevel;
 					}
+					else if (Entry->LinkedSoundSources.Contains(ReceiverEntryKey))
+					{
+						// Source buses always have a send value of 1.0
+						return 1.0f;
+					}
 				}
 
 				return NullOpt;
@@ -1439,6 +1453,35 @@ namespace UE::Audio::Insights
 				continue;
 			}
 
+			const TSharedPtr<FSignalFlowDashboardEntry> Entry = (*InputNode)->GetEntry();
+
+			// Draw any source bus connections first (flowing backwards through the graph)
+			for (int32 SourceBusIndex = 0; SourceBusIndex < Node->FilteredLinkedSoundSources.Num(); ++SourceBusIndex)
+			{
+				const FSignalFlowEntryKey& LinkedSourceKey = Node->FilteredLinkedSoundSources[SourceBusIndex];
+				const TSharedPtr<SSignalFlowGraphNode>* LinkedSourceNode = NodeWidgets.Find(LinkedSourceKey);
+
+				if (LinkedSourceNode && LinkedSourceNode->IsValid())
+				{
+					constexpr float SendLevel = 1.0f;
+					const int32 SendLabelIndex = SourceBusIndex + Node->FilteredOutputs.Num();
+					constexpr bool bIsSourceBusConnection = true;
+
+					LayerId = PaintConnectionToNode(AllottedGeometry,
+													OutDrawElements,
+													LayerId,
+													Entry,
+													**InputNode,
+													**LinkedSourceNode,
+													LocalBoundsBox,
+													ViewportPosInGraph,
+													SendLevel,
+													SendLabelIndex,
+													bIsSourceBusConnection);
+				}
+			}
+
+			// Now draw any output connections (flowing forwards through the graph)
 			for (int32 OutputIndex = 0; OutputIndex < Node->FilteredOutputs.Num(); ++OutputIndex)
 			{
 				const FSignalFlowEntryKey& NodeOutputKey = Node->FilteredOutputs[OutputIndex];
@@ -1446,46 +1489,74 @@ namespace UE::Audio::Insights
 
 				if (OutputNode.IsValid())
 				{
-					const FVector2D SendLabelPos = (*InputNode)->GetLocalSendLabelPos(OutputNode->GetNodeEntryKey(), OutputIndex);
-					const FVector2D SendLabelPosInView = SendLabelPos - ViewportPosInGraph;
+					FSignalFlowOutputData* OutputData = Entry.IsValid() ? Entry->Outputs.Find(OutputNode->GetNodeEntryKey()) : nullptr;
+					const float SendLevel = OutputData && OutputData->SendLevel.IsSet() ? OutputData->SendLevel.GetValue() : 1.0f;
+
+					constexpr bool bIsSourceBusConnection = false;
+
+					LayerId = PaintConnectionToNode(AllottedGeometry,
+													OutDrawElements,
+													LayerId,
+													Entry,
+													**InputNode,
+													*OutputNode,
+													LocalBoundsBox,
+													ViewportPosInGraph,
+													SendLevel,
+													OutputIndex,
+													bIsSourceBusConnection);
+				}
+			}
+		}
 
-					// Check if the node connection wire is inside the viewport
-					const FVector2D ConnectionEndPos = OutputNode->GetLocalInputPos();
-					const FVector2D ConnectionEndPosInView = ConnectionEndPos - ViewportPosInGraph;
-					if (!FMath::LineBoxIntersection2D(LocalBoundsBox, SendLabelPosInView, ConnectionEndPosInView))
-					{
-						continue;
-					}
+		return LayerId;
+	}
 
-					const TSharedPtr<SSignalFlowGraphNode> RealInputNode = FindRealNode((*InputNode)->GetNodeEntryKey(), ENodeConnectionType::Input);
-					const TSharedPtr<FSignalFlowDashboardEntry> Entry = RealInputNode.IsValid() ? RealInputNode->GetEntry() : nullptr;
+	int32 SSignalFlowGraph::PaintConnectionToNode(const FGeometry& AllottedGeometry,
+												  FSlateWindowElementList& OutDrawElements,
+												  int32 LayerId,
+												  const TSharedPtr<FSignalFlowDashboardEntry>& SenderEntry,
+												  const SSignalFlowGraphNode& SenderNode,
+												  const SSignalFlowGraphNode& ReceiverNode,
+												  const FBox2D& LocalBoundsBox,
+												  const FVector2D& ViewportPosInGraph,
+												  const float SendLevel,
+												  const int32 SendLabelIndex,
+												  const bool bIsSourceBusConnection) const
+	{
+		using namespace SSignalFlowGraphStyle;
 
-					if (AnimateWires.Get() && Entry.IsValid() && Entry->Amplitude.IsSet())
-					{
-						FSignalFlowOutputData* OutputData = Entry->Outputs.Find(OutputNode->GetNodeEntryKey());
-						const float SendLevel = OutputData && OutputData->SendLevel.IsSet() ? OutputData->SendLevel.GetValue() : 1.0f;
+		const FVector2D SendPos = SenderNode.GetLocalSendLabelPos(ReceiverNode.GetNodeEntryKey(), SendLabelIndex);
+		const FVector2D SendPosInView = SendPos - ViewportPosInGraph;
 
-						LayerId = AnimateWire(AllottedGeometry, OutDrawElements, LayerId, Entry, SendLabelPos, ConnectionEndPos, SendLevel);
-					}
-					else
-					{
-						const float ScaledSplineThickness = FMath::Max(SplineThicknessDefault * ZoomScaleFactor, SplineThicknessMin);
-
-						FSlateDrawElement::MakeSpline(
-							OutDrawElements,
-							++LayerId,
-							AllottedGeometry.ToPaintGeometry(),
-							SendLabelPos,
-							GetSplineStartDirection(Orientation),
-							ConnectionEndPos,
-							GetSplineEndDirection(Orientation),
-							ScaledSplineThickness,
-							ESlateDrawEffect::NoPixelSnapping,
-							GetConnectionColor(0.0f, WireSplineAmplitudePowerFactor)
-						);
-					}
-				}
-			}
+		// Check if the node connection wire is inside the viewport
+		const FVector2D ReceiverPos = ReceiverNode.GetLocalInputPos();
+		const FVector2D ReceiverPosInView = ReceiverPos - ViewportPosInGraph;
+		if (!FMath::LineBoxIntersection2D(LocalBoundsBox, SendPosInView, ReceiverPosInView))
+		{
+			return LayerId;
+		}
+
+		if (AnimateWires.Get() && SenderEntry.IsValid() && SenderEntry->Amplitude.IsSet())
+		{
+			LayerId = AnimateWire(AllottedGeometry, OutDrawElements, LayerId, SenderEntry, SendPos, ReceiverPos, SendLevel, bIsSourceBusConnection);
+		}
+		else
+		{
+			const float ScaledSplineThickness = FMath::Max(SplineThicknessDefault * ZoomScaleFactor, SplineThicknessMin);
+
+			FSlateDrawElement::MakeSpline(
+				OutDrawElements,
+				++LayerId,
+				AllottedGeometry.ToPaintGeometry(),
+				SendPos,
+				GetSplineStartDirection(Orientation, bIsSourceBusConnection),
+				ReceiverPos,
+				GetSplineEndDirection(Orientation, bIsSourceBusConnection),
+				ScaledSplineThickness,
+				ESlateDrawEffect::NoPixelSnapping,
+				GetConnectionColor(0.0f, WireSplineAmplitudePowerFactor, bIsSourceBusConnection)
+			);
 		}
 
 		return LayerId;
@@ -1497,7 +1568,8 @@ namespace UE::Audio::Insights
 										const TSharedPtr<FSignalFlowDashboardEntry>& Entry,
 										const FVector2D& ConnectionStartPos,
 										const FVector2D& ConnectionEndPos,
-										const float SendLevel) const
+										const float SendLevel,
+										const bool bIsSourceBusConnection) const
 	{
 		using namespace SSignalFlowGraphStyle;
 
@@ -1528,8 +1600,8 @@ namespace UE::Audio::Insights
 		int32 ReadIndex = Entry->AmplitudeDataRange.GetPreviousIndex(AmpWriteIndex);
 		uint32 AmpReadingCount = 0u;
 
-		const FVector2D StartDirection = GetSplineStartDirection(Orientation);
-		const FVector2D EndDirection = GetSplineEndDirection(Orientation);
+		const FVector2D StartDirection = GetSplineStartDirection(Orientation, bIsSourceBusConnection);
+		const FVector2D EndDirection = GetSplineEndDirection(Orientation, bIsSourceBusConnection);
 		FVector2D LastPointPos = ConnectionStartPos;
 
 		float InterpolationAmount = 0.0f;
@@ -1562,7 +1634,7 @@ namespace UE::Audio::Insights
 			}
 		}
 
-		const FLinearColor ConnectionColor(GetConnectionColor(Entry->Amplitude.GetValue(), WireSplineAmplitudePowerFactor));
+		const FLinearColor ConnectionColor(GetConnectionColor(Entry->Amplitude.GetValue(), WireSplineAmplitudePowerFactor, bIsSourceBusConnection));
 		FSlateDrawElement::MakeLines(OutDrawElements,
 									 ++LayerId,
 									 AllottedGeometry.ToPaintGeometry(),
@@ -1762,6 +1834,11 @@ namespace UE::Audio::Insights
 				}
 			}
 
+			for (const FSignalFlowEntryKey& LinkedSoundSourceKey : Node->FilteredLinkedSoundSources)
+			{
+				NodeWidget->AddSendLabelWidget(LinkedSoundSourceKey);
+			}
+
 			MaxNumNodeOutputsInGraph = FMath::Max(MaxNumNodeOutputsInGraph, Node->FilteredOutputs.Num());
 		}
 	}
@@ -1780,9 +1857,10 @@ namespace UE::Audio::Insights
 
 			Canvas->RemoveSlot(RemovedNode->ToSharedRef());
 
-			if ((*RemovedNode)->GetEntry().IsValid())
+			const TSharedPtr<FSignalFlowDashboardEntry> RemovedEntry = (*RemovedNode)->GetEntry();
+			if (RemovedEntry.IsValid())
 			{
-				for (const FSignalFlowEntryKey& InputKey : (*RemovedNode)->GetEntry()->Inputs)
+				for (const FSignalFlowEntryKey& InputKey : RemovedEntry->Inputs)
 				{
 					const TOptional<FSignalFlowEntryKey> RealInputKey = FindRealNodeKey(InputKey, ENodeConnectionType::Input);
 					if (!RealInputKey.IsSet())
@@ -1798,6 +1876,16 @@ namespace UE::Audio::Insights
 
 					(*InputNode)->RemoveSendLabelWidget(Key);
 				}
+
+				if (RemovedEntry->LinkedSourceBus.IsSet())
+				{
+					TSharedPtr<SSignalFlowGraphNode>* SourceBusNode = NodeWidgets.Find(RemovedEntry->LinkedSourceBus.GetValue());
+
+					if (SourceBusNode && SourceBusNode->IsValid())
+					{
+						(*SourceBusNode)->RemoveSendLabelWidget(Key);
+					}
+				}
 			}
 
 			NodeWidgets.Remove(Key);
@@ -1966,6 +2054,11 @@ namespace UE::Audio::Insights
 					}
 				}
 
+				for (const FSignalFlowEntryKey& LinkedSoundSourceKey : Node->FilteredLinkedSoundSources)
+				{
+					ActiveOutputKeys.AddUnique(LinkedSoundSourceKey);
+				}
+
 				(*NodeWidget)->SetSendLabelWidgets(ActiveOutputKeys);
 
 				return;
diff --git a/Engine/Plugins/AudioInsights/Source/AudioInsights/Private/Views/SignalFlowDashboardViewFactory.cpp b/Engine/Plugins/AudioInsights/Source/AudioInsights/Private/Views/SignalFlowDashboardViewFactory.cpp
--- a/Engine/Plugins/AudioInsights/Source/AudioInsights/Private/Views/SignalFlowDashboardViewFactory.cpp
+++ b/Engine/Plugins/AudioInsights/Source/AudioInsights/Private/Views/SignalFlowDashboardViewFactory.cpp
@@ -1552,7 +1552,13 @@ namespace UE::Audio::Insights
 			FilteredEntries.Add(Entry->GetSignalFlowEntryKey(), EntryNode);
 		}
 
-		if (EntryNode.IsValid() && ChildEntry.IsValid())
+		if (!ensure(EntryNode.IsValid()))
+		{
+			// We should never hit this, but return just in case to avoid a crash
+			return;
+		}
+
+		if (ChildEntry.IsValid())
 		{
 			EntryNode->FilteredOutputs.AddUnique(ChildEntry->GetSignalFlowEntryKey());
 		}
@@ -1563,15 +1569,14 @@ namespace UE::Audio::Insights
 
 			if (InputEntry && InputEntry->IsValid())
 			{
-				if (EntryNode.IsValid())
-				{
-					EntryNode->FilteredInputs.AddUnique((*InputEntry)->GetSignalFlowEntryKey());
-				}
+				EntryNode->FilteredInputs.AddUnique((*InputEntry)->GetSignalFlowEntryKey());
 
 				bool bResetTreeDepth = (*InputEntry)->EntryType != Entry->EntryType;
 				AddFilteredInputsRecursive(DeviceData, *InputEntry, Entry, bResetTreeDepth ? 0 : TreeDepth - 1, OutExistingEntryKeys);
 			}
 		}
+
+		AddFilteredLinkedSourceBus(DeviceData, Entry, *EntryNode, OutExistingEntryKeys);
 	}
 
 	void FSignalFlowDashboardViewFactory::AddFilteredOutputs(const TSortedMap<FSignalFlowEntryKey, TSharedPtr<FSignalFlowDashboardEntry>>& DeviceData, TSharedPtr<FSignalFlowDashboardEntry> Entry, int32 TreeDepth, TSet<FSignalFlowEntryKey>& OutExistingEntryKeys)
@@ -1591,6 +1596,12 @@ namespace UE::Audio::Insights
 				}
 			}
 		}
+
+		TSharedPtr<ISignalFlowNode>* EntryNode = FilteredEntries.Find(Entry->GetSignalFlowEntryKey());
+		if (EntryNode && EntryNode->IsValid())
+		{
+			AddFilteredLinkedSources(DeviceData, Entry, **EntryNode, OutExistingEntryKeys);
+		}
 	}
 
 	void FSignalFlowDashboardViewFactory::AddFilteredOutputsRecursive(const TSortedMap<FSignalFlowEntryKey, TSharedPtr<FSignalFlowDashboardEntry>>& DeviceData, TSharedPtr<FSignalFlowDashboardEntry> Entry, TSharedPtr<FSignalFlowDashboardEntry> ParentEntry, int32 TreeDepth, TSet<FSignalFlowEntryKey>& OutExistingEntryKeys)
@@ -1635,7 +1646,13 @@ namespace UE::Audio::Insights
 			FilteredEntries.Add(EntryKey, Node);
 		}
 
-		if (Node.IsValid() && ParentEntry.IsValid())
+		if (!ensure(Node.IsValid()))
+		{
+			// We should never hit this, but return just in case to avoid a crash
+			return;
+		}
+
+		if (ParentEntry.IsValid())
 		{
 			Node->FilteredInputs.AddUnique(ParentEntry->GetSignalFlowEntryKey());
 		}
@@ -1645,14 +1662,93 @@ namespace UE::Audio::Insights
 			const TSharedPtr<FSignalFlowDashboardEntry>* OutputEntry = DeviceData.Find(OutputKey);
 			if (OutputEntry && OutputEntry->IsValid())
 			{
-				if (Node.IsValid())
+				Node->FilteredOutputs.AddUnique(OutputKey);
+
+				bool bResetTreeDepth = (*OutputEntry)->EntryType != Entry->EntryType;
+				AddFilteredOutputsRecursive(DeviceData, *OutputEntry, Entry, bResetTreeDepth ? 0 : TreeDepth + 1, OutExistingEntryKeys);
+			}
+		}
+
+		AddFilteredLinkedSources(DeviceData, Entry, *Node, OutExistingEntryKeys);
+	}
+
+	void FSignalFlowDashboardViewFactory::AddFilteredLinkedSourceBus(const TSortedMap<FSignalFlowEntryKey, TSharedPtr<FSignalFlowDashboardEntry>>& DeviceData, TSharedPtr<FSignalFlowDashboardEntry> Entry, ISignalFlowNode& OutEntryNode, TSet<FSignalFlowEntryKey>& OutExistingEntryKeys)
+	{
+		if (!Entry.IsValid() || !Entry->LinkedSourceBus.IsSet())
+		{
+			return;
+		}
+
+		const FSignalFlowEntryKey SourceBusEntryKey = Entry->LinkedSourceBus.GetValue();
+		OutEntryNode.FilteredLinkedSourceBus = SourceBusEntryKey;
+
+		// Check if we need to create a new path inside the graph from the linked source bus entry
+		if (!EntryIsInGraph(SourceBusEntryKey, OutExistingEntryKeys))
+		{
+			const TSharedPtr<FSignalFlowDashboardEntry>* SourceBusEntry = DeviceData.Find(SourceBusEntryKey);
+
+			if (SourceBusEntry && SourceBusEntry->IsValid())
+			{
+				// Detected a new source bus connected to this entry
+				// Filter the graph from the found source bus to make sure it's full path is present in the graph
+				AddFilteredInputsRecursive(DeviceData, *SourceBusEntry, nullptr, 0, OutExistingEntryKeys);
+
+				for (const auto& [OutputKey, SendInfo] : (*SourceBusEntry)->Outputs)
 				{
-					Node->FilteredOutputs.AddUnique(OutputKey);
+					const TSharedPtr<FSignalFlowDashboardEntry>* OutputEntry = DeviceData.Find(OutputKey);
+					if (OutputEntry && OutputEntry->IsValid())
+					{
+						AddFilteredOutputsRecursive(DeviceData, *OutputEntry, *SourceBusEntry, 0, OutExistingEntryKeys);
+					}
 				}
+			}
+		}
 
-				bool bResetTreeDepth = (*OutputEntry)->EntryType != Entry->EntryType;
+		// Make sure the source bus entry is connected to this entry
+		const TSharedPtr<ISignalFlowNode>* LinkedSourceBusNode = FilteredEntries.Find(SourceBusEntryKey);
+		if (LinkedSourceBusNode && LinkedSourceBusNode->IsValid())
+		{
+			(*LinkedSourceBusNode)->FilteredLinkedSoundSources.AddUnique(Entry->GetSignalFlowEntryKey());
+		}
+	}
 
-				AddFilteredOutputsRecursive(DeviceData, *OutputEntry, Entry, bResetTreeDepth ? 0 : TreeDepth + 1, OutExistingEntryKeys);
+	void FSignalFlowDashboardViewFactory::AddFilteredLinkedSources(const TSortedMap<FSignalFlowEntryKey, TSharedPtr<FSignalFlowDashboardEntry>>& DeviceData, TSharedPtr<FSignalFlowDashboardEntry> Entry, ISignalFlowNode& OutEntryNode, TSet<FSignalFlowEntryKey>& OutExistingEntryKeys)
+	{
+		if (!Entry.IsValid() || Entry->LinkedSoundSources.IsEmpty())
+		{
+			return;
+		}
+
+		// Cycle through this source bus's linked sound sources, detect any new paths that need adding to the filtered entries
+		for (const FSignalFlowEntryKey& LinkedSourceEntryKey : Entry->LinkedSoundSources)
+		{
+			const TSharedPtr<FSignalFlowDashboardEntry>* LinkedSourceEntry = DeviceData.Find(LinkedSourceEntryKey);
+			if (LinkedSourceEntry && LinkedSourceEntry->IsValid())
+			{
+				OutEntryNode.FilteredLinkedSoundSources.AddUnique(LinkedSourceEntryKey);
+
+				if (!EntryIsInGraph(LinkedSourceEntryKey, OutExistingEntryKeys))
+				{
+					// Detected a new sound source connected to this source bus entry
+					// Filter the graph from the found sound source to make sure it's full path is present in the graph
+					AddFilteredOutputsRecursive(DeviceData, *LinkedSourceEntry, nullptr, 0, OutExistingEntryKeys);
+
+					for (const FSignalFlowEntryKey& InputKey : (*LinkedSourceEntry)->Inputs)
+					{
+						const TSharedPtr<FSignalFlowDashboardEntry>* InputEntry = DeviceData.Find(InputKey);
+						if (InputEntry && InputEntry->IsValid())
+						{
+							AddFilteredInputsRecursive(DeviceData, *InputEntry, *LinkedSourceEntry, 0, OutExistingEntryKeys);
+						}
+					}
+				}
+
+				// Make sure the linked sound source is connected to this source bus entry
+				const TSharedPtr<ISignalFlowNode>* LinkedSourceNode = FilteredEntries.Find(LinkedSourceEntryKey);
+				if (LinkedSourceNode && LinkedSourceNode->IsValid())
+				{
+					(*LinkedSourceNode)->FilteredLinkedSourceBus = Entry->GetSignalFlowEntryKey();
+				}
 			}
 		}
 	}
@@ -1673,6 +1769,12 @@ namespace UE::Audio::Insights
 
 			Entry->FilteredInputs.RemoveAll([&EntriesToRemove](const FSignalFlowEntryKey& Input) { return EntriesToRemove.Contains(Input); });
 			Entry->FilteredOutputs.RemoveAll([&EntriesToRemove](const FSignalFlowEntryKey& Output) { return EntriesToRemove.Contains(Output); });
+
+			if (Entry->FilteredLinkedSourceBus.IsSet() && EntriesToRemove.Contains(Entry->FilteredLinkedSourceBus.GetValue()))
+			{
+				Entry->FilteredLinkedSourceBus.Reset();
+			}
+			Entry->FilteredLinkedSoundSources.RemoveAll([&EntriesToRemove](const FSignalFlowEntryKey& LinkedSource) { return EntriesToRemove.Contains(LinkedSource); });
 		}
 
 		for (const FSignalFlowEntryKey& EntryKey : EntriesToRemove)
diff --git a/Engine/Plugins/AudioInsights/Source/AudioInsights/Public/Views/SSignalFlowGraph.h b/Engine/Plugins/AudioInsights/Source/AudioInsights/Public/Views/SSignalFlowGraph.h
--- a/Engine/Plugins/AudioInsights/Source/AudioInsights/Public/Views/SSignalFlowGraph.h
+++ b/Engine/Plugins/AudioInsights/Source/AudioInsights/Public/Views/SSignalFlowGraph.h
@@ -272,13 +272,26 @@ namespace UE::Audio::Insights
 	private:
 		int32 OnPaintConnections(const FGeometry& AllottedGeometry, FSlateWindowElementList& OutDrawElements, int32 LayerId) const;
 
+		int32 PaintConnectionToNode(const FGeometry& AllottedGeometry,
+									FSlateWindowElementList& OutDrawElements,
+									int32 LayerId,
+									const TSharedPtr<FSignalFlowDashboardEntry>& SenderEntry,
+									const SSignalFlowGraphNode& SenderNode,
+									const SSignalFlowGraphNode& ReceiverNode,
+									const FBox2D& LocalBoundsBox,
+									const FVector2D& ViewportPosInGraph,
+									const float SendLevel,
+									const int32 SendLabelIndex,
+									const bool bIsSourceBusConnection) const;
+
 		int32 AnimateWire(const FGeometry& AllottedGeometry,
 							FSlateWindowElementList& OutDrawElements,
 							int32 LayerId,
 							const TSharedPtr<FSignalFlowDashboardEntry>& Entry,
 							const FVector2D& ConnectionStartPos,
 							const FVector2D& ConnectionEndPos,
-							const float SendLevel) const;
+							const float SendLevel,
+							const bool bIsSourceBusConnection) const;
 
 		void CollectAmplitudeConnectionPoints(const float Amplitude,
 												const FVector2D& ConnectionStartPos,
diff --git a/Engine/Plugins/AudioInsights/Source/AudioInsights/Public/Views/SignalFlowDashboardViewFactory.h b/Engine/Plugins/AudioInsights/Source/AudioInsights/Public/Views/SignalFlowDashboardViewFactory.h
--- a/Engine/Plugins/AudioInsights/Source/AudioInsights/Public/Views/SignalFlowDashboardViewFactory.h
+++ b/Engine/Plugins/AudioInsights/Source/AudioInsights/Public/Views/SignalFlowDashboardViewFactory.h
@@ -118,6 +118,9 @@ namespace UE::Audio::Insights
 		void AddFilteredInputsRecursive(const TSortedMap<FSignalFlowEntryKey, TSharedPtr<FSignalFlowDashboardEntry>>& DeviceData, TSharedPtr<FSignalFlowDashboardEntry> Entry, TSharedPtr<FSignalFlowDashboardEntry> ChildEntry, int32 TreeDepth, TSet<FSignalFlowEntryKey>& OutExistingEntryKeys);
 		void AddFilteredOutputs(const TSortedMap<FSignalFlowEntryKey, TSharedPtr<FSignalFlowDashboardEntry>>& DeviceData, TSharedPtr<FSignalFlowDashboardEntry> Entry, int32 TreeDepth, TSet<FSignalFlowEntryKey>& OutExistingEntryKeys);
 		void AddFilteredOutputsRecursive(const TSortedMap<FSignalFlowEntryKey, TSharedPtr<FSignalFlowDashboardEntry>>& DeviceData, TSharedPtr<FSignalFlowDashboardEntry> Entry, TSharedPtr<FSignalFlowDashboardEntry> ParentEntry, int32 TreeDepth, TSet<FSignalFlowEntryKey>& OutExistingEntryKeys);
+		void AddFilteredLinkedSourceBus(const TSortedMap<FSignalFlowEntryKey, TSharedPtr<FSignalFlowDashboardEntry>>& DeviceData, TSharedPtr<FSignalFlowDashboardEntry> Entry, ISignalFlowNode& OutEntryNode, TSet<FSignalFlowEntryKey>& OutExistingEntryKeys);
+		void AddFilteredLinkedSources(const TSortedMap<FSignalFlowEntryKey, TSharedPtr<FSignalFlowDashboardEntry>>& DeviceData, TSharedPtr<FSignalFlowDashboardEntry> Entry, ISignalFlowNode& OutEntryNode, TSet<FSignalFlowEntryKey>& OutExistingEntryKeys);
+		
 		void CleanUpFilteredEntries(const TSet<FSignalFlowEntryKey>& EntriesToRemove);
 
 		void AddMissingNodesToGraph(const TSortedMap<FSignalFlowEntryKey, TSharedPtr<FSignalFlowDashboardEntry>>& DeviceData, TSet<FSignalFlowEntryKey>& OutEntriesMarkedToDelete);
diff --git a/Engine/Plugins/AudioInsights/Source/AudioInsights/Public/Views/SignalFlowNodes.h b/Engine/Plugins/AudioInsights/Source/AudioInsights/Public/Views/SignalFlowNodes.h
--- a/Engine/Plugins/AudioInsights/Source/AudioInsights/Public/Views/SignalFlowNodes.h
+++ b/Engine/Plugins/AudioInsights/Source/AudioInsights/Public/Views/SignalFlowNodes.h
@@ -93,6 +93,9 @@ namespace UE::Audio::Insights
 		TArray<FSignalFlowEntryKey> FilteredInputs;
 		TArray<FSignalFlowEntryKey> FilteredOutputs;
 
+		TArray<FSignalFlowEntryKey> FilteredLinkedSoundSources;
+		TOptional<FSignalFlowEntryKey> FilteredLinkedSourceBus;
+
 		float XPos = 0.0f;
 
 	private:
