[Audio Insights] Add Owner Object selectable list to Signal Flow
* Original previz only had selectable rows for sources, buses and submixes, but neglected owner objects.
* Now we have the ability to collapse/expand these selectable lists, we can support more active entry types in the left hand side without causing too much fighting for screen space by having these new lists collapsed by default.

#Synced-CL 51626790

diff --git a/Engine/Plugins/AudioInsights/Source/AudioInsights/Private/Views/SignalFlowDashboardViewFactory.cpp b/Engine/Plugins/AudioInsights/Source/AudioInsights/Private/Views/SignalFlowDashboardViewFactory.cpp
--- a/Engine/Plugins/AudioInsights/Source/AudioInsights/Private/Views/SignalFlowDashboardViewFactory.cpp
+++ b/Engine/Plugins/AudioInsights/Source/AudioInsights/Private/Views/SignalFlowDashboardViewFactory.cpp
@@ -35,6 +35,7 @@ namespace UE::Audio::Insights
 		const FName ButtonTextStyleName = "SmallButtonText";
 
 		constexpr float GraphWidthProp = 0.8f;
+		constexpr float OwnerObjectListHeightProp = 1.0f;
 		constexpr float SourcesListHeightProp = 2.5f;
 		constexpr float BusListHeightProp = 1.0f;
 		constexpr float SubmixesListHeightProp = 1.5f;
@@ -886,6 +887,7 @@ namespace UE::Audio::Insights
 	{
 		using namespace FSignalFlowDashboardViewFactoryPrivate;
 
+		ActiveOwnerObjectsListView = CreateSelectableList(LOCTEXT("AudioDashboard_SignalFlow_ActiveOwnerObjectsTitle", "Owner Objects"), ActiveOwnerObjects, ESignalFlowEntryType::OwnerObject);
 		ActiveSourcesListView = CreateSelectableList(LOCTEXT("AudioDashboard_SignalFlow_ActiveSourcesTitle", "Active Sources"), ActiveSources, ESignalFlowEntryType::SoundSource);
 		ActiveBusesListView = CreateSelectableList(LOCTEXT("AudioDashboard_SignalFlow_ActiveBusesTitle", "Active Buses"), ActiveBuses, ESignalFlowEntryType::AudioBus);
 		ActiveSubmixesListView = CreateSelectableList(LOCTEXT("AudioDashboard_SignalFlow_ActiveSubmixesTitle", "Active Submixes"), ActiveSubmixes, ESignalFlowEntryType::Submix);
@@ -916,6 +918,7 @@ namespace UE::Audio::Insights
 			OutSlotIndex = ActiveEntryMenusContainer->NumSlots() - 1;
 		};
 
+		AddSelectableListToSlot(ActiveOwnerObjectsListView.ToSharedRef(), OwnerObjectListHeightProp, ActiveOwnerObjectsSlotIndex);
 		AddSelectableListToSlot(ActiveSourcesListView.ToSharedRef(), SourcesListHeightProp, ActiveSourcesSlotIndex);
 		AddSelectableListToSlot(ActiveBusesListView.ToSharedRef(), BusListHeightProp, ActiveBusesSlotIndex);
 		AddSelectableListToSlot(ActiveSubmixesListView.ToSharedRef(), SubmixesListHeightProp, ActiveSubmixesSlotIndex);
@@ -980,6 +983,13 @@ namespace UE::Audio::Insights
 							{
 								switch (SelectedNodeEntry->EntryType)
 								{
+									case ESignalFlowEntryType::OwnerObject:
+										if (ActiveOwnerObjectsListView.IsValid())
+										{
+											ActiveOwnerObjectsListView->SetSelection(SelectedNodeEntry);
+										}
+										break;
+
 									case ESignalFlowEntryType::SoundSource:
 										if (ActiveSourcesListView.IsValid())
 										{
@@ -1002,6 +1012,11 @@ namespace UE::Audio::Insights
 										break;
 
 									default:
+										if (ActiveOwnerObjectsListView.IsValid())
+										{
+											ActiveOwnerObjectsListView->ClearSelection();
+										}
+
 										if (ActiveSourcesListView.IsValid())
 										{
 											ActiveSourcesListView->ClearSelection();
@@ -1142,6 +1157,9 @@ namespace UE::Audio::Insights
 	{
 		switch (EntryType)
 		{
+			case ESignalFlowEntryType::OwnerObject:
+				return ActiveEntryMenuExpansionSettings.bShowActiveOwnerObjects;
+
 			case ESignalFlowEntryType::SoundSource:
 				return ActiveEntryMenuExpansionSettings.bShowActiveSources;
 
@@ -1165,6 +1183,16 @@ namespace UE::Audio::Insights
 
 		switch (EntryType)
 		{
+			case ESignalFlowEntryType::OwnerObject:
+				ActiveEntryMenuExpansionSettings.bShowActiveOwnerObjects = bIsExpanded;
+
+				if (ActiveEntryMenusContainer.IsValid() && ensure(ActiveEntryMenusContainer->IsValidSlotIndex(ActiveOwnerObjectsSlotIndex)))
+				{
+					SetSelectableListSlotSize(ActiveEntryMenusContainer->GetSlot(ActiveOwnerObjectsSlotIndex), OwnerObjectListHeightProp, bIsExpanded);
+				}
+
+				break;
+
 			case ESignalFlowEntryType::SoundSource:
 				ActiveEntryMenuExpansionSettings.bShowActiveSources = bIsExpanded;
 
@@ -1238,7 +1266,7 @@ namespace UE::Audio::Insights
 
 	void FSignalFlowDashboardViewFactory::SetNewSelection(const TSharedPtr<FSignalFlowDashboardEntry>& InSelectedItem)
 	{
-		if (!ActiveSourcesListView.IsValid() || !ActiveBusesListView.IsValid() || !ActiveSubmixesListView.IsValid())
+		if (!ActiveOwnerObjectsListView.IsValid() || !ActiveSourcesListView.IsValid() || !ActiveBusesListView.IsValid() || !ActiveSubmixesListView.IsValid())
 		{
 #if WITH_EDITOR
 			AssetContextMenuHelper.ResetAssetEntry();
@@ -1253,20 +1281,30 @@ namespace UE::Audio::Insights
 		{
 			switch (InSelectedItem->EntryType)
 			{
+				case ESignalFlowEntryType::OwnerObject:
+				{
+					ActiveSourcesListView->ClearSelection();
+					ActiveBusesListView->ClearSelection();
+					ActiveSubmixesListView->ClearSelection();
+					break;
+				}
 				case ESignalFlowEntryType::SoundSource:
 				{
+					ActiveOwnerObjectsListView->ClearSelection();
 					ActiveBusesListView->ClearSelection();
 					ActiveSubmixesListView->ClearSelection();
 					break;
 				}
 				case ESignalFlowEntryType::AudioBus:
 				{
+					ActiveOwnerObjectsListView->ClearSelection();
 					ActiveSourcesListView->ClearSelection();
 					ActiveSubmixesListView->ClearSelection();
 					break;
 				}
 				case ESignalFlowEntryType::Submix:
 				{
+					ActiveOwnerObjectsListView->ClearSelection();
 					ActiveSourcesListView->ClearSelection();
 					ActiveBusesListView->ClearSelection();
 					break;
@@ -1295,6 +1333,7 @@ namespace UE::Audio::Insights
 #endif // WITH_EDITOR
 
 			SelectedItem.Reset();
+			ActiveOwnerObjectsListView->ClearSelection();
 			ActiveSourcesListView->ClearSelection();
 			ActiveBusesListView->ClearSelection();
 			ActiveSubmixesListView->ClearSelection();
@@ -1324,6 +1363,7 @@ namespace UE::Audio::Insights
 
 	void FSignalFlowDashboardViewFactory::BuildSelectionMenus(const TSortedMap<FSignalFlowEntryKey, TSharedPtr<FSignalFlowDashboardEntry>>& DeviceData, const EProcessReason Reason)
 	{
+		ActiveOwnerObjects.Reset();
 		ActiveSources.Reset();
 		ActiveBuses.Reset();
 		ActiveSubmixes.Reset();
@@ -1345,6 +1385,9 @@ namespace UE::Audio::Insights
 
 			switch (Entry->EntryType)
 			{
+				case ESignalFlowEntryType::OwnerObject:
+					ActiveOwnerObjects.Add(Entry);
+					break;
 				case ESignalFlowEntryType::SoundSource:
 					ActiveSources.Add(Entry);
 					break;
@@ -1362,6 +1405,11 @@ namespace UE::Audio::Insights
 
 	void FSignalFlowDashboardViewFactory::RefreshFilteredEntriesListView()
 	{
+		if (ActiveOwnerObjectsListView.IsValid())
+		{
+			ActiveOwnerObjectsListView->RequestListRefresh();
+		}
+
 		if (ActiveSourcesListView.IsValid())
 		{
 			ActiveSourcesListView->RequestListRefresh();
@@ -2553,6 +2601,7 @@ namespace UE::Audio::Insights
 
 		ActiveEntryMenuExpansionSettings = InSettings.ActiveEntryMenuExpansionSettings;
 
+		SetSelectableListIsExpanded(ESignalFlowEntryType::OwnerObject, ActiveEntryMenuExpansionSettings.bShowActiveOwnerObjects);
 		SetSelectableListIsExpanded(ESignalFlowEntryType::SoundSource, ActiveEntryMenuExpansionSettings.bShowActiveSources);
 		SetSelectableListIsExpanded(ESignalFlowEntryType::AudioBus, ActiveEntryMenuExpansionSettings.bShowActiveBuses);
 		SetSelectableListIsExpanded(ESignalFlowEntryType::Submix, ActiveEntryMenuExpansionSettings.bShowActiveSubmixes);
diff --git a/Engine/Plugins/AudioInsights/Source/AudioInsights/Public/Settings/SignalFlowSettings.h b/Engine/Plugins/AudioInsights/Source/AudioInsights/Public/Settings/SignalFlowSettings.h
--- a/Engine/Plugins/AudioInsights/Source/AudioInsights/Public/Settings/SignalFlowSettings.h
+++ b/Engine/Plugins/AudioInsights/Source/AudioInsights/Public/Settings/SignalFlowSettings.h
@@ -84,6 +84,10 @@ struct FSignalFlowEntryTypeListExpansionSettings
 {
 	GENERATED_BODY()
 
+	/** Toggles expansion of the left hand Active Owner Objects menu. */
+	UPROPERTY(EditAnywhere, config, Category = SignalFlow)
+	bool bShowActiveOwnerObjects = false;
+
 	/** Toggles expansion of the left hand Active Sources menu. */
 	UPROPERTY(EditAnywhere, config, Category = SignalFlow)
 	bool bShowActiveSources = true;
diff --git a/Engine/Plugins/AudioInsights/Source/AudioInsights/Public/Views/SignalFlowDashboardViewFactory.h b/Engine/Plugins/AudioInsights/Source/AudioInsights/Public/Views/SignalFlowDashboardViewFactory.h
--- a/Engine/Plugins/AudioInsights/Source/AudioInsights/Public/Views/SignalFlowDashboardViewFactory.h
+++ b/Engine/Plugins/AudioInsights/Source/AudioInsights/Public/Views/SignalFlowDashboardViewFactory.h
@@ -164,6 +164,10 @@ namespace UE::Audio::Insights
 
 		TSharedPtr<SVerticalBox> ActiveEntryMenusContainer;
 
+		TSharedPtr<SListView<TSharedPtr<FSignalFlowDashboardEntry>>> ActiveOwnerObjectsListView;
+		TArray<TSharedPtr<FSignalFlowDashboardEntry>> ActiveOwnerObjects;
+		int32 ActiveOwnerObjectsSlotIndex = INDEX_NONE;
+
 		TSharedPtr<SListView<TSharedPtr<FSignalFlowDashboardEntry>>> ActiveSourcesListView;
 		TArray<TSharedPtr<FSignalFlowDashboardEntry>> ActiveSources;
 		int32 ActiveSourcesSlotIndex = INDEX_NONE;
