[AudioInsights] - Toolbar: Move CacheStats out of event log into Toolbar

#Synced-CL 51664040

diff --git a/Engine/Plugins/AudioInsights/Source/AudioInsights/Private/Views/AudioEventLogDashboardViewFactory.cpp b/Engine/Plugins/AudioInsights/Source/AudioInsights/Private/Views/AudioEventLogDashboardViewFactory.cpp
--- a/Engine/Plugins/AudioInsights/Source/AudioInsights/Private/Views/AudioEventLogDashboardViewFactory.cpp
+++ b/Engine/Plugins/AudioInsights/Source/AudioInsights/Private/Views/AudioEventLogDashboardViewFactory.cpp
@@ -686,62 +686,6 @@ namespace UE::Audio::Insights
 				
 			]
 
-			// Cache stats
-			+ SHorizontalBox::Slot()
-			.Padding(4.0f, 0.0f, 0.0f, 0.0f)
-			.VAlign(VAlign_Center)
-			.AutoWidth()
-			[
-				SNew(STextBlock)
-				.Justification(ETextJustify::Center)
-				.ColorAndOpacity(StatTitleColor)
-				.Text(LOCTEXT("AudioEventLog_CacheStats_Title", "Cache Size: "))
-			]
-
-			+ SHorizontalBox::Slot()
-			.VAlign(VAlign_Center)
-			.AutoWidth()
-			[
-				SNew(STextBlock)
-				.Justification(ETextJustify::Center)
-				.Text_Lambda([]()
-				{
-					FAudioInsightsCacheManager& CacheManager = FAudioInsightsModule::GetChecked().GetCacheManager();
-					const float CacheSizeMb = static_cast<float>(CacheManager.GetUsedCacheSize()) / 1000000.0f;
-					const float MaxCacheSizeMb = static_cast<float>(CacheManager.GetMaxCacheSize()) / 1000000.0f;
-
-					return FText::Format(LOCTEXT("AudioEventLog_CacheStats_Memory", "{0} / {1}")
-										, FText::AsNumber(CacheSizeMb, FSlateStyle::Get().GetMemoryFormat())
-										, FSlateStyle::Get().FormatMemoryAsMegabytes(MaxCacheSizeMb));
-				})
-			]
-
-			+ SHorizontalBox::Slot()
-			.VAlign(VAlign_Center)
-			.AutoWidth()
-			[
-				SNew(STextBlock)
-				.Justification(ETextJustify::Center)
-				.ColorAndOpacity(StatTitleColor)
-				.Text(LOCTEXT("AudioEventLog_CacheStats_Duration", " : Duration: "))
-			]
-
-			+ SHorizontalBox::Slot()
-			.VAlign(VAlign_Center)
-			.AutoWidth()
-			[
-				SNew(STextBlock)
-				.Justification(ETextJustify::Center)
-				.Text_Lambda([]()
-				{
-					FAudioInsightsCacheManager& CacheManager = FAudioInsightsModule::GetChecked().GetCacheManager();
-					const float CacheDuration = CacheManager.GetCacheDuration();
-
-					return FText::Format(LOCTEXT("AudioEventLog_CacheStats_DurationValue", "{0}s")
-										, FText::AsNumber(CacheDuration, FSlateStyle::Get().GetShortTimeFormat()));
-				})
-			]
-
 			// Clear log Button
 			+ SHorizontalBox::Slot()
 			.HAlign(HAlign_Left)
diff --git a/Engine/Plugins/AudioInsights/Source/AudioInsightsEditor/Private/AudioInsightsEditorDashboardFactory.cpp b/Engine/Plugins/AudioInsights/Source/AudioInsightsEditor/Private/AudioInsightsEditorDashboardFactory.cpp
--- a/Engine/Plugins/AudioInsights/Source/AudioInsightsEditor/Private/AudioInsightsEditorDashboardFactory.cpp
+++ b/Engine/Plugins/AudioInsights/Source/AudioInsightsEditor/Private/AudioInsightsEditorDashboardFactory.cpp
@@ -568,6 +568,8 @@ namespace UE::Audio::Insights
 		BackgroundColorBrush.TintColor = FSlateColor(FLinearColor(0.018f, 0.018f, 0.018f, 1.0f));
 		BackgroundColorBrush.DrawAs = ESlateBrushDrawType::Box;
 
+		constexpr FLinearColor CacheStatsColor(1.0f, 1.0f, 1.0f, 0.5f);
+
 		return SNew(SBorder)
 			.BorderImage(&BackgroundColorBrush)
 			[
@@ -766,6 +768,63 @@ namespace UE::Audio::Insights
 						.Thickness(2.0f)
 					]
 				]
+
+				// Cache stats
+				+ SHorizontalBox::Slot()
+				.Padding(12.0f, 0.0f, 0.0f, 0.0f)
+				.VAlign(VAlign_Center)
+				.AutoWidth()
+				[
+					SNew(STextBlock)
+					.Justification(ETextJustify::Center)
+					.ColorAndOpacity(CacheStatsColor)
+					.Text(LOCTEXT("AudioDashboardEditor_CacheStats_Title", "Cache Size: "))
+				]
+
+				+ SHorizontalBox::Slot()
+				.VAlign(VAlign_Center)
+				.AutoWidth()
+				[
+					SNew(STextBlock)
+					.Justification(ETextJustify::Center)
+					.Text_Lambda([]()
+					{
+						FAudioInsightsCacheManager& CacheManager = FAudioInsightsEditorModule::GetChecked().GetCacheManager();
+						const float CacheSizeMb = static_cast<float>(CacheManager.GetUsedCacheSize()) / 1000000.0f;
+						const float MaxCacheSizeMb = static_cast<float>(CacheManager.GetMaxCacheSize()) / 1000000.0f;
+
+						return FText::Format(LOCTEXT("AudioDashboardEditor_CacheStats_Memory", "{0} / {1}")
+											, FText::AsNumber(CacheSizeMb, FSlateStyle::Get().GetMemoryFormat())
+											, FSlateStyle::Get().FormatMemoryAsMegabytes(MaxCacheSizeMb));
+					})
+				]
+
+				+ SHorizontalBox::Slot()
+				.VAlign(VAlign_Center)
+				.AutoWidth()
+				[
+					SNew(STextBlock)
+					.Justification(ETextJustify::Center)
+					.ColorAndOpacity(CacheStatsColor)
+					.Text(LOCTEXT("AudioDashboardEditor_CacheStats_Duration", " : Duration: "))
+				]
+
+				+ SHorizontalBox::Slot()
+				.VAlign(VAlign_Center)
+				.AutoWidth()
+				[
+					SNew(STextBlock)
+					.Justification(ETextJustify::Center)
+					.Text_Lambda([]()
+					{
+						FAudioInsightsCacheManager& CacheManager = FAudioInsightsEditorModule::GetChecked().GetCacheManager();
+						const float CacheDuration = CacheManager.GetCacheDuration();
+
+						return FText::Format(LOCTEXT("AudioDashboardEditor_CacheStats_DurationValue", "{0}s")
+											, FText::AsNumber(CacheDuration, FSlateStyle::Get().GetShortTimeFormat()));
+					})
+				]
+
 				// World filter label
 				+ SHorizontalBox::Slot()
 				.HAlign(HAlign_Right)
