[Audio Insights] Handle selected sounds being destroyed in Signal Flow
* When an item is removed from the device data in the provider, we can end up with a stale reference inside the dashboard in the SelectedItem.
* Currently, the graph will all disappear, apart from the stale selected item, which will show as a signle lone node.
* Now, when we detect this happening, clear the selection and refresh the whole graph

#Synced-CL 51626790

diff --git a/Engine/Plugins/AudioInsights/Source/AudioInsights/Private/Views/SignalFlowDashboardViewFactory.cpp b/Engine/Plugins/AudioInsights/Source/AudioInsights/Private/Views/SignalFlowDashboardViewFactory.cpp
--- a/Engine/Plugins/AudioInsights/Source/AudioInsights/Private/Views/SignalFlowDashboardViewFactory.cpp
+++ b/Engine/Plugins/AudioInsights/Source/AudioInsights/Private/Views/SignalFlowDashboardViewFactory.cpp
@@ -1447,10 +1447,28 @@ namespace UE::Audio::Insights
 			return;
 		}
 
+		const FSignalFlowTraceProvider::FDeviceData* DeviceData = Provider->FindFilteredDeviceData();
+		if (DeviceData == nullptr)
+		{
+			return;
+		}
+
 		if (SelectedItem.IsValid())
 		{
-			XPosFocusNodeEntryKey = SelectedItem->GetSignalFlowEntryKey();
-			CreateFilteredNodeGraph(*Provider, SelectedItem, false /*bAddAllNodes*/);
+			const FSignalFlowEntryKey SelectedEntryKey = SelectedItem->GetSignalFlowEntryKey();
+
+			// Check that the selected item still exists in the device data - if not we need to clear the selection first
+			if (DeviceData->Contains(SelectedEntryKey))
+			{
+				// Selection is valid, begin filtering the graph
+				XPosFocusNodeEntryKey = SelectedEntryKey;
+				CreateFilteredNodeGraph(*Provider, SelectedItem, false /*bAddAllNodes*/);
+			}
+			else
+			{
+				// Reset the selection - this will refresh the graph again
+				SetNewSelection(nullptr);
+			}
 		}
 		else
 		{
