[Audio Insights] Solo/Mute sound sources in signal flow
* Enable muting/soloing sound source nodes in the signal flow graph
* Only available in editor and with audio debug enabled (MUTE_SOLO_ENABLED)
* Only available on sound sources (IsMuteSoloAvailableForNode())
* Deduce and cache the asset name from the path and cache as both a string and FName. This is because AudioDebugger uses both formats when querying mute/solo states.
* Use the same style of button we use in the sounds tab. When the node is neither muted or soloed, we show the icon at 25% opacity (instead of transparent) so the icon is still visible
* Struct FMuteSoloStates used to simplify querying code as audio debugger returns both states when we check

#Synced-CL 51626790

diff --git a/Engine/Plugins/AudioInsights/Source/AudioInsights/Private/Views/SSignalFlowGraph.cpp b/Engine/Plugins/AudioInsights/Source/AudioInsights/Private/Views/SSignalFlowGraph.cpp
--- a/Engine/Plugins/AudioInsights/Source/AudioInsights/Private/Views/SSignalFlowGraph.cpp
+++ b/Engine/Plugins/AudioInsights/Source/AudioInsights/Private/Views/SSignalFlowGraph.cpp
@@ -24,6 +24,10 @@
 #include "Layout/WidgetPath.h"
 #endif // WITH_EDITOR
 
+#if MUTE_SOLO_ENABLED
+#include "Audio/AudioDebug.h"
+#endif // MUTE_SOLO_ENABLED
+
 #define LOCTEXT_NAMESPACE "AudioInsights"
 
 namespace UE::Audio::Insights
@@ -305,6 +309,20 @@ namespace UE::Audio::Insights
 
 		NodeEntryKey = Entry->GetSignalFlowEntryKey();
 
+#if MUTE_SOLO_ENABLED
+		const FSoftObjectPath AssetPath(Entry->Name);
+		if (AssetPath.IsValid())
+		{
+			AssetString = AssetPath.GetAssetName();
+			AssetFName = AssetPath.GetAssetFName();
+		}
+		else
+		{
+			AssetString = Entry->Name;
+			AssetFName = FName(Entry->Name);
+		}
+#endif // MUTE_SOLO_ENABLED
+
 		ChildSlot
 		[
 			SAssignNew(SelectedNodeBorder, SBorder)
@@ -725,6 +743,60 @@ namespace UE::Audio::Insights
 				.OverflowPolicy(ETextOverflowPolicy::Ellipsis)
 				.Clipping(EWidgetClipping::ClipToBounds)
 			]
+
+#if MUTE_SOLO_ENABLED
+			+ SHorizontalBox::Slot()
+			.AutoWidth()
+			.Padding(4.0f, 0.0f, 0.0f, 0.0f)
+			[
+				SNew(SCheckBox)
+				.Style(&FSlateStyle::Get().GetWidgetStyle<FCheckBoxStyle>("SoundDashboard.MuteSoloButton"))
+				.Visibility_Lambda([this]()
+				{
+					return IsMuteSoloAvailableForNode() ? EVisibility::Visible : EVisibility::Collapsed;
+				})
+				.IsChecked(ECheckBoxState::Unchecked)
+				.OnCheckStateChanged_Lambda([this](ECheckBoxState)
+				{
+					ToggleMute();
+				})
+				[
+					SNew(SImage)
+					.Image(FSlateStyle::Get().CreateIcon("AudioInsights.Icon.SoundDashboard.Mute").GetIcon())
+					.ToolTipText(LOCTEXT("SignalFlowGraph_MuteNode", "Mute this node.\nMute will apply to all instances of an asset."))
+					.ColorAndOpacity_Lambda([this]()
+					{
+						return GetMuteSoloState().bIsMuted ? FStyleColors::Foreground : FStyleColors::White25;
+					})
+				]
+			]
+
+			+ SHorizontalBox::Slot()
+			.AutoWidth()
+			.Padding(0.0f, 0.0f, 4.0f, 0.0f)
+			[
+				SNew(SCheckBox)
+				.Style(&FSlateStyle::Get().GetWidgetStyle<FCheckBoxStyle>("SoundDashboard.MuteSoloButton"))
+				.Visibility_Lambda([this]()
+				{
+					return IsMuteSoloAvailableForNode() ? EVisibility::Visible : EVisibility::Collapsed;
+				})
+				.IsChecked(ECheckBoxState::Unchecked)
+				.OnCheckStateChanged_Lambda([this](ECheckBoxState)
+				{
+					ToggleSolo();
+				})
+				[
+					SNew(SImage)
+					.Image(FSlateStyle::Get().CreateIcon("AudioInsights.Icon.SoundDashboard.Solo").GetIcon())
+					.ToolTipText(LOCTEXT("SignalFlowGraph_SoloNode", "Solo this node.\nSolo will apply to all instances of an asset."))
+					.ColorAndOpacity_Lambda([this]()
+					{
+						return GetMuteSoloState().bIsSoloed ? FStyleColors::Foreground : FStyleColors::White25;
+					})
+				]
+			]
+#endif // MUTE_SOLO_ENABLED
 		
 			+ SHorizontalBox::Slot()
 			.AutoWidth()
@@ -974,6 +1046,44 @@ namespace UE::Audio::Insights
 		return false;
 	}
 
+#if MUTE_SOLO_ENABLED
+	void SSignalFlowGraphNode::ToggleMute()
+	{
+		if (FAudioDeviceManager* AudioDeviceManager = FAudioDeviceManager::Get())
+		{
+			AudioDeviceManager->GetDebugger().ToggleMuteSoundWave(AssetFName);
+		}
+	}
+
+	void SSignalFlowGraphNode::ToggleSolo()
+	{
+		if (FAudioDeviceManager* AudioDeviceManager = FAudioDeviceManager::Get())
+		{
+			AudioDeviceManager->GetDebugger().ToggleSoloSoundWave(AssetFName);
+		}
+	}
+
+	SSignalFlowGraphNode::FMuteSoloStates SSignalFlowGraphNode::GetMuteSoloState() const
+	{
+		FMuteSoloStates MuteSoloStates;
+		if (Entry.IsValid())
+		{
+			if (FAudioDeviceManager* AudioDeviceManager = FAudioDeviceManager::Get())
+			{
+				FString Reason;
+				AudioDeviceManager->GetDebugger().QuerySoloMuteSoundWave(AssetString, MuteSoloStates.bIsSoloed, MuteSoloStates.bIsMuted, Reason);
+			}
+		}
+
+		return MuteSoloStates;
+	}
+
+	bool SSignalFlowGraphNode::IsMuteSoloAvailableForNode() const
+	{
+		return NodeEntryKey.EntryType == ESignalFlowEntryType::SoundSource;
+	}
+#endif // MUTE_SOLO_ENABLED
+
 	void SSignalFlowGraphSendLabel::Construct(const SSignalFlowGraphSendLabel::FArguments& InArgs)
 	{
 		using namespace SSignalFlowGraphStyle;
diff --git a/Engine/Plugins/AudioInsights/Source/AudioInsights/Public/Views/SSignalFlowGraph.h b/Engine/Plugins/AudioInsights/Source/AudioInsights/Public/Views/SSignalFlowGraph.h
--- a/Engine/Plugins/AudioInsights/Source/AudioInsights/Public/Views/SSignalFlowGraph.h
+++ b/Engine/Plugins/AudioInsights/Source/AudioInsights/Public/Views/SSignalFlowGraph.h
@@ -13,6 +13,7 @@
 #endif // WITH_EDITOR
 
 #define UE_API AUDIOINSIGHTS_API
+#define MUTE_SOLO_ENABLED WITH_EDITOR && ENABLE_AUDIO_DEBUG
 
 struct FSignalFlowNodeDetailFilterSettings;
 
@@ -143,6 +144,20 @@ namespace UE::Audio::Insights
 		bool NodeDetailsExpanderArrowIsEnabled() const;
 		bool SendLabelsAreVisible() const;
 
+#if MUTE_SOLO_ENABLED
+		void ToggleMute();
+		void ToggleSolo();
+
+		struct FMuteSoloStates
+		{
+			bool bIsMuted = false;
+			bool bIsSoloed = false;
+		};
+
+		FMuteSoloStates GetMuteSoloState() const;
+		bool IsMuteSoloAvailableForNode() const;
+#endif // MUTE_SOLO_ENABLED
+
 		TSharedPtr<FSignalFlowDashboardEntry> Entry;
 		TAttribute<FVector2d> PositionInGraph;
 		TAttribute<float> XPos;
@@ -162,6 +177,11 @@ namespace UE::Audio::Insights
 		TAttribute<bool> GraphIsDraggingToMove;
 #endif // WITH_EDITOR
 
+#if MUTE_SOLO_ENABLED
+		FName AssetFName;
+		FString AssetString;
+#endif // MUTE_SOLO_ENABLED
+
 		FVector2d CachedPositionInGraph = FVector2d::Zero();
 		FSignalFlowEntryKey NodeEntryKey;
 
