[Audio Insights] Move signal flow graph node and send labels widget classes into their own files

#Synced-CL 51665297

diff --git a/Engine/Plugins/AudioInsights/Source/AudioInsights/Private/Views/SSignalFlowGraph.cpp b/Engine/Plugins/AudioInsights/Source/AudioInsights/Private/Views/SSignalFlowGraph.cpp
--- a/Engine/Plugins/AudioInsights/Source/AudioInsights/Private/Views/SSignalFlowGraph.cpp
+++ b/Engine/Plugins/AudioInsights/Source/AudioInsights/Private/Views/SSignalFlowGraph.cpp
@@ -2,20 +2,12 @@
 
 #include "Views/SSignalFlowGraph.h"
 
-#include "AudioInsightsStyle.h"
-#include "AudioInsightsUtils.h"
-#include "Brushes/SlateRoundedBoxBrush.h"
-#include "Internationalization/Text.h"
 #include "Messages/SignalFlowTraceMessages.h"
 #include "Settings/SignalFlowSettings.h"
-#include "Styling/StyleColors.h"
 #include "Views/SignalFlowNodes.h"
-#include "Widgets/Colors/SColorBlock.h"
-#include "Widgets/Input/SCheckBox.h"
-#include "Widgets/Layout/SBorder.h"
+#include "Views/SSignalFlowGraphStyle.h"
 #include "Widgets/Layout/SScaleBox.h"
 #include "Widgets/SBoxPanel.h"
-#include "Widgets/Text/STextBlock.h"
 
 #if WITH_EDITOR
 #include "Framework/Application/IMenu.h"
@@ -24,1217 +16,8 @@
 #include "Layout/WidgetPath.h"
 #endif // WITH_EDITOR
 
-#if MUTE_SOLO_ENABLED
-#include "Audio/AudioDebug.h"
-#endif // MUTE_SOLO_ENABLED
-
-#define LOCTEXT_NAMESPACE "AudioInsights"
-
 namespace UE::Audio::Insights
 {
-	namespace SSignalFlowGraphStyle
-	{
-		constexpr float NodeMaxWidth = 256.0f;
-
-		constexpr float NodeRoundedBrushRadius = 12.0f;
-
-		constexpr float SplineRadiusFlowingForwards = 80.0f;
-		constexpr float SplineRadiusFlowingBackwards = 600.0f;
-		constexpr float SplineThicknessScalarMin = 1.0f;
-		constexpr float SplineThicknessDefault = 1.5f;
-		constexpr float SplineAnimatedWireThickness = 3.0f;
-		constexpr float SplineThicknessMin = 1.0f;
-		constexpr float SplineAnimPointsPerUnit = 0.2f;
-
-		constexpr float SelectedNodeBorderThickness = 3.0f;
-		constexpr float NodeBorderThickness = 2.0f;
-
-		constexpr float NodeBorderPadding = 12.0f;
-		constexpr float NodeDetailVerticalPadding = 8.0f;
-		constexpr float NodeDetailIndentationPadding = 20.0f;
-		constexpr float LargeBetweenNodesPadding = 48.0f;
-		constexpr float SmallBetweenNodePadding = 8.0f;
-
-		constexpr float SendLabelVerticalWidth = 20.0f;
-		constexpr float SendLabelPadding = 8.0f;
-
-		constexpr float AnimationSpeed = 15.0f;
-		constexpr float DefaultZoomScale = 1.0f;
-		constexpr float ZoomStepFactor = 0.1f;
-		constexpr float MinZoomScale = 0.5f;
-		constexpr float MaxZoomScale = 2.0f;
-		constexpr float ScrollAnimInRangeThreshold = (1.e-2f);
-
-		const FSlateRoundedBoxBrush& GetNodeRoundedBoxBrush()
-		{
-			static const FSlateRoundedBoxBrush RoundedBoxBrush(FLinearColor::White, NodeRoundedBrushRadius);
-			return RoundedBoxBrush;
-		}
-
-		float GetNodeMaxHeight(const int32 NumParamsVisible, const bool bShowNodeDetails)
-		{
-			constexpr float NodeMaxHeight = 50.0f;
-			constexpr float NodeDetailsRowHeight = 24.0f;
-
-			return bShowNodeDetails ? NodeMaxHeight + (NodeDetailsRowHeight * NumParamsVisible) : NodeMaxHeight;
-		}
-
-		FSlateColor GetNodeAccentColor(const ESignalFlowEntryType NodeType)
-		{
-			switch (NodeType)
-			{
-				case ESignalFlowEntryType::SoundSource:
-					return FStyleColors::AccentBlue;
-				case ESignalFlowEntryType::AudioBus:
-					return FStyleColors::AccentYellow;
-				case ESignalFlowEntryType::Submix:
-					return FStyleColors::AccentGreen;
-				case ESignalFlowEntryType::OwnerObject:
-				case ESignalFlowEntryType::AudioDevice:
-				default:
-					return FStyleColors::Foreground;
-			}
-		}
-
-		FLinearColor GetConnectionColor(const float AmplitudeValue, const float WireSplineAmplitudePowerFactor, const bool bIsSourceBusConnection)
-		{
-			if (bIsSourceBusConnection)
-			{
-				constexpr float MaxAlpha = 0.5f;
-				return FLinearColor::LerpUsingHSV(FStyleColors::White25.GetSpecifiedColor(), GetNodeAccentColor(ESignalFlowEntryType::SoundSource).GetSpecifiedColor().CopyWithNewOpacity(MaxAlpha), FMath::Pow(AmplitudeValue, WireSplineAmplitudePowerFactor));
-			}
-
-			return FLinearColor::LerpUsingHSV(FStyleColors::Hover.GetSpecifiedColor(), FStyleColors::AccentWhite.GetSpecifiedColor(), FMath::Pow(AmplitudeValue, WireSplineAmplitudePowerFactor));
-		}
-
-		FVector2d GetSplineStartDirection(const EOrientation Orientation, const bool bIsSourceBusConnection)
-		{
-			const float SplineRadius = bIsSourceBusConnection ? SplineRadiusFlowingBackwards : SplineRadiusFlowingForwards;
-			return Orientation == EOrientation::Orient_Horizontal ? FVector2d{ SplineRadius, 0.0f } : FVector2d{ 0.0f, SplineRadius };
-		}
-
-		FVector2d GetSplineEndDirection(const EOrientation Orientation, const bool bIsSourceBusConnection)
-		{
-			const float SplineRadius = bIsSourceBusConnection ?  SplineRadiusFlowingBackwards : SplineRadiusFlowingForwards;
-			return Orientation == EOrientation::Orient_Horizontal ? FVector2d{ SplineRadius, 0.0f } : FVector2d{ 0.0f, SplineRadius };
-		}
-
-		TOptional<float> GetNodeDetailParamValue(const TSharedPtr<FSignalFlowDashboardEntry>& Entry, const ESignalFlowNodeDetailParam ParamType)
-		{
-			if (!Entry.IsValid())
-			{
-				return NullOpt;
-			}
-
-			switch (ParamType)
-			{
-				case ESignalFlowNodeDetailParam::Amplitude:
-					return Entry->Amplitude;
-
-				case ESignalFlowNodeDetailParam::Volume:
-					return Entry->Volume;
-
-				case ESignalFlowNodeDetailParam::Pitch:
-					return Entry->Pitch;
-
-				case ESignalFlowNodeDetailParam::LPFFreq:
-					return Entry->LPFFreq;
-
-				case ESignalFlowNodeDetailParam::HPFFreq:
-					return Entry->HPFFreq;
-
-				case ESignalFlowNodeDetailParam::Priority:
-					return Entry->Priority;
-
-				case ESignalFlowNodeDetailParam::Distance:
-					return Entry->Distance;
-
-				case ESignalFlowNodeDetailParam::Attenuation:
-					return Entry->Attenuation;
-
-				case ESignalFlowNodeDetailParam::RelativeRenderCost:
-					return Entry->RelativeRenderCost;
-
-				default:
-					ensureMsgf(false, TEXT("Unrecognized ESignalFlowNodeDetailParam value"));
-					break;
-			}
-
-			return NullOpt;
-		}
-
-		FSlateColor GetNodeDetailParamColor(const ESignalFlowNodeDetailParam ParamType)
-		{
-			switch (ParamType)
-			{
-				case ESignalFlowNodeDetailParam::Amplitude:
-					return FStyleColors::AccentGreen;
-
-				case ESignalFlowNodeDetailParam::Volume:
-					return FStyleColors::AccentYellow;
-
-				case ESignalFlowNodeDetailParam::Pitch:
-					return FStyleColors::AccentOrange;
-
-				case ESignalFlowNodeDetailParam::LPFFreq:
-					return FStyleColors::AccentBrown;
-
-				case ESignalFlowNodeDetailParam::HPFFreq:
-					return FStyleColors::AccentRed;
-
-				case ESignalFlowNodeDetailParam::Priority:
-					return FStyleColors::AccentPink;
-
-				case ESignalFlowNodeDetailParam::Distance:
-					return FStyleColors::AccentPurple;
-
-				case ESignalFlowNodeDetailParam::Attenuation:
-					return FStyleColors::AccentBlue;
-
-				case ESignalFlowNodeDetailParam::RelativeRenderCost:
-					return FStyleColors::AccentGray;
-
-				default:
-					ensureMsgf(false, TEXT("Unrecognized ESignalFlowNodeDetailParam value"));
-					break;
-			}
-
-			return FStyleColors::White;
-		}
-
-		FText GetNodeDetailDisplayName(const ESignalFlowNodeDetailParam ParamType)
-		{
-			switch (ParamType)
-			{
-				case ESignalFlowNodeDetailParam::Amplitude:
-					return LOCTEXT("SignalFlowGraph_NodeDetails_Amplitude", "Amp (peak): {0}");
-
-				case ESignalFlowNodeDetailParam::Volume:
-					return LOCTEXT("SignalFlowGraph_NodeDetails_Volume", "Volume: {0}");
-
-				case ESignalFlowNodeDetailParam::Pitch:
-					return LOCTEXT("SignalFlowGraph_NodeDetails_Pitch", "Pitch: {0}");
-
-				case ESignalFlowNodeDetailParam::LPFFreq:
-					return LOCTEXT("SignalFlowGraph_NodeDetails_LPFFreq", "LPF Freq (Hz): {0}");
-
-				case ESignalFlowNodeDetailParam::HPFFreq:
-					return LOCTEXT("SignalFlowGraph_NodeDetails_HPFFreq", "HPF Freq (Hz): {0}");
-
-				case ESignalFlowNodeDetailParam::Priority:
-					return LOCTEXT("SignalFlowGraph_NodeDetails_Priority", "Priority: {0}");
-
-				case ESignalFlowNodeDetailParam::Distance:
-					return LOCTEXT("SignalFlowGraph_NodeDetails_Distance", "Distance: {0}");
-
-				case ESignalFlowNodeDetailParam::Attenuation:
-					return LOCTEXT("SignalFlowGraph_NodeDetails_Attenuation", "Distance/Occlusion: {0}");
-
-				case ESignalFlowNodeDetailParam::RelativeRenderCost:
-					return LOCTEXT("SignalFlowGraph_NodeDetails_RekativeRenderCost", "Rel. Render Cost: {0}");
-
-				default:
-					ensureMsgf(false, TEXT("Unrecognized ESignalFlowNodeDetailParam value"));
-					break;
-			}
-
-			return FText::GetEmpty();
-		}
-
-		const FNumberFormattingOptions* GetNodeDetailNumberFormat(const ESignalFlowNodeDetailParam ParamType)
-		{
-			switch (ParamType)
-			{
-				case ESignalFlowNodeDetailParam::Amplitude:
-					return FSlateStyle::Get().GetLinearAmplitudeFloatFormat();
-
-				case ESignalFlowNodeDetailParam::Volume:
-					return FSlateStyle::Get().GetLinearVolumeFloatFormat();
-
-				case ESignalFlowNodeDetailParam::Pitch:
-					return FSlateStyle::Get().GetPitchFloatFormat();
-
-				case ESignalFlowNodeDetailParam::LPFFreq:
-					return FSlateStyle::Get().GetFreqFloatFormat();
-
-				case ESignalFlowNodeDetailParam::HPFFreq:
-					return FSlateStyle::Get().GetFreqFloatFormat();
-
-				case ESignalFlowNodeDetailParam::Priority:
-					return FSlateStyle::Get().GetLinearVolumeFloatFormat();
-
-				case ESignalFlowNodeDetailParam::Distance:
-					return FSlateStyle::Get().GetDefaultFloatFormat();
-
-				case ESignalFlowNodeDetailParam::Attenuation:
-					return FSlateStyle::Get().GetLinearVolumeFloatFormat();
-
-				case ESignalFlowNodeDetailParam::RelativeRenderCost:
-					return FSlateStyle::Get().GetRelativeRenderCostFormat();
-
-				default:
-					ensureMsgf(false, TEXT("Unrecognized ESignalFlowNodeDetailParam value"));
-					break;
-			}
-
-			return FSlateStyle::Get().GetDefaultFloatFormat();
-		}
-	}
-
-	void SSignalFlowGraphNode::Construct(const SSignalFlowGraphNode::FArguments& InArgs)
-	{
-		using namespace SSignalFlowGraphStyle;
-
-		Entry = InArgs._Entry;
-		NodeDetailFilterSettings = InArgs._NodeDetailFilterSettings;
-		PositionInGraph = InArgs._PositionInGraph;
-		XPos = InArgs._XPos;
-		YPos = InArgs._YPos;
-		ZoomScaleFactor = InArgs._ZoomScaleFactor;
-		IsSelected = InArgs._IsSelected;
-		ShowNodeDetails = InArgs._ShowNodeDetails;
-		DisplayAmpPeakInDb = InArgs._DisplayAmpPeakInDb;
-		GraphOrientation = InArgs._GraphOrientation;
-		OnNodeSelected = InArgs._OnNodeSelected;
-
-#if WITH_EDITOR
-		OnOpenContextMenu = InArgs._OnOpenContextMenu;
-		GraphIsDraggingToMove = InArgs._GraphIsDraggingToMove;
-#endif // WITH_EDITOR
-
-		if (!Entry.IsValid())
-		{
-			return;
-		}
-
-		NodeEntryKey = Entry->GetSignalFlowEntryKey();
-
-#if MUTE_SOLO_ENABLED
-		const FSoftObjectPath AssetPath(Entry->Name);
-		if (AssetPath.IsValid())
-		{
-			AssetString = AssetPath.GetAssetName();
-			AssetFName = AssetPath.GetAssetFName();
-		}
-		else
-		{
-			AssetString = Entry->Name;
-			AssetFName = FName(Entry->Name);
-		}
-#endif // MUTE_SOLO_ENABLED
-
-		ChildSlot
-		[
-			SAssignNew(SelectedNodeBorder, SBorder)
-			.BorderImage(&GetNodeRoundedBoxBrush())
-			.HAlign(HAlign_Fill)
-			.VAlign(VAlign_Fill)
-			.BorderBackgroundColor_Lambda([this]()
-			{
-				return IsSelected.IsBound() && IsSelected.Get() ? FStyleColors::AccentOrange.GetSpecifiedColor() : FLinearColor::Transparent;
-			})
-			.Padding(SelectedNodeBorderThickness)
-			[
-				SAssignNew(NodeBorder, SBorder)
-				.BorderImage(&GetNodeRoundedBoxBrush())
-				.HAlign(HAlign_Fill)
-				.VAlign(VAlign_Fill)
-				.BorderBackgroundColor(GetNodeAccentColor(NodeEntryKey.EntryType))
-				.Padding(NodeBorderThickness)
-				.Content()
-				[
-					SNew(SBorder)
-					.BorderImage(&GetNodeRoundedBoxBrush())
-					.HAlign(HAlign_Fill)
-					.VAlign(VAlign_Fill)
-					.BorderBackgroundColor_Lambda([this]()
-					{
-						return (IsSelected.IsBound() && IsSelected.Get()) ? FStyleColors::Header : FStyleColors::Panel;
-					})
-					.Padding(0.0f, NodeBorderPadding, 0.0f, NodeBorderPadding)
-					.Content()
-					[
-						SNew(SVerticalBox)
-
-						+ SVerticalBox::Slot()
-						.Padding(NodeBorderPadding, 0.0f, NodeBorderPadding, 0.0f)
-						.AutoHeight()
-						[
-							CreateNodeName(InArgs)
-						]
-
-						+ SVerticalBox::Slot()
-						.Padding(NodeBorderPadding + NodeDetailIndentationPadding, 0.0f, 0.0f, 0.0f)
-						.AutoHeight()
-						[
-							SNew(SHorizontalBox)
-
-							+ SHorizontalBox::Slot()
-							.FillWidth(1.0f)
-							[
-								CreateNodeDetails()
-							]
-
-							+ SHorizontalBox::Slot()
-							.HAlign(HAlign_Right)
-							.VAlign(VAlign_Top)
-							.AutoWidth()
-							[
-								// In vertical mode, we want to ensure that we maintain padding on the RHS of the node when node details is on
-								// We do this by drawing a box with no content inside.
-								// In horizontal mode, this is replaced by SendLevelsVerticalBox, which is flush to the RHS
-								SNew(SVerticalBox)
-								.Visibility_Lambda([this]()
-								{
-									return GraphOrientation.Get() == EOrientation::Orient_Vertical && ShowNodeDetails.Get()
-																									? EVisibility::Visible
-																									: EVisibility::Collapsed;
-								})
-
-								+ SVerticalBox::Slot()
-								[
-									SNew(SBox)
-									.WidthOverride(NodeBorderPadding)
-								]
-							]
-
-							+ SHorizontalBox::Slot()
-							.HAlign(HAlign_Right)
-							.VAlign(VAlign_Top)
-							.AutoWidth()
-							[
-								CreateSendOutputsVerticalContainer()
-							]
-						]
-
-						+ SVerticalBox::Slot()
-						.HAlign(HAlign_Center)
-						.VAlign(VAlign_Bottom)
-						.AutoHeight()
-						[
-							CreateSendOutputsHorizontalContainer()
-						]
-					]
-				]
-			]
-		];
-	}
-
-	FReply SSignalFlowGraphNode::OnMouseButtonDown(const FGeometry& MyGeometry, const FPointerEvent& MouseEvent)
-	{
-		if (MouseEvent.GetEffectingButton() == EKeys::LeftMouseButton)
-		{
-			if (Entry.IsValid())
-			{
-				OnNodeSelected.ExecuteIfBound(IsSelected.IsBound() && IsSelected.Get() ? nullptr : Entry);
-			}
-
-			return FReply::Handled();
-		}
-		else
-		{
-			return SCompoundWidget::OnMouseButtonDown(MyGeometry, MouseEvent);
-		}
-	}
-
-#if WITH_EDITOR
-	FReply SSignalFlowGraphNode::OnMouseButtonUp(const FGeometry& MyGeometry, const FPointerEvent& MouseEvent)
-	{
-		if (MouseEvent.GetEffectingButton() == EKeys::RightMouseButton)
-		{
-			if (Entry.IsValid() && !GraphIsDraggingToMove.Get())
-			{
-				const TObjectPtr<UObject> ObjectPtr = Entry->GetObject();
-				if (ObjectPtr && ObjectPtr->IsAsset())
-				{
-					OnOpenContextMenu.ExecuteIfBound(MouseEvent, Entry);
-					return FReply::Handled();
-				}
-			}
-		}
-
-		return SCompoundWidget::OnMouseButtonUp(MyGeometry, MouseEvent);
-	}
-#endif // WITH_EDITOR
-
-	FVector2D SSignalFlowGraphNode::ComputeDesiredSize(float InLayoutScaleMultiplier) const
-	{
-		using namespace SSignalFlowGraphStyle;
-
-		int32 NumParamsVisible = 0;
-		if (NodeDetailFilterSettings)
-		{
-			if (bNodeDetailsAreExpanded)
-			{
-				const TSet<ESignalFlowNodeDetailParam> IgnoredParams { ESignalFlowNodeDetailParam::SendOutputVolume };
-				NumParamsVisible = NodeDetailFilterSettings->GetNumVisibleParams(IgnoredParams);
-			}
-			// We always show amplitude if the node is collapsed
-			else if (Entry.IsValid() && Entry->Amplitude.IsSet())
-			{
-				NumParamsVisible = 1;
-			}
-		}
-
-		const int32 NumOutputPinsShowing = SendLabels.Num();
-		const FVector2D Size = SCompoundWidget::ComputeDesiredSize(InLayoutScaleMultiplier);
-
-		if (GraphOrientation.Get() == EOrientation::Orient_Horizontal)
-		{
-			const int32 MaxNumParamRows = FMath::Max(NumParamsVisible, NumOutputPinsShowing);
-
-			const FVector2D ClampedSize = { FMath::Min(Size.X, NodeMaxWidth), FMath::Min(Size.Y, GetNodeMaxHeight(MaxNumParamRows, ShowNodeDetails.Get())) };
-			return ClampedSize;
-		}
-		else
-		{
-			const float PinSectionWidth = SendLevelsHorizontalBox.IsValid() ? SendLevelsHorizontalBox->GetDesiredSize().X : 0.0f;
-			const float OutputPinSectionHeight = SendLevelsHorizontalBox.IsValid() ? SendLevelsHorizontalBox->GetDesiredSize().Y : 0.0f;
-
-			// Make sure Max width/height calculations account for the bottom output pin row
-			const float ScaledMaxWidth = FMath::Max(NodeMaxWidth, PinSectionWidth);
-			const float TotalMaxHeight = GetNodeMaxHeight(NumParamsVisible, ShowNodeDetails.Get()) + OutputPinSectionHeight;
-
-			const FVector2D ClampedSize = { FMath::Min(Size.X, ScaledMaxWidth), FMath::Min(Size.Y, TotalMaxHeight) };
-			return ClampedSize;
-		}
-	}
-
-	FVector2D SSignalFlowGraphNode::GetLocalInputPos() const
-	{
-		using namespace SSignalFlowGraphStyle;
-
-		const FVector2D Size = GetDesiredSize();
-
-		if (GraphOrientation.IsBound() && GraphOrientation.Get() == EOrientation::Orient_Horizontal)
-		{
-			// Input pos = left center of node
-			return FVector2D{ CachedPositionInGraph.X + SelectedNodeBorderThickness, CachedPositionInGraph.Y + (Size.Y * 0.5f) } * ZoomScaleFactor.Get();
-		}
-		else
-		{
-			// Input pos = top center of node
-			return FVector2D{ CachedPositionInGraph.X + (Size.X * 0.5f), CachedPositionInGraph.Y + SelectedNodeBorderThickness } * ZoomScaleFactor.Get();
-		}
-	}
-	
-	FVector2D SSignalFlowGraphNode::GetLocalCenteredOutputPos() const
-	{
-		using namespace SSignalFlowGraphStyle;
-
-		const FVector2D Size = GetDesiredSize();
-
-		if (GraphOrientation.IsBound() && GraphOrientation.Get() == EOrientation::Orient_Horizontal)
-		{
-			// Out pos = right center of node
-			return FVector2D{ CachedPositionInGraph.X + Size.X - SelectedNodeBorderThickness, CachedPositionInGraph.Y + (Size.Y / 2.0f) } * ZoomScaleFactor.Get();
-		}
-		else
-		{
-			// Output pos = bottom center of node
-			return FVector2D{ CachedPositionInGraph.X + (Size.X / 2.0f), CachedPositionInGraph.Y + Size.Y - SelectedNodeBorderThickness } * ZoomScaleFactor.Get();
-		}
-	}
-
-	FVector2D SSignalFlowGraphNode::GetLocalSendLabelPos(const FSignalFlowEntryKey& ReceiverEntryKey, const int32 LabelIndex) const
-	{
-		using namespace SSignalFlowGraphStyle;
-
-		const FVector2D OutputPos = GetLocalCenteredOutputPos();
-		if (!NodeLabelHorizontalBox.IsValid() || !SendLabelsAreVisible())
-		{
-			return OutputPos;
-		}
-
-		const TSharedRef<SSignalFlowGraphSendLabel>* SendLabelWidget = SendLabels.Find(ReceiverEntryKey);
-		if (SendLabelWidget == nullptr)
-		{
-			return OutputPos;
-		}
-
-		constexpr float BorderThicknessYOffset = (SelectedNodeBorderThickness * 0.5f) + (NodeBorderThickness * 0.5f);
-
-		const FVector2D SendLabelWidgetSize = (*SendLabelWidget)->GetDesiredSize();
-
-		float XPosInGraph;
-		float YPosInGraph;
-
-		if (GraphOrientation.Get() == EOrientation::Orient_Horizontal)
-		{
-			XPosInGraph = OutputPos.X;
-
-			const float NodeLabelHeight = NodeLabelHorizontalBox->GetDesiredSize().Y + NodeBorderPadding;
-			const float YOffsetToLabelArea = CachedPositionInGraph.Y + BorderThicknessYOffset + NodeLabelHeight;
-			const float YOffsetInLabelList = ((SendLabelWidgetSize.Y + NodeDetailVerticalPadding) * LabelIndex) + NodeDetailVerticalPadding;
-			const float HalfLabelHeight = (SendLabelWidgetSize.Y * 0.5f);
-
-			YPosInGraph = (YOffsetToLabelArea + YOffsetInLabelList + HalfLabelHeight) * ZoomScaleFactor.Get();
-		}
-		else
-		{
-			YPosInGraph = OutputPos.Y;
-
-			const float SendLabelWidth = SendLabelWidgetSize.X + SendLabelPadding;
-			const float XOffset = (SendLabelWidth * (LabelIndex + 1)) - ((SendLabelWidth * SendLabels.Num()) * 0.5f) - (SendLabelWidth * 0.5f);
-			XPosInGraph = OutputPos.X + (XOffset * ZoomScaleFactor.Get());
-		}
-
-		const FVector2D SendLabelPosition(XPosInGraph, YPosInGraph);
-		return SendLabelPosition;
-	}
-
-	int32 SSignalFlowGraphNode::GetNumSendLabels() const
-	{
-		return SendLabels.Num();
-	}
-
-	void SSignalFlowGraphNode::UpdateNodeLayout()
-	{
-		using namespace SSignalFlowGraphStyle;
-
-		// Ensure border thickness does not get scaled to avoid rendering artifacts
-		if (SelectedNodeBorder.IsValid())
-		{
-			SelectedNodeBorder->SetPadding(SelectedNodeBorderThickness / ZoomScaleFactor.Get());
-		}
-
-		if (NodeBorder.IsValid())
-		{
-			NodeBorder->SetPadding(NodeBorderThickness / ZoomScaleFactor.Get());
-		}
-
-		// Update the node's position in the graph
-		CachedPositionInGraph = PositionInGraph.Get();
-	}
-
-	FVector2D SSignalFlowGraphNode::GetCachedPositionInGraph() const
-	{
-		return CachedPositionInGraph;
-	}
-
-	void SSignalFlowGraphNode::SetSendLabelWidgets(const TArray<FSignalFlowEntryKey>& ReceiverEntryKeys)
-	{
-		for (const FSignalFlowEntryKey& ReceiverKey : ReceiverEntryKeys)
-		{
-			AddSendLabelWidget(ReceiverKey);
-		}
-
-		for (auto Itr = SendLabels.CreateIterator(); Itr; ++Itr)
-		{
-			if (!ReceiverEntryKeys.Contains(Itr->Key))
-			{
-				if (SendLevelsVerticalBox.IsValid())
-				{
-					SendLevelsVerticalBox->RemoveSlot(Itr->Value);
-				}
-
-				if (SendLevelsHorizontalBox.IsValid())
-				{
-					SendLevelsHorizontalBox->RemoveSlot(Itr->Value);
-				}
-
-				Itr.RemoveCurrent();
-			}
-		}
-	}
-
-	void SSignalFlowGraphNode::AddSendLabelWidget(const FSignalFlowEntryKey& ReceiverEntryKey)
-	{
-		using namespace SSignalFlowGraphStyle;
-
-		if (!SendLevelsVerticalBox.IsValid() || !SendLevelsHorizontalBox.IsValid() || SendLabels.Contains(ReceiverEntryKey))
-		{
-			return;
-		}
-
-		TSharedPtr<SSignalFlowGraphSendLabel> SendLabelWidget = SNew(SSignalFlowGraphSendLabel)
-			.PixelSnappingMethod(EWidgetPixelSnapping::SnapToPixel)
-			.ToolTipText(LOCTEXT("SignalFlowGraph_NodeDetails_SendLabelTooltip", "Send Output Volume"))
-			.ReceiverEntryType(ReceiverEntryKey.EntryType)
-			.SendLevel_Lambda([this, ReceiverEntryKey]() -> TOptional<float>
-			{
-				if (Entry.IsValid())
-				{
-					const FSignalFlowOutputData* SendInfo = Entry->Outputs.Find(ReceiverEntryKey);
-					if (SendInfo && SendInfo->SendLevel.IsSet())
-					{
-						return SendInfo->SendLevel;
-					}
-					else if (Entry->LinkedSoundSources.Contains(ReceiverEntryKey))
-					{
-						// Source buses always have a send value of 1.0
-						return 1.0f;
-					}
-				}
-
-				return NullOpt;
-			})
-			.GraphOrientation(GraphOrientation)
-			.ZoomScaleFactor(ZoomScaleFactor);
-
-		SendLabels.Add(ReceiverEntryKey, SendLabelWidget.ToSharedRef());
-
-		SendLevelsVerticalBox->AddSlot()
-		.Padding(SendLabelPadding, NodeDetailVerticalPadding, 0.0f, 0.0f)
-		.VAlign(VAlign_Top)
-		.HAlign(HAlign_Right)
-		[
-			SendLabelWidget.ToSharedRef()
-		];
-
-		SendLevelsHorizontalBox->AddSlot()
-		.Padding(SendLabelPadding * 0.5f, NodeDetailVerticalPadding, SendLabelPadding * 0.5f, 0.0f)
-		.VAlign(VAlign_Center)
-		.HAlign(HAlign_Center)
-		[
-			SendLabelWidget.ToSharedRef()
-		];
-	}
-
-	void SSignalFlowGraphNode::RemoveSendLabelWidget(const FSignalFlowEntryKey& ReceiverEntryKey)
-	{
-		TSharedRef<SSignalFlowGraphSendLabel>* SendLabelWidget = SendLabels.Find(ReceiverEntryKey);
-		if (SendLabelWidget == nullptr)
-		{
-			return;
-		}
-
-		if (SendLevelsVerticalBox.IsValid())
-		{
-			SendLevelsVerticalBox->RemoveSlot(*SendLabelWidget);
-		}
-
-		if (SendLevelsHorizontalBox.IsValid())
-		{
-			SendLevelsHorizontalBox->RemoveSlot(*SendLabelWidget);
-		}
-
-		SendLabels.Remove(ReceiverEntryKey);
-	}
-
-	bool SSignalFlowGraphNode::ContainsSendLabel(const FSignalFlowEntryKey& ReceiverEntryKey) const
-	{
-		return SendLabels.Contains(ReceiverEntryKey);
-	}
-
-	TSharedRef<SHorizontalBox> SSignalFlowGraphNode::CreateNodeName(const SSignalFlowGraphNode::FArguments& InArgs)
-	{
-		using namespace SSignalFlowGraphStyle;
-
-		return SAssignNew(NodeLabelHorizontalBox, SHorizontalBox)
-
-			+ SHorizontalBox::Slot()
-			.AutoWidth()
-			.Padding(0.0f, 0.0f, 4.0f, 0.0f)
-			[
-				SNew(SImage)
-				.Image(FSlateStyle::Get().CreateIcon(InArgs._Entry->IconName).GetIcon())
-				.ToolTipText(InArgs._DisplayName)
-				.ColorAndOpacity(GetNodeAccentColor(NodeEntryKey.EntryType))
-			]
-
-			+ SHorizontalBox::Slot()
-			.FillWidth(1.0f)
-			[
-				SNew(STextBlock)
-				.Text(InArgs._DisplayName)
-				.ToolTipText(InArgs._DisplayName)
-				.Justification(ETextJustify::Left)
-				.OverflowPolicy(ETextOverflowPolicy::Ellipsis)
-				.Clipping(EWidgetClipping::ClipToBounds)
-			]
-
-#if MUTE_SOLO_ENABLED
-			+ SHorizontalBox::Slot()
-			.AutoWidth()
-			.Padding(4.0f, 0.0f, 0.0f, 0.0f)
-			[
-				SNew(SCheckBox)
-				.Style(&FSlateStyle::Get().GetWidgetStyle<FCheckBoxStyle>("SoundDashboard.MuteSoloButton"))
-				.Visibility_Lambda([this]()
-				{
-					return IsMuteSoloAvailableForNode() ? EVisibility::Visible : EVisibility::Collapsed;
-				})
-				.IsChecked(ECheckBoxState::Unchecked)
-				.OnCheckStateChanged_Lambda([this](ECheckBoxState)
-				{
-					ToggleMute();
-				})
-				[
-					SNew(SImage)
-					.Image(FSlateStyle::Get().CreateIcon("AudioInsights.Icon.SoundDashboard.Mute").GetIcon())
-					.ToolTipText(LOCTEXT("SignalFlowGraph_MuteNode", "Mute this node.\nMute will apply to all instances of an asset."))
-					.ColorAndOpacity_Lambda([this]()
-					{
-						return GetMuteSoloState().bIsMuted ? FStyleColors::Foreground : FStyleColors::White25;
-					})
-				]
-			]
-
-			+ SHorizontalBox::Slot()
-			.AutoWidth()
-			.Padding(0.0f, 0.0f, 4.0f, 0.0f)
-			[
-				SNew(SCheckBox)
-				.Style(&FSlateStyle::Get().GetWidgetStyle<FCheckBoxStyle>("SoundDashboard.MuteSoloButton"))
-				.Visibility_Lambda([this]()
-				{
-					return IsMuteSoloAvailableForNode() ? EVisibility::Visible : EVisibility::Collapsed;
-				})
-				.IsChecked(ECheckBoxState::Unchecked)
-				.OnCheckStateChanged_Lambda([this](ECheckBoxState)
-				{
-					ToggleSolo();
-				})
-				[
-					SNew(SImage)
-					.Image(FSlateStyle::Get().CreateIcon("AudioInsights.Icon.SoundDashboard.Solo").GetIcon())
-					.ToolTipText(LOCTEXT("SignalFlowGraph_SoloNode", "Solo this node.\nSolo will apply to all instances of an asset."))
-					.ColorAndOpacity_Lambda([this]()
-					{
-						return GetMuteSoloState().bIsSoloed ? FStyleColors::Foreground : FStyleColors::White25;
-					})
-				]
-			]
-#endif // MUTE_SOLO_ENABLED
-		
-			+ SHorizontalBox::Slot()
-			.AutoWidth()
-			.MinSize(12.0f)
-			.MaxSize(12.0f)
-			.Padding(4.0f, 0.0f, 0.0f, 0.0f)
-			[
-				SNew(SCheckBox)
-				.Style(&FSlateStyle::Get().GetWidgetStyle<FCheckBoxStyle>("AudioInsights.LeftExpandArrowToggle"))
-				.Visibility_Lambda([this]()
-				{ 
-					return NodeDetailsExpanderArrowIsEnabled() ? EVisibility::Visible : EVisibility::Collapsed;
-				})
-				.IsChecked_Lambda([this]()
-				{
-					return bNodeDetailsAreExpanded ? ECheckBoxState::Checked : ECheckBoxState::Unchecked;
-				})
-				.OnCheckStateChanged_Lambda([this](const ECheckBoxState CheckboxState)
-				{
-					bNodeDetailsAreExpanded = CheckboxState == ECheckBoxState::Checked;
-				})
-			];
-	}
-
-	TSharedRef<SVerticalBox> SSignalFlowGraphNode::CreateNodeDetails() const
-	{
-		using namespace SSignalFlowGraphStyle;
-
-		return SNew(SVerticalBox)
-				.Visibility_Lambda([this]() { return ShowNodeDetails.Get() ? EVisibility::Visible : EVisibility::Collapsed; })
-
-				+ SVerticalBox::Slot()
-				.Padding(0.0f, NodeDetailVerticalPadding, 0.0f, 0.0f)
-				.AutoHeight()
-				[
-					CreateNodeDetailsRow(ESignalFlowNodeDetailParam::Amplitude)
-				]
-			
-				+ SVerticalBox::Slot()
-				.Padding(0.0f, NodeDetailVerticalPadding, 0.0f, 0.0f)
-				.AutoHeight()
-				[
-					CreateNodeDetailsRow(ESignalFlowNodeDetailParam::Volume)
-				]
-
-				+ SVerticalBox::Slot()
-				.Padding(0.0f, NodeDetailVerticalPadding, 0.0f, 0.0f)
-				.AutoHeight()
-				[
-					CreateNodeDetailsRow(ESignalFlowNodeDetailParam::Pitch)
-				]
-
-				+ SVerticalBox::Slot()
-				.Padding(0.0f, NodeDetailVerticalPadding, 0.0f, 0.0f)
-				.AutoHeight()
-				[
-					CreateNodeDetailsRow(ESignalFlowNodeDetailParam::LPFFreq)
-				]
-			
-				+ SVerticalBox::Slot()
-				.Padding(0.0f, NodeDetailVerticalPadding, 0.0f, 0.0f)
-				.AutoHeight()
-				[
-					CreateNodeDetailsRow(ESignalFlowNodeDetailParam::HPFFreq)
-				]
-
-				+ SVerticalBox::Slot()
-				.Padding(0.0f, NodeDetailVerticalPadding, 0.0f, 0.0f)
-				.AutoHeight()
-				[
-					CreateNodeDetailsRow(ESignalFlowNodeDetailParam::Priority)
-				]
-
-				+ SVerticalBox::Slot()
-				.Padding(0.0f, NodeDetailVerticalPadding, 0.0f, 0.0f)
-				.AutoHeight()
-				[
-					CreateNodeDetailsRow(ESignalFlowNodeDetailParam::Distance)
-				]
-
-				+ SVerticalBox::Slot()
-				.Padding(0.0f, NodeDetailVerticalPadding, 0.0f, 0.0f)
-				.AutoHeight()
-				[
-					CreateNodeDetailsRow(ESignalFlowNodeDetailParam::Attenuation)
-				]
-
-				+ SVerticalBox::Slot()
-				.Padding(0.0f, NodeDetailVerticalPadding, 0.0f, 0.0f)
-				.AutoHeight()
-				[
-					CreateNodeDetailsRow(ESignalFlowNodeDetailParam::RelativeRenderCost)
-				];
-	}
-
-	TSharedRef<SHorizontalBox> SSignalFlowGraphNode::CreateNodeDetailsRow(const ESignalFlowNodeDetailParam ParamType) const
-	{
-		using namespace SSignalFlowGraphStyle;
-
-		return SNew(SHorizontalBox)
-				.Visibility_Lambda([this, ParamType]()
-				{
-					const bool bParamIsFiltered = NodeDetailFilterSettings && NodeDetailFilterSettings->GetParameterIsVisible(ParamType);
-					const bool bParamIsHiddenByExpander = !bNodeDetailsAreExpanded && ParamType != ESignalFlowNodeDetailParam::Amplitude;
-
-					return bParamIsFiltered && !bParamIsHiddenByExpander && Entry.IsValid() && GetNodeDetailParamValue(Entry, ParamType).IsSet() ? EVisibility::Visible : EVisibility::Collapsed;
-				})
-
-				+ SHorizontalBox::Slot()
-				.Padding(0.0f, 0.0f, 4.0f, 0.0f)
-				.AutoWidth()
-				[
-					SNew(SImage)
-					.Image(FSlateStyle::Get().CreateIcon("AudioInsights.Icon.SignalFlow.Pill").GetIcon())
-					.ColorAndOpacity(GetNodeDetailParamColor(ParamType))
-				]
-
-				+ SHorizontalBox::Slot()
-				.FillWidth(1.0f)
-				[
-					SNew(STextBlock)
-					.Text_Lambda([this, ParamType]()
-					{
-						const TOptional<float> ParamValueOpt = GetNodeDetailParamValue(Entry, ParamType);
-						if (!ParamValueOpt.IsSet())
-						{
-							return FText::GetEmpty();
-						}
-
-						if (ParamType == ESignalFlowNodeDetailParam::Amplitude && DisplayAmpPeakInDb.Get() == true)
-						{
-							return FText::Format(GetNodeDetailDisplayName(ParamType), AudioInsightsUtils::ConvertToDecibelsText(ParamValueOpt.GetValue()));
-						}
-
-						if (ParamType == ESignalFlowNodeDetailParam::Priority && SoundMessageUtils::IsPriorityValueSetToMax(ParamValueOpt.GetValue()))
-						{
-							return FText::Format(GetNodeDetailDisplayName(ParamType), LOCTEXT("SignalFlowGraph_NodeDetails_Max", "MAX"));
-						}
-
-						return FText::Format(GetNodeDetailDisplayName(ParamType), FText::AsNumber(ParamValueOpt.GetValue(), GetNodeDetailNumberFormat(ParamType)));
-					})
-					.Font(FCoreStyle::Get().GetFontStyle("SmallFont"))
-					.Justification(ETextJustify::Left)
-				];
-	}
-
-	TSharedRef<SVerticalBox> SSignalFlowGraphNode::CreateSendOutputsVerticalContainer()
-	{
-		SendLevelsVerticalBox = SNew(SVerticalBox)
-								.Visibility_Lambda([this]()
-								{
-									return GraphOrientation.Get() == EOrientation::Orient_Horizontal && SendLabelsAreVisible() 
-																	? EVisibility::Visible 
-																	: EVisibility::Collapsed;
-								});
-
-		return SendLevelsVerticalBox.ToSharedRef();
-	}
-
-	TSharedRef<SHorizontalBox> SSignalFlowGraphNode::CreateSendOutputsHorizontalContainer()
-	{
-		SendLevelsHorizontalBox = SNew(SHorizontalBox)
-									.Visibility_Lambda([this]()
-									{
-										return GraphOrientation.Get() == EOrientation::Orient_Vertical && SendLabelsAreVisible()
-																	? EVisibility::Visible
-																	: EVisibility::Collapsed;
-									});
-
-		return SendLevelsHorizontalBox.ToSharedRef();
-	}
-
-	bool SSignalFlowGraphNode::NodeDetailsExpanderArrowIsEnabled() const
-	{
-		using namespace SSignalFlowGraphStyle;
-
-		// The expander arrow should only be visible if a parameter other than amplitude is visible on the node
-		// This requires show node details to be switched on, the parameter filter to be enabled, and data to be available for that parameter
-		if (!ShowNodeDetails.Get() || !Entry.IsValid() || NodeDetailFilterSettings == nullptr)
-		{
-			return false;
-		}
-
-		const int32 NumVisibleParams = NodeDetailFilterSettings->GetNumVisibleParams();
-		if (NumVisibleParams <= 0)
-		{
-			return false;
-		}
-
-		if (NodeDetailFilterSettings->GetNumVisibleParams() == 1 && NodeDetailFilterSettings->GetParameterIsVisible(ESignalFlowNodeDetailParam::Amplitude))
-		{
-			return false;
-		}
-		
-		const int32 NumFilters = static_cast<uint8>(ESignalFlowNodeDetailParam::MAX);
-		for (uint8 FilterID = 0u; FilterID < NumFilters; ++FilterID)
-		{
-			const ESignalFlowNodeDetailParam Param = static_cast<ESignalFlowNodeDetailParam>(FilterID);
-			if (Param == ESignalFlowNodeDetailParam::Amplitude)
-			{
-				continue;
-			}
-
-			if (Param == ESignalFlowNodeDetailParam::SendOutputVolume)
-			{
-				if (NodeDetailFilterSettings->GetParameterIsVisible(Param))
-				{
-					for (const auto& [OutputKey, SendInfo] : Entry->Outputs)
-					{
-						if (SendInfo.SendLevel.IsSet())
-						{
-							return true;
-						}
-					}
-				}
-			}
-			else if (NodeDetailFilterSettings->GetParameterIsVisible(Param) && GetNodeDetailParamValue(Entry, Param).IsSet())
-			{
-				return true;
-			}
-		}
-
-		return false;
-	}
-
-	bool SSignalFlowGraphNode::SendLabelsAreVisible() const
-	{
-		if (!ShowNodeDetails.Get() || SendLabels.IsEmpty())
-		{
-			return false;
-		}
-
-		const bool bParamIsFiltered = NodeDetailFilterSettings && NodeDetailFilterSettings->GetParameterIsVisible(ESignalFlowNodeDetailParam::SendOutputVolume);
-		if (!bNodeDetailsAreExpanded || !bParamIsFiltered)
-		{
-			return false;
-		}
-
-		for (const auto& [ReceiverKey, Label] : SendLabels)
-		{
-			if (Label->ValueIsSet())
-			{
-				return true;
-			}
-		}
-
-		return false;
-	}
-
-#if MUTE_SOLO_ENABLED
-	void SSignalFlowGraphNode::ToggleMute()
-	{
-		if (FAudioDeviceManager* AudioDeviceManager = FAudioDeviceManager::Get())
-		{
-			AudioDeviceManager->GetDebugger().ToggleMuteSoundWave(AssetFName);
-		}
-	}
-
-	void SSignalFlowGraphNode::ToggleSolo()
-	{
-		if (FAudioDeviceManager* AudioDeviceManager = FAudioDeviceManager::Get())
-		{
-			AudioDeviceManager->GetDebugger().ToggleSoloSoundWave(AssetFName);
-		}
-	}
-
-	SSignalFlowGraphNode::FMuteSoloStates SSignalFlowGraphNode::GetMuteSoloState() const
-	{
-		FMuteSoloStates MuteSoloStates;
-		if (Entry.IsValid())
-		{
-			if (FAudioDeviceManager* AudioDeviceManager = FAudioDeviceManager::Get())
-			{
-				FString Reason;
-				AudioDeviceManager->GetDebugger().QuerySoloMuteSoundWave(AssetString, MuteSoloStates.bIsSoloed, MuteSoloStates.bIsMuted, Reason);
-			}
-		}
-
-		return MuteSoloStates;
-	}
-
-	bool SSignalFlowGraphNode::IsMuteSoloAvailableForNode() const
-	{
-		return NodeEntryKey.EntryType == ESignalFlowEntryType::SoundSource;
-	}
-#endif // MUTE_SOLO_ENABLED
-
-	void SSignalFlowGraphSendLabel::Construct(const SSignalFlowGraphSendLabel::FArguments& InArgs)
-	{
-		using namespace SSignalFlowGraphStyle;
-
-		SendLevel = InArgs._SendLevel;
-		GraphOrientation = InArgs._GraphOrientation;
-		ZoomScaleFactor = InArgs._ZoomScaleFactor;
-
-		ChildSlot
-		[
-			SNew(SOverlay)
-
-			// Horizontal flow label (text then icon)
-			+ SOverlay::Slot()
-			[
-				SNew(SHorizontalBox)
-				.Visibility_Lambda([this]()
-				{
-					return GraphOrientation.Get() == EOrientation::Orient_Horizontal && SendLevel.IsBound() && SendLevel.Get().IsSet() 
-																												? EVisibility::Visible 
-																												: EVisibility::Collapsed;
-				})
-
-				+ SHorizontalBox::Slot()
-				.FillContentWidth(1.0f)
-				.HAlign(HAlign_Right)
-				.VAlign(VAlign_Center)
-				[
-					CreateOutputLabelText(InArgs._ToolTipText.Get())
-				]
-
-				+ SHorizontalBox::Slot()
-				.FillContentWidth(1.0f)
-				.HAlign(HAlign_Right)
-				.VAlign(VAlign_Center)
-				[
-					CreatePinIconWidget(InArgs._ToolTipText.Get(), InArgs._ReceiverEntryType)
-				]
-			]
-
-			// Vertical flow label (text above icon)
-			+ SOverlay::Slot()
-			[
-				SNew(SVerticalBox)
-				.Visibility_Lambda([this]()
-				{
-					return GraphOrientation.Get() == EOrientation::Orient_Vertical && SendLevel.IsBound() && SendLevel.Get().IsSet() 
-																												? EVisibility::Visible 
-																												: EVisibility::Collapsed;
-				})
-
-				+ SVerticalBox::Slot()
-				.FillContentHeight(1.0f)
-				.HAlign(HAlign_Center)
-				.VAlign(VAlign_Bottom)
-				[
-					CreateOutputLabelText(InArgs._ToolTipText.Get())
-				]
-
-				+ SVerticalBox::Slot()
-				.FillContentHeight(1.0f)
-				.HAlign(HAlign_Center)
-				.VAlign(VAlign_Bottom)
-				[
-					CreatePinIconWidget(InArgs._ToolTipText.Get(), InArgs._ReceiverEntryType)
-				]
-			]
-		];
-	}
-
-	FVector2D SSignalFlowGraphSendLabel::ComputeDesiredSize(float InLayoutScaleMultiplier) const
-	{
-		using namespace SSignalFlowGraphStyle;
-
-		const FVector2D Size = SCompoundWidget::ComputeDesiredSize(InLayoutScaleMultiplier);
-		if (GraphOrientation.Get() == EOrientation::Orient_Horizontal)
-		{
-			return Size;
-		}
-		else
-		{
-			return FVector2D(SendLabelVerticalWidth, Size.Y);
-		}
-	}
-
-	bool SSignalFlowGraphSendLabel::ValueIsSet() const
-	{
-		return SendLevel.IsBound() && SendLevel.Get().IsSet();
-	}
-
-	TSharedRef<STextBlock> SSignalFlowGraphSendLabel::CreateOutputLabelText(const FText& ToolTipText) const
-	{
-		return SNew(STextBlock)
-				.Text_Lambda([this]()
-				{
-					if (!SendLevel.IsBound())
-					{
-						return FText::GetEmpty();
-					}
-
-					const TOptional<float> SendLevelOpt = SendLevel.Get();
-					if (!SendLevelOpt.IsSet())
-					{
-						return FText::GetEmpty();
-					}
-
-					return FText::AsNumber(SendLevelOpt.GetValue(), FSlateStyle::Get().GetSendLevelFormat());
-				})
-				.ToolTipText(ToolTipText)
-				.Justification(ETextJustify::Center)
-				.Font(FCoreStyle::Get().GetFontStyle("SmallFont"));
-	}
-
-	TSharedRef<SOverlay> SSignalFlowGraphSendLabel::CreatePinIconWidget(const FText& ToolTipText, const ESignalFlowEntryType ReceiverEntryType)
-	{
-		using namespace SSignalFlowGraphStyle;
-
-		return SNew(SOverlay)
-
-			+ SOverlay::Slot()
-			[
-				SAssignNew(OutputPinImage, SImage)
-				.Image(FSlateStyle::Get().CreateIcon("AudioInsights.Icon.SignalFlow.Disconnected").GetIcon())
-				.ColorAndOpacity(GetNodeAccentColor(ReceiverEntryType))
-			]
-
-			+ SOverlay::Slot()
-			[
-				SNew(SImage)
-				.Image(FSlateStyle::Get().CreateIcon("AudioInsights.Icon.SignalFlow.Connected").GetIcon())
-				.ToolTipText(ToolTipText)
-				.ColorAndOpacity_Lambda([this, ReceiverEntryType]()
-				{
-					const FSlateColor BaseColor = GetNodeAccentColor(ReceiverEntryType);
-
-					if (!SendLevel.IsBound())
-					{
-						return BaseColor;
-					}
-
-					const TOptional<float> SendLevelOpt = SendLevel.Get();
-					if (!SendLevelOpt.IsSet())
-					{
-						return BaseColor;
-					}
-
-					return FSlateColor(BaseColor.GetSpecifiedColor().CopyWithNewOpacity(SendLevelOpt.GetValue()));
-				})
-			];
-	}
-
 	void SSignalFlowGraph::Construct(const SSignalFlowGraph::FArguments& InArgs)
 	{
 		ListSource = InArgs._ListSource;
@@ -2375,5 +1158,3 @@ namespace UE::Audio::Insights
 	}
 #endif // WITH_EDITOR
 } // namespace UE::Audio::Insights
-
-#undef LOCTEXT_NAMESPACE
diff --git a/Engine/Plugins/AudioInsights/Source/AudioInsights/Private/Views/SSignalFlowGraphNode.cpp b/Engine/Plugins/AudioInsights/Source/AudioInsights/Private/Views/SSignalFlowGraphNode.cpp
new file mode 100644
--- /dev/null
+++ b/Engine/Plugins/AudioInsights/Source/AudioInsights/Private/Views/SSignalFlowGraphNode.cpp
@@ -0,0 +1,830 @@
+// Copyright Epic Games, Inc. All Rights Reserved.
+
+#include "Views/SSignalFlowGraphNode.h"
+
+#include "AudioInsightsStyle.h"
+#include "AudioInsightsUtils.h"
+#include "Internationalization/Text.h"
+#include "Messages/SignalFlowTraceMessages.h"
+#include "Settings/SignalFlowSettings.h"
+#include "Styling/StyleColors.h"
+#include "Views/SSignalFlowGraphStyle.h"
+#include "Views/SSignalFlowGraphSendLabel.h"
+#include "Widgets/Input/SCheckBox.h"
+#include "Widgets/Layout/SBorder.h"
+#include "Widgets/SBoxPanel.h"
+#include "Widgets/Text/STextBlock.h"
+
+#if MUTE_SOLO_ENABLED
+#include "Audio/AudioDebug.h"
+#endif // MUTE_SOLO_ENABLED
+
+#define LOCTEXT_NAMESPACE "AudioInsights"
+
+namespace UE::Audio::Insights
+{
+	void SSignalFlowGraphNode::Construct(const SSignalFlowGraphNode::FArguments& InArgs)
+	{
+		using namespace SSignalFlowGraphStyle;
+
+		Entry = InArgs._Entry;
+		NodeDetailFilterSettings = InArgs._NodeDetailFilterSettings;
+		PositionInGraph = InArgs._PositionInGraph;
+		XPos = InArgs._XPos;
+		YPos = InArgs._YPos;
+		ZoomScaleFactor = InArgs._ZoomScaleFactor;
+		IsSelected = InArgs._IsSelected;
+		ShowNodeDetails = InArgs._ShowNodeDetails;
+		DisplayAmpPeakInDb = InArgs._DisplayAmpPeakInDb;
+		GraphOrientation = InArgs._GraphOrientation;
+		OnNodeSelected = InArgs._OnNodeSelected;
+
+#if WITH_EDITOR
+		OnOpenContextMenu = InArgs._OnOpenContextMenu;
+		GraphIsDraggingToMove = InArgs._GraphIsDraggingToMove;
+#endif // WITH_EDITOR
+
+		if (!Entry.IsValid())
+		{
+			return;
+		}
+
+		NodeEntryKey = Entry->GetSignalFlowEntryKey();
+
+#if MUTE_SOLO_ENABLED
+		const FSoftObjectPath AssetPath(Entry->Name);
+		if (AssetPath.IsValid())
+		{
+			AssetString = AssetPath.GetAssetName();
+			AssetFName = AssetPath.GetAssetFName();
+		}
+		else
+		{
+			AssetString = Entry->Name;
+			AssetFName = FName(Entry->Name);
+		}
+#endif // MUTE_SOLO_ENABLED
+
+		ChildSlot
+		[
+			SAssignNew(SelectedNodeBorder, SBorder)
+			.BorderImage(&GetNodeRoundedBoxBrush())
+			.HAlign(HAlign_Fill)
+			.VAlign(VAlign_Fill)
+			.BorderBackgroundColor_Lambda([this]()
+			{
+				return IsSelected.IsBound() && IsSelected.Get() ? FStyleColors::AccentOrange.GetSpecifiedColor() : FLinearColor::Transparent;
+			})
+			.Padding(SelectedNodeBorderThickness)
+			[
+				SAssignNew(NodeBorder, SBorder)
+				.BorderImage(&GetNodeRoundedBoxBrush())
+				.HAlign(HAlign_Fill)
+				.VAlign(VAlign_Fill)
+				.BorderBackgroundColor(GetNodeAccentColor(NodeEntryKey.EntryType))
+				.Padding(NodeBorderThickness)
+				.Content()
+				[
+					SNew(SBorder)
+					.BorderImage(&GetNodeRoundedBoxBrush())
+					.HAlign(HAlign_Fill)
+					.VAlign(VAlign_Fill)
+					.BorderBackgroundColor_Lambda([this]()
+					{
+						return (IsSelected.IsBound() && IsSelected.Get()) ? FStyleColors::Header : FStyleColors::Panel;
+					})
+					.Padding(0.0f, NodeBorderPadding, 0.0f, NodeBorderPadding)
+					.Content()
+					[
+						SNew(SVerticalBox)
+
+						+ SVerticalBox::Slot()
+						.Padding(NodeBorderPadding, 0.0f, NodeBorderPadding, 0.0f)
+						.AutoHeight()
+						[
+							CreateNodeName(InArgs)
+						]
+
+						+ SVerticalBox::Slot()
+						.Padding(NodeBorderPadding + NodeDetailIndentationPadding, 0.0f, 0.0f, 0.0f)
+						.AutoHeight()
+						[
+							SNew(SHorizontalBox)
+
+							+ SHorizontalBox::Slot()
+							.FillWidth(1.0f)
+							[
+								CreateNodeDetails()
+							]
+
+							+ SHorizontalBox::Slot()
+							.HAlign(HAlign_Right)
+							.VAlign(VAlign_Top)
+							.AutoWidth()
+							[
+								// In vertical mode, we want to ensure that we maintain padding on the RHS of the node when node details is on
+								// We do this by drawing a box with no content inside.
+								// In horizontal mode, this is replaced by SendLevelsVerticalBox, which is flush to the RHS
+								SNew(SVerticalBox)
+								.Visibility_Lambda([this]()
+								{
+									return GraphOrientation.Get() == EOrientation::Orient_Vertical && ShowNodeDetails.Get()
+																									? EVisibility::Visible
+																									: EVisibility::Collapsed;
+								})
+
+								+ SVerticalBox::Slot()
+								[
+									SNew(SBox)
+									.WidthOverride(NodeBorderPadding)
+								]
+							]
+
+							+ SHorizontalBox::Slot()
+							.HAlign(HAlign_Right)
+							.VAlign(VAlign_Top)
+							.AutoWidth()
+							[
+								CreateSendOutputsVerticalContainer()
+							]
+						]
+
+						+ SVerticalBox::Slot()
+						.HAlign(HAlign_Center)
+						.VAlign(VAlign_Bottom)
+						.AutoHeight()
+						[
+							CreateSendOutputsHorizontalContainer()
+						]
+					]
+				]
+			]
+		];
+	}
+
+	FReply SSignalFlowGraphNode::OnMouseButtonDown(const FGeometry& MyGeometry, const FPointerEvent& MouseEvent)
+	{
+		if (MouseEvent.GetEffectingButton() == EKeys::LeftMouseButton)
+		{
+			if (Entry.IsValid())
+			{
+				OnNodeSelected.ExecuteIfBound(IsSelected.IsBound() && IsSelected.Get() ? nullptr : Entry);
+			}
+
+			return FReply::Handled();
+		}
+		else
+		{
+			return SCompoundWidget::OnMouseButtonDown(MyGeometry, MouseEvent);
+		}
+	}
+
+#if WITH_EDITOR
+	FReply SSignalFlowGraphNode::OnMouseButtonUp(const FGeometry& MyGeometry, const FPointerEvent& MouseEvent)
+	{
+		if (MouseEvent.GetEffectingButton() == EKeys::RightMouseButton)
+		{
+			if (Entry.IsValid() && !GraphIsDraggingToMove.Get())
+			{
+				const TObjectPtr<UObject> ObjectPtr = Entry->GetObject();
+				if (ObjectPtr && ObjectPtr->IsAsset())
+				{
+					OnOpenContextMenu.ExecuteIfBound(MouseEvent, Entry);
+					return FReply::Handled();
+				}
+			}
+		}
+
+		return SCompoundWidget::OnMouseButtonUp(MyGeometry, MouseEvent);
+	}
+#endif // WITH_EDITOR
+
+	FVector2D SSignalFlowGraphNode::ComputeDesiredSize(float InLayoutScaleMultiplier) const
+	{
+		using namespace SSignalFlowGraphStyle;
+
+		int32 NumParamsVisible = 0;
+		if (NodeDetailFilterSettings)
+		{
+			if (bNodeDetailsAreExpanded)
+			{
+				const TSet<ESignalFlowNodeDetailParam> IgnoredParams { ESignalFlowNodeDetailParam::SendOutputVolume };
+				NumParamsVisible = NodeDetailFilterSettings->GetNumVisibleParams(IgnoredParams);
+			}
+			// We always show amplitude if the node is collapsed
+			else if (Entry.IsValid() && Entry->Amplitude.IsSet())
+			{
+				NumParamsVisible = 1;
+			}
+		}
+
+		const int32 NumOutputPinsShowing = SendLabels.Num();
+		const FVector2D Size = SCompoundWidget::ComputeDesiredSize(InLayoutScaleMultiplier);
+
+		if (GraphOrientation.Get() == EOrientation::Orient_Horizontal)
+		{
+			const int32 MaxNumParamRows = FMath::Max(NumParamsVisible, NumOutputPinsShowing);
+
+			const FVector2D ClampedSize = { FMath::Min(Size.X, NodeMaxWidth), FMath::Min(Size.Y, GetNodeMaxHeight(MaxNumParamRows, ShowNodeDetails.Get())) };
+			return ClampedSize;
+		}
+		else
+		{
+			const float PinSectionWidth = SendLevelsHorizontalBox.IsValid() ? SendLevelsHorizontalBox->GetDesiredSize().X : 0.0f;
+			const float OutputPinSectionHeight = SendLevelsHorizontalBox.IsValid() ? SendLevelsHorizontalBox->GetDesiredSize().Y : 0.0f;
+
+			// Make sure Max width/height calculations account for the bottom output pin row
+			const float ScaledMaxWidth = FMath::Max(NodeMaxWidth, PinSectionWidth);
+			const float TotalMaxHeight = GetNodeMaxHeight(NumParamsVisible, ShowNodeDetails.Get()) + OutputPinSectionHeight;
+
+			const FVector2D ClampedSize = { FMath::Min(Size.X, ScaledMaxWidth), FMath::Min(Size.Y, TotalMaxHeight) };
+			return ClampedSize;
+		}
+	}
+
+	FVector2D SSignalFlowGraphNode::GetLocalInputPos() const
+	{
+		using namespace SSignalFlowGraphStyle;
+
+		const FVector2D Size = GetDesiredSize();
+
+		if (GraphOrientation.IsBound() && GraphOrientation.Get() == EOrientation::Orient_Horizontal)
+		{
+			// Input pos = left center of node
+			return FVector2D{ CachedPositionInGraph.X + SelectedNodeBorderThickness, CachedPositionInGraph.Y + (Size.Y * 0.5f) } * ZoomScaleFactor.Get();
+		}
+		else
+		{
+			// Input pos = top center of node
+			return FVector2D{ CachedPositionInGraph.X + (Size.X * 0.5f), CachedPositionInGraph.Y + SelectedNodeBorderThickness } * ZoomScaleFactor.Get();
+		}
+	}
+	
+	FVector2D SSignalFlowGraphNode::GetLocalCenteredOutputPos() const
+	{
+		using namespace SSignalFlowGraphStyle;
+
+		const FVector2D Size = GetDesiredSize();
+
+		if (GraphOrientation.IsBound() && GraphOrientation.Get() == EOrientation::Orient_Horizontal)
+		{
+			// Out pos = right center of node
+			return FVector2D{ CachedPositionInGraph.X + Size.X - SelectedNodeBorderThickness, CachedPositionInGraph.Y + (Size.Y / 2.0f) } * ZoomScaleFactor.Get();
+		}
+		else
+		{
+			// Output pos = bottom center of node
+			return FVector2D{ CachedPositionInGraph.X + (Size.X / 2.0f), CachedPositionInGraph.Y + Size.Y - SelectedNodeBorderThickness } * ZoomScaleFactor.Get();
+		}
+	}
+
+	FVector2D SSignalFlowGraphNode::GetLocalSendLabelPos(const FSignalFlowEntryKey& ReceiverEntryKey, const int32 LabelIndex) const
+	{
+		using namespace SSignalFlowGraphStyle;
+
+		const FVector2D OutputPos = GetLocalCenteredOutputPos();
+		if (!NodeLabelHorizontalBox.IsValid() || !SendLabelsAreVisible())
+		{
+			return OutputPos;
+		}
+
+		const TSharedRef<SSignalFlowGraphSendLabel>* SendLabelWidget = SendLabels.Find(ReceiverEntryKey);
+		if (SendLabelWidget == nullptr)
+		{
+			return OutputPos;
+		}
+
+		constexpr float BorderThicknessYOffset = (SelectedNodeBorderThickness * 0.5f) + (NodeBorderThickness * 0.5f);
+
+		const FVector2D SendLabelWidgetSize = (*SendLabelWidget)->GetDesiredSize();
+
+		float XPosInGraph;
+		float YPosInGraph;
+
+		if (GraphOrientation.Get() == EOrientation::Orient_Horizontal)
+		{
+			XPosInGraph = OutputPos.X;
+
+			const float NodeLabelHeight = NodeLabelHorizontalBox->GetDesiredSize().Y + NodeBorderPadding;
+			const float YOffsetToLabelArea = CachedPositionInGraph.Y + BorderThicknessYOffset + NodeLabelHeight;
+			const float YOffsetInLabelList = ((SendLabelWidgetSize.Y + NodeDetailVerticalPadding) * LabelIndex) + NodeDetailVerticalPadding;
+			const float HalfLabelHeight = (SendLabelWidgetSize.Y * 0.5f);
+
+			YPosInGraph = (YOffsetToLabelArea + YOffsetInLabelList + HalfLabelHeight) * ZoomScaleFactor.Get();
+		}
+		else
+		{
+			YPosInGraph = OutputPos.Y;
+
+			const float SendLabelWidth = SendLabelWidgetSize.X + SendLabelPadding;
+			const float XOffset = (SendLabelWidth * (LabelIndex + 1)) - ((SendLabelWidth * SendLabels.Num()) * 0.5f) - (SendLabelWidth * 0.5f);
+			XPosInGraph = OutputPos.X + (XOffset * ZoomScaleFactor.Get());
+		}
+
+		const FVector2D SendLabelPosition(XPosInGraph, YPosInGraph);
+		return SendLabelPosition;
+	}
+
+	int32 SSignalFlowGraphNode::GetNumSendLabels() const
+	{
+		return SendLabels.Num();
+	}
+
+	void SSignalFlowGraphNode::UpdateNodeLayout()
+	{
+		using namespace SSignalFlowGraphStyle;
+
+		// Ensure border thickness does not get scaled to avoid rendering artifacts
+		if (SelectedNodeBorder.IsValid())
+		{
+			SelectedNodeBorder->SetPadding(SelectedNodeBorderThickness / ZoomScaleFactor.Get());
+		}
+
+		if (NodeBorder.IsValid())
+		{
+			NodeBorder->SetPadding(NodeBorderThickness / ZoomScaleFactor.Get());
+		}
+
+		// Update the node's position in the graph
+		CachedPositionInGraph = PositionInGraph.Get();
+	}
+
+	FVector2D SSignalFlowGraphNode::GetCachedPositionInGraph() const
+	{
+		return CachedPositionInGraph;
+	}
+
+	void SSignalFlowGraphNode::SetSendLabelWidgets(const TArray<FSignalFlowEntryKey>& ReceiverEntryKeys)
+	{
+		for (const FSignalFlowEntryKey& ReceiverKey : ReceiverEntryKeys)
+		{
+			AddSendLabelWidget(ReceiverKey);
+		}
+
+		for (auto Itr = SendLabels.CreateIterator(); Itr; ++Itr)
+		{
+			if (!ReceiverEntryKeys.Contains(Itr->Key))
+			{
+				if (SendLevelsVerticalBox.IsValid())
+				{
+					SendLevelsVerticalBox->RemoveSlot(Itr->Value);
+				}
+
+				if (SendLevelsHorizontalBox.IsValid())
+				{
+					SendLevelsHorizontalBox->RemoveSlot(Itr->Value);
+				}
+
+				Itr.RemoveCurrent();
+			}
+		}
+	}
+
+	void SSignalFlowGraphNode::AddSendLabelWidget(const FSignalFlowEntryKey& ReceiverEntryKey)
+	{
+		using namespace SSignalFlowGraphStyle;
+
+		if (!SendLevelsVerticalBox.IsValid() || !SendLevelsHorizontalBox.IsValid() || SendLabels.Contains(ReceiverEntryKey))
+		{
+			return;
+		}
+
+		TSharedPtr<SSignalFlowGraphSendLabel> SendLabelWidget = SNew(SSignalFlowGraphSendLabel)
+			.PixelSnappingMethod(EWidgetPixelSnapping::SnapToPixel)
+			.ToolTipText(LOCTEXT("SignalFlowGraph_NodeDetails_SendLabelTooltip", "Send Output Volume"))
+			.ReceiverEntryType(ReceiverEntryKey.EntryType)
+			.SendLevel_Lambda([this, ReceiverEntryKey]() -> TOptional<float>
+			{
+				if (Entry.IsValid())
+				{
+					const FSignalFlowOutputData* SendInfo = Entry->Outputs.Find(ReceiverEntryKey);
+					if (SendInfo && SendInfo->SendLevel.IsSet())
+					{
+						return SendInfo->SendLevel;
+					}
+					else if (Entry->LinkedSoundSources.Contains(ReceiverEntryKey))
+					{
+						// Source buses always have a send value of 1.0
+						return 1.0f;
+					}
+				}
+
+				return NullOpt;
+			})
+			.GraphOrientation(GraphOrientation)
+			.ZoomScaleFactor(ZoomScaleFactor);
+
+		SendLabels.Add(ReceiverEntryKey, SendLabelWidget.ToSharedRef());
+
+		SendLevelsVerticalBox->AddSlot()
+		.Padding(SendLabelPadding, NodeDetailVerticalPadding, 0.0f, 0.0f)
+		.VAlign(VAlign_Top)
+		.HAlign(HAlign_Right)
+		[
+			SendLabelWidget.ToSharedRef()
+		];
+
+		SendLevelsHorizontalBox->AddSlot()
+		.Padding(SendLabelPadding * 0.5f, NodeDetailVerticalPadding, SendLabelPadding * 0.5f, 0.0f)
+		.VAlign(VAlign_Center)
+		.HAlign(HAlign_Center)
+		[
+			SendLabelWidget.ToSharedRef()
+		];
+	}
+
+	void SSignalFlowGraphNode::RemoveSendLabelWidget(const FSignalFlowEntryKey& ReceiverEntryKey)
+	{
+		TSharedRef<SSignalFlowGraphSendLabel>* SendLabelWidget = SendLabels.Find(ReceiverEntryKey);
+		if (SendLabelWidget == nullptr)
+		{
+			return;
+		}
+
+		if (SendLevelsVerticalBox.IsValid())
+		{
+			SendLevelsVerticalBox->RemoveSlot(*SendLabelWidget);
+		}
+
+		if (SendLevelsHorizontalBox.IsValid())
+		{
+			SendLevelsHorizontalBox->RemoveSlot(*SendLabelWidget);
+		}
+
+		SendLabels.Remove(ReceiverEntryKey);
+	}
+
+	bool SSignalFlowGraphNode::ContainsSendLabel(const FSignalFlowEntryKey& ReceiverEntryKey) const
+	{
+		return SendLabels.Contains(ReceiverEntryKey);
+	}
+
+	TSharedRef<SHorizontalBox> SSignalFlowGraphNode::CreateNodeName(const SSignalFlowGraphNode::FArguments& InArgs)
+	{
+		using namespace SSignalFlowGraphStyle;
+
+		return SAssignNew(NodeLabelHorizontalBox, SHorizontalBox)
+
+			+ SHorizontalBox::Slot()
+			.AutoWidth()
+			.Padding(0.0f, 0.0f, 4.0f, 0.0f)
+			[
+				SNew(SImage)
+				.Image(FSlateStyle::Get().CreateIcon(InArgs._Entry->IconName).GetIcon())
+				.ToolTipText(InArgs._DisplayName)
+				.ColorAndOpacity(GetNodeAccentColor(NodeEntryKey.EntryType))
+			]
+
+			+ SHorizontalBox::Slot()
+			.FillWidth(1.0f)
+			[
+				SNew(STextBlock)
+				.Text(InArgs._DisplayName)
+				.ToolTipText(InArgs._DisplayName)
+				.Justification(ETextJustify::Left)
+				.OverflowPolicy(ETextOverflowPolicy::Ellipsis)
+				.Clipping(EWidgetClipping::ClipToBounds)
+			]
+
+#if MUTE_SOLO_ENABLED
+			+ SHorizontalBox::Slot()
+			.AutoWidth()
+			.Padding(4.0f, 0.0f, 0.0f, 0.0f)
+			[
+				SNew(SCheckBox)
+				.Style(&FSlateStyle::Get().GetWidgetStyle<FCheckBoxStyle>("SoundDashboard.MuteSoloButton"))
+				.Visibility_Lambda([this]()
+				{
+					return IsMuteSoloAvailableForNode() ? EVisibility::Visible : EVisibility::Collapsed;
+				})
+				.IsChecked(ECheckBoxState::Unchecked)
+				.OnCheckStateChanged_Lambda([this](ECheckBoxState)
+				{
+					ToggleMute();
+				})
+				[
+					SNew(SImage)
+					.Image(FSlateStyle::Get().CreateIcon("AudioInsights.Icon.SoundDashboard.Mute").GetIcon())
+					.ToolTipText(LOCTEXT("SignalFlowGraph_MuteNode", "Mute this node.\nMute will apply to all instances of an asset."))
+					.ColorAndOpacity_Lambda([this]()
+					{
+						return GetMuteSoloState().bIsMuted ? FStyleColors::Foreground : FStyleColors::White25;
+					})
+				]
+			]
+
+			+ SHorizontalBox::Slot()
+			.AutoWidth()
+			.Padding(0.0f, 0.0f, 4.0f, 0.0f)
+			[
+				SNew(SCheckBox)
+				.Style(&FSlateStyle::Get().GetWidgetStyle<FCheckBoxStyle>("SoundDashboard.MuteSoloButton"))
+				.Visibility_Lambda([this]()
+				{
+					return IsMuteSoloAvailableForNode() ? EVisibility::Visible : EVisibility::Collapsed;
+				})
+				.IsChecked(ECheckBoxState::Unchecked)
+				.OnCheckStateChanged_Lambda([this](ECheckBoxState)
+				{
+					ToggleSolo();
+				})
+				[
+					SNew(SImage)
+					.Image(FSlateStyle::Get().CreateIcon("AudioInsights.Icon.SoundDashboard.Solo").GetIcon())
+					.ToolTipText(LOCTEXT("SignalFlowGraph_SoloNode", "Solo this node.\nSolo will apply to all instances of an asset."))
+					.ColorAndOpacity_Lambda([this]()
+					{
+						return GetMuteSoloState().bIsSoloed ? FStyleColors::Foreground : FStyleColors::White25;
+					})
+				]
+			]
+#endif // MUTE_SOLO_ENABLED
+		
+			+ SHorizontalBox::Slot()
+			.AutoWidth()
+			.MinSize(12.0f)
+			.MaxSize(12.0f)
+			.Padding(4.0f, 0.0f, 0.0f, 0.0f)
+			[
+				SNew(SCheckBox)
+				.Style(&FSlateStyle::Get().GetWidgetStyle<FCheckBoxStyle>("AudioInsights.LeftExpandArrowToggle"))
+				.Visibility_Lambda([this]()
+				{ 
+					return NodeDetailsExpanderArrowIsEnabled() ? EVisibility::Visible : EVisibility::Collapsed;
+				})
+				.IsChecked_Lambda([this]()
+				{
+					return bNodeDetailsAreExpanded ? ECheckBoxState::Checked : ECheckBoxState::Unchecked;
+				})
+				.OnCheckStateChanged_Lambda([this](const ECheckBoxState CheckboxState)
+				{
+					bNodeDetailsAreExpanded = CheckboxState == ECheckBoxState::Checked;
+				})
+			];
+	}
+
+	TSharedRef<SVerticalBox> SSignalFlowGraphNode::CreateNodeDetails() const
+	{
+		using namespace SSignalFlowGraphStyle;
+
+		return SNew(SVerticalBox)
+				.Visibility_Lambda([this]() { return ShowNodeDetails.Get() ? EVisibility::Visible : EVisibility::Collapsed; })
+
+				+ SVerticalBox::Slot()
+				.Padding(0.0f, NodeDetailVerticalPadding, 0.0f, 0.0f)
+				.AutoHeight()
+				[
+					CreateNodeDetailsRow(ESignalFlowNodeDetailParam::Amplitude)
+				]
+			
+				+ SVerticalBox::Slot()
+				.Padding(0.0f, NodeDetailVerticalPadding, 0.0f, 0.0f)
+				.AutoHeight()
+				[
+					CreateNodeDetailsRow(ESignalFlowNodeDetailParam::Volume)
+				]
+
+				+ SVerticalBox::Slot()
+				.Padding(0.0f, NodeDetailVerticalPadding, 0.0f, 0.0f)
+				.AutoHeight()
+				[
+					CreateNodeDetailsRow(ESignalFlowNodeDetailParam::Pitch)
+				]
+
+				+ SVerticalBox::Slot()
+				.Padding(0.0f, NodeDetailVerticalPadding, 0.0f, 0.0f)
+				.AutoHeight()
+				[
+					CreateNodeDetailsRow(ESignalFlowNodeDetailParam::LPFFreq)
+				]
+			
+				+ SVerticalBox::Slot()
+				.Padding(0.0f, NodeDetailVerticalPadding, 0.0f, 0.0f)
+				.AutoHeight()
+				[
+					CreateNodeDetailsRow(ESignalFlowNodeDetailParam::HPFFreq)
+				]
+
+				+ SVerticalBox::Slot()
+				.Padding(0.0f, NodeDetailVerticalPadding, 0.0f, 0.0f)
+				.AutoHeight()
+				[
+					CreateNodeDetailsRow(ESignalFlowNodeDetailParam::Priority)
+				]
+
+				+ SVerticalBox::Slot()
+				.Padding(0.0f, NodeDetailVerticalPadding, 0.0f, 0.0f)
+				.AutoHeight()
+				[
+					CreateNodeDetailsRow(ESignalFlowNodeDetailParam::Distance)
+				]
+
+				+ SVerticalBox::Slot()
+				.Padding(0.0f, NodeDetailVerticalPadding, 0.0f, 0.0f)
+				.AutoHeight()
+				[
+					CreateNodeDetailsRow(ESignalFlowNodeDetailParam::Attenuation)
+				]
+
+				+ SVerticalBox::Slot()
+				.Padding(0.0f, NodeDetailVerticalPadding, 0.0f, 0.0f)
+				.AutoHeight()
+				[
+					CreateNodeDetailsRow(ESignalFlowNodeDetailParam::RelativeRenderCost)
+				];
+	}
+
+	TSharedRef<SHorizontalBox> SSignalFlowGraphNode::CreateNodeDetailsRow(const ESignalFlowNodeDetailParam ParamType) const
+	{
+		using namespace SSignalFlowGraphStyle;
+
+		return SNew(SHorizontalBox)
+				.Visibility_Lambda([this, ParamType]()
+				{
+					const bool bParamIsFiltered = NodeDetailFilterSettings && NodeDetailFilterSettings->GetParameterIsVisible(ParamType);
+					const bool bParamIsHiddenByExpander = !bNodeDetailsAreExpanded && ParamType != ESignalFlowNodeDetailParam::Amplitude;
+
+					return bParamIsFiltered && !bParamIsHiddenByExpander && Entry.IsValid() && GetNodeDetailParamValue(Entry, ParamType).IsSet() ? EVisibility::Visible : EVisibility::Collapsed;
+				})
+
+				+ SHorizontalBox::Slot()
+				.Padding(0.0f, 0.0f, 4.0f, 0.0f)
+				.AutoWidth()
+				[
+					SNew(SImage)
+					.Image(FSlateStyle::Get().CreateIcon("AudioInsights.Icon.SignalFlow.Pill").GetIcon())
+					.ColorAndOpacity(GetNodeDetailParamColor(ParamType))
+				]
+
+				+ SHorizontalBox::Slot()
+				.FillWidth(1.0f)
+				[
+					SNew(STextBlock)
+					.Text_Lambda([this, ParamType]()
+					{
+						const TOptional<float> ParamValueOpt = GetNodeDetailParamValue(Entry, ParamType);
+						if (!ParamValueOpt.IsSet())
+						{
+							return FText::GetEmpty();
+						}
+
+						if (ParamType == ESignalFlowNodeDetailParam::Amplitude && DisplayAmpPeakInDb.Get() == true)
+						{
+							return FText::Format(GetNodeDetailDisplayName(ParamType), AudioInsightsUtils::ConvertToDecibelsText(ParamValueOpt.GetValue()));
+						}
+
+						if (ParamType == ESignalFlowNodeDetailParam::Priority && SoundMessageUtils::IsPriorityValueSetToMax(ParamValueOpt.GetValue()))
+						{
+							return FText::Format(GetNodeDetailDisplayName(ParamType), LOCTEXT("SignalFlowGraph_NodeDetails_Max", "MAX"));
+						}
+
+						return FText::Format(GetNodeDetailDisplayName(ParamType), FText::AsNumber(ParamValueOpt.GetValue(), GetNodeDetailNumberFormat(ParamType)));
+					})
+					.Font(FCoreStyle::Get().GetFontStyle("SmallFont"))
+					.Justification(ETextJustify::Left)
+				];
+	}
+
+	TSharedRef<SVerticalBox> SSignalFlowGraphNode::CreateSendOutputsVerticalContainer()
+	{
+		SendLevelsVerticalBox = SNew(SVerticalBox)
+								.Visibility_Lambda([this]()
+								{
+									return GraphOrientation.Get() == EOrientation::Orient_Horizontal && SendLabelsAreVisible() 
+																	? EVisibility::Visible 
+																	: EVisibility::Collapsed;
+								});
+
+		return SendLevelsVerticalBox.ToSharedRef();
+	}
+
+	TSharedRef<SHorizontalBox> SSignalFlowGraphNode::CreateSendOutputsHorizontalContainer()
+	{
+		SendLevelsHorizontalBox = SNew(SHorizontalBox)
+									.Visibility_Lambda([this]()
+									{
+										return GraphOrientation.Get() == EOrientation::Orient_Vertical && SendLabelsAreVisible()
+																	? EVisibility::Visible
+																	: EVisibility::Collapsed;
+									});
+
+		return SendLevelsHorizontalBox.ToSharedRef();
+	}
+
+	bool SSignalFlowGraphNode::NodeDetailsExpanderArrowIsEnabled() const
+	{
+		using namespace SSignalFlowGraphStyle;
+
+		// The expander arrow should only be visible if a parameter other than amplitude is visible on the node
+		// This requires show node details to be switched on, the parameter filter to be enabled, and data to be available for that parameter
+		if (!ShowNodeDetails.Get() || !Entry.IsValid() || NodeDetailFilterSettings == nullptr)
+		{
+			return false;
+		}
+
+		const int32 NumVisibleParams = NodeDetailFilterSettings->GetNumVisibleParams();
+		if (NumVisibleParams <= 0)
+		{
+			return false;
+		}
+
+		if (NodeDetailFilterSettings->GetNumVisibleParams() == 1 && NodeDetailFilterSettings->GetParameterIsVisible(ESignalFlowNodeDetailParam::Amplitude))
+		{
+			return false;
+		}
+		
+		const int32 NumFilters = static_cast<uint8>(ESignalFlowNodeDetailParam::MAX);
+		for (uint8 FilterID = 0u; FilterID < NumFilters; ++FilterID)
+		{
+			const ESignalFlowNodeDetailParam Param = static_cast<ESignalFlowNodeDetailParam>(FilterID);
+			if (Param == ESignalFlowNodeDetailParam::Amplitude)
+			{
+				continue;
+			}
+
+			if (Param == ESignalFlowNodeDetailParam::SendOutputVolume)
+			{
+				if (NodeDetailFilterSettings->GetParameterIsVisible(Param))
+				{
+					for (const auto& [OutputKey, SendInfo] : Entry->Outputs)
+					{
+						if (SendInfo.SendLevel.IsSet())
+						{
+							return true;
+						}
+					}
+				}
+			}
+			else if (NodeDetailFilterSettings->GetParameterIsVisible(Param) && GetNodeDetailParamValue(Entry, Param).IsSet())
+			{
+				return true;
+			}
+		}
+
+		return false;
+	}
+
+	bool SSignalFlowGraphNode::SendLabelsAreVisible() const
+	{
+		if (!ShowNodeDetails.Get() || SendLabels.IsEmpty())
+		{
+			return false;
+		}
+
+		const bool bParamIsFiltered = NodeDetailFilterSettings && NodeDetailFilterSettings->GetParameterIsVisible(ESignalFlowNodeDetailParam::SendOutputVolume);
+		if (!bNodeDetailsAreExpanded || !bParamIsFiltered)
+		{
+			return false;
+		}
+
+		for (const auto& [ReceiverKey, Label] : SendLabels)
+		{
+			if (Label->ValueIsSet())
+			{
+				return true;
+			}
+		}
+
+		return false;
+	}
+
+#if MUTE_SOLO_ENABLED
+	void SSignalFlowGraphNode::ToggleMute()
+	{
+		if (FAudioDeviceManager* AudioDeviceManager = FAudioDeviceManager::Get())
+		{
+			AudioDeviceManager->GetDebugger().ToggleMuteSoundWave(AssetFName);
+		}
+	}
+
+	void SSignalFlowGraphNode::ToggleSolo()
+	{
+		if (FAudioDeviceManager* AudioDeviceManager = FAudioDeviceManager::Get())
+		{
+			AudioDeviceManager->GetDebugger().ToggleSoloSoundWave(AssetFName);
+		}
+	}
+
+	SSignalFlowGraphNode::FMuteSoloStates SSignalFlowGraphNode::GetMuteSoloState() const
+	{
+		FMuteSoloStates MuteSoloStates;
+		if (Entry.IsValid())
+		{
+			if (FAudioDeviceManager* AudioDeviceManager = FAudioDeviceManager::Get())
+			{
+				FString Reason;
+				AudioDeviceManager->GetDebugger().QuerySoloMuteSoundWave(AssetString, MuteSoloStates.bIsSoloed, MuteSoloStates.bIsMuted, Reason);
+			}
+		}
+
+		return MuteSoloStates;
+	}
+
+	bool SSignalFlowGraphNode::IsMuteSoloAvailableForNode() const
+	{
+		return NodeEntryKey.EntryType == ESignalFlowEntryType::SoundSource;
+	}
+#endif // MUTE_SOLO_ENABLED
+} // namespace UE::Audio::Insights
+
+#undef LOCTEXT_NAMESPACE
diff --git a/Engine/Plugins/AudioInsights/Source/AudioInsights/Private/Views/SSignalFlowGraphSendLabel.cpp b/Engine/Plugins/AudioInsights/Source/AudioInsights/Private/Views/SSignalFlowGraphSendLabel.cpp
new file mode 100644
--- /dev/null
+++ b/Engine/Plugins/AudioInsights/Source/AudioInsights/Private/Views/SSignalFlowGraphSendLabel.cpp
@@ -0,0 +1,166 @@
+// Copyright Epic Games, Inc. All Rights Reserved.
+
+#include "Views/SSignalFlowGraphSendLabel.h"
+
+#include "AudioInsightsStyle.h"
+#include "Internationalization/Text.h"
+#include "Styling/StyleColors.h"
+#include "Views/SSignalFlowGraphStyle.h"
+#include "Widgets/SBoxPanel.h"
+#include "Widgets/Images/SImage.h"
+#include "Widgets/SOverlay.h"
+#include "Widgets/Text/STextBlock.h"
+
+namespace UE::Audio::Insights
+{
+	void SSignalFlowGraphSendLabel::Construct(const SSignalFlowGraphSendLabel::FArguments& InArgs)
+	{
+		using namespace SSignalFlowGraphStyle;
+
+		SendLevel = InArgs._SendLevel;
+		GraphOrientation = InArgs._GraphOrientation;
+		ZoomScaleFactor = InArgs._ZoomScaleFactor;
+
+		ChildSlot
+		[
+			SNew(SOverlay)
+
+			// Horizontal flow label (text then icon)
+			+ SOverlay::Slot()
+			[
+				SNew(SHorizontalBox)
+				.Visibility_Lambda([this]()
+				{
+					return GraphOrientation.Get() == EOrientation::Orient_Horizontal && SendLevel.IsBound() && SendLevel.Get().IsSet() 
+																												? EVisibility::Visible 
+																												: EVisibility::Collapsed;
+				})
+
+				+ SHorizontalBox::Slot()
+				.FillContentWidth(1.0f)
+				.HAlign(HAlign_Right)
+				.VAlign(VAlign_Center)
+				[
+					CreateOutputLabelText(InArgs._ToolTipText.Get())
+				]
+
+				+ SHorizontalBox::Slot()
+				.FillContentWidth(1.0f)
+				.HAlign(HAlign_Right)
+				.VAlign(VAlign_Center)
+				[
+					CreatePinIconWidget(InArgs._ToolTipText.Get(), InArgs._ReceiverEntryType)
+				]
+			]
+
+			// Vertical flow label (text above icon)
+			+ SOverlay::Slot()
+			[
+				SNew(SVerticalBox)
+				.Visibility_Lambda([this]()
+				{
+					return GraphOrientation.Get() == EOrientation::Orient_Vertical && SendLevel.IsBound() && SendLevel.Get().IsSet() 
+																												? EVisibility::Visible 
+																												: EVisibility::Collapsed;
+				})
+
+				+ SVerticalBox::Slot()
+				.FillContentHeight(1.0f)
+				.HAlign(HAlign_Center)
+				.VAlign(VAlign_Bottom)
+				[
+					CreateOutputLabelText(InArgs._ToolTipText.Get())
+				]
+
+				+ SVerticalBox::Slot()
+				.FillContentHeight(1.0f)
+				.HAlign(HAlign_Center)
+				.VAlign(VAlign_Bottom)
+				[
+					CreatePinIconWidget(InArgs._ToolTipText.Get(), InArgs._ReceiverEntryType)
+				]
+			]
+		];
+	}
+
+	FVector2D SSignalFlowGraphSendLabel::ComputeDesiredSize(float InLayoutScaleMultiplier) const
+	{
+		using namespace SSignalFlowGraphStyle;
+
+		const FVector2D Size = SCompoundWidget::ComputeDesiredSize(InLayoutScaleMultiplier);
+		if (GraphOrientation.Get() == EOrientation::Orient_Horizontal)
+		{
+			return Size;
+		}
+		else
+		{
+			return FVector2D(SendLabelVerticalWidth, Size.Y);
+		}
+	}
+
+	bool SSignalFlowGraphSendLabel::ValueIsSet() const
+	{
+		return SendLevel.IsBound() && SendLevel.Get().IsSet();
+	}
+
+	TSharedRef<STextBlock> SSignalFlowGraphSendLabel::CreateOutputLabelText(const FText& ToolTipText) const
+	{
+		return SNew(STextBlock)
+				.Text_Lambda([this]()
+				{
+					if (!SendLevel.IsBound())
+					{
+						return FText::GetEmpty();
+					}
+
+					const TOptional<float> SendLevelOpt = SendLevel.Get();
+					if (!SendLevelOpt.IsSet())
+					{
+						return FText::GetEmpty();
+					}
+
+					return FText::AsNumber(SendLevelOpt.GetValue(), FSlateStyle::Get().GetSendLevelFormat());
+				})
+				.ToolTipText(ToolTipText)
+				.Justification(ETextJustify::Center)
+				.Font(FCoreStyle::Get().GetFontStyle("SmallFont"));
+	}
+
+	TSharedRef<SOverlay> SSignalFlowGraphSendLabel::CreatePinIconWidget(const FText& ToolTipText, const ESignalFlowEntryType ReceiverEntryType)
+	{
+		using namespace SSignalFlowGraphStyle;
+
+		return SNew(SOverlay)
+
+			+ SOverlay::Slot()
+			[
+				SAssignNew(OutputPinImage, SImage)
+				.Image(FSlateStyle::Get().CreateIcon("AudioInsights.Icon.SignalFlow.Disconnected").GetIcon())
+				.ColorAndOpacity(GetNodeAccentColor(ReceiverEntryType))
+			]
+
+			+ SOverlay::Slot()
+			[
+				SNew(SImage)
+				.Image(FSlateStyle::Get().CreateIcon("AudioInsights.Icon.SignalFlow.Connected").GetIcon())
+				.ToolTipText(ToolTipText)
+				.ColorAndOpacity_Lambda([this, ReceiverEntryType]()
+				{
+					const FSlateColor BaseColor = GetNodeAccentColor(ReceiverEntryType);
+
+					if (!SendLevel.IsBound())
+					{
+						return BaseColor;
+					}
+
+					const TOptional<float> SendLevelOpt = SendLevel.Get();
+					if (!SendLevelOpt.IsSet())
+					{
+						return BaseColor;
+					}
+
+					return FSlateColor(BaseColor.GetSpecifiedColor().CopyWithNewOpacity(SendLevelOpt.GetValue()));
+				})
+			];
+	}
+} // namespace UE::Audio::Insights
diff --git a/Engine/Plugins/AudioInsights/Source/AudioInsights/Private/Views/SSignalFlowGraphStyle.cpp b/Engine/Plugins/AudioInsights/Source/AudioInsights/Private/Views/SSignalFlowGraphStyle.cpp
new file mode 100644
--- /dev/null
+++ b/Engine/Plugins/AudioInsights/Source/AudioInsights/Private/Views/SSignalFlowGraphStyle.cpp
@@ -0,0 +1,234 @@
+// Copyright Epic Games, Inc. All Rights Reserved.
+
+#include "Views/SSignalFlowGraphStyle.h"
+
+#include "AudioInsightsStyle.h"
+#include "Brushes/SlateRoundedBoxBrush.h"
+#include "Internationalization/Text.h"
+#include "Messages/SignalFlowTraceMessages.h"
+#include "Styling/StyleColors.h"
+
+#define LOCTEXT_NAMESPACE "AudioInsights"
+
+namespace UE::Audio::Insights
+{
+	namespace SSignalFlowGraphStyle
+	{
+		const FSlateRoundedBoxBrush& GetNodeRoundedBoxBrush()
+		{
+			static const FSlateRoundedBoxBrush RoundedBoxBrush(FLinearColor::White, NodeRoundedBrushRadius);
+			return RoundedBoxBrush;
+		}
+
+		float GetNodeMaxHeight(const int32 NumParamsVisible, const bool bShowNodeDetails)
+		{
+			constexpr float NodeMaxHeight = 50.0f;
+			constexpr float NodeDetailsRowHeight = 24.0f;
+
+			return bShowNodeDetails ? NodeMaxHeight + (NodeDetailsRowHeight * NumParamsVisible) : NodeMaxHeight;
+		}
+
+		FSlateColor GetNodeAccentColor(const ESignalFlowEntryType NodeType)
+		{
+			switch (NodeType)
+			{
+				case ESignalFlowEntryType::SoundSource:
+					return FStyleColors::AccentBlue;
+				case ESignalFlowEntryType::AudioBus:
+					return FStyleColors::AccentYellow;
+				case ESignalFlowEntryType::Submix:
+					return FStyleColors::AccentGreen;
+				case ESignalFlowEntryType::OwnerObject:
+				case ESignalFlowEntryType::AudioDevice:
+				default:
+					return FStyleColors::Foreground;
+			}
+		}
+
+		FLinearColor GetConnectionColor(const float AmplitudeValue, const float WireSplineAmplitudePowerFactor, const bool bIsSourceBusConnection)
+		{
+			if (bIsSourceBusConnection)
+			{
+				constexpr float MaxAlpha = 0.5f;
+				return FLinearColor::LerpUsingHSV(FStyleColors::White25.GetSpecifiedColor(), GetNodeAccentColor(ESignalFlowEntryType::SoundSource).GetSpecifiedColor().CopyWithNewOpacity(MaxAlpha), FMath::Pow(AmplitudeValue, WireSplineAmplitudePowerFactor));
+			}
+
+			return FLinearColor::LerpUsingHSV(FStyleColors::Hover.GetSpecifiedColor(), FStyleColors::AccentWhite.GetSpecifiedColor(), FMath::Pow(AmplitudeValue, WireSplineAmplitudePowerFactor));
+		}
+
+		FVector2d GetSplineStartDirection(const EOrientation Orientation, const bool bIsSourceBusConnection)
+		{
+			const float SplineRadius = bIsSourceBusConnection ? SplineRadiusFlowingBackwards : SplineRadiusFlowingForwards;
+			return Orientation == EOrientation::Orient_Horizontal ? FVector2d{ SplineRadius, 0.0f } : FVector2d{ 0.0f, SplineRadius };
+		}
+
+		FVector2d GetSplineEndDirection(const EOrientation Orientation, const bool bIsSourceBusConnection)
+		{
+			const float SplineRadius = bIsSourceBusConnection ?  SplineRadiusFlowingBackwards : SplineRadiusFlowingForwards;
+			return Orientation == EOrientation::Orient_Horizontal ? FVector2d{ SplineRadius, 0.0f } : FVector2d{ 0.0f, SplineRadius };
+		}
+
+		TOptional<float> GetNodeDetailParamValue(const TSharedPtr<FSignalFlowDashboardEntry>& Entry, const ESignalFlowNodeDetailParam ParamType)
+		{
+			if (!Entry.IsValid())
+			{
+				return NullOpt;
+			}
+
+			switch (ParamType)
+			{
+				case ESignalFlowNodeDetailParam::Amplitude:
+					return Entry->Amplitude;
+
+				case ESignalFlowNodeDetailParam::Volume:
+					return Entry->Volume;
+
+				case ESignalFlowNodeDetailParam::Pitch:
+					return Entry->Pitch;
+
+				case ESignalFlowNodeDetailParam::LPFFreq:
+					return Entry->LPFFreq;
+
+				case ESignalFlowNodeDetailParam::HPFFreq:
+					return Entry->HPFFreq;
+
+				case ESignalFlowNodeDetailParam::Priority:
+					return Entry->Priority;
+
+				case ESignalFlowNodeDetailParam::Distance:
+					return Entry->Distance;
+
+				case ESignalFlowNodeDetailParam::Attenuation:
+					return Entry->Attenuation;
+
+				case ESignalFlowNodeDetailParam::RelativeRenderCost:
+					return Entry->RelativeRenderCost;
+
+				default:
+					ensureMsgf(false, TEXT("Unrecognized ESignalFlowNodeDetailParam value"));
+					break;
+			}
+
+			return NullOpt;
+		}
+
+		FSlateColor GetNodeDetailParamColor(const ESignalFlowNodeDetailParam ParamType)
+		{
+			switch (ParamType)
+			{
+				case ESignalFlowNodeDetailParam::Amplitude:
+					return FStyleColors::AccentGreen;
+
+				case ESignalFlowNodeDetailParam::Volume:
+					return FStyleColors::AccentYellow;
+
+				case ESignalFlowNodeDetailParam::Pitch:
+					return FStyleColors::AccentOrange;
+
+				case ESignalFlowNodeDetailParam::LPFFreq:
+					return FStyleColors::AccentBrown;
+
+				case ESignalFlowNodeDetailParam::HPFFreq:
+					return FStyleColors::AccentRed;
+
+				case ESignalFlowNodeDetailParam::Priority:
+					return FStyleColors::AccentPink;
+
+				case ESignalFlowNodeDetailParam::Distance:
+					return FStyleColors::AccentPurple;
+
+				case ESignalFlowNodeDetailParam::Attenuation:
+					return FStyleColors::AccentBlue;
+
+				case ESignalFlowNodeDetailParam::RelativeRenderCost:
+					return FStyleColors::AccentGray;
+
+				default:
+					ensureMsgf(false, TEXT("Unrecognized ESignalFlowNodeDetailParam value"));
+					break;
+			}
+
+			return FStyleColors::White;
+		}
+
+		FText GetNodeDetailDisplayName(const ESignalFlowNodeDetailParam ParamType)
+		{
+			switch (ParamType)
+			{
+				case ESignalFlowNodeDetailParam::Amplitude:
+					return LOCTEXT("SignalFlowGraph_NodeDetails_Amplitude", "Amp (peak): {0}");
+
+				case ESignalFlowNodeDetailParam::Volume:
+					return LOCTEXT("SignalFlowGraph_NodeDetails_Volume", "Volume: {0}");
+
+				case ESignalFlowNodeDetailParam::Pitch:
+					return LOCTEXT("SignalFlowGraph_NodeDetails_Pitch", "Pitch: {0}");
+
+				case ESignalFlowNodeDetailParam::LPFFreq:
+					return LOCTEXT("SignalFlowGraph_NodeDetails_LPFFreq", "LPF Freq (Hz): {0}");
+
+				case ESignalFlowNodeDetailParam::HPFFreq:
+					return LOCTEXT("SignalFlowGraph_NodeDetails_HPFFreq", "HPF Freq (Hz): {0}");
+
+				case ESignalFlowNodeDetailParam::Priority:
+					return LOCTEXT("SignalFlowGraph_NodeDetails_Priority", "Priority: {0}");
+
+				case ESignalFlowNodeDetailParam::Distance:
+					return LOCTEXT("SignalFlowGraph_NodeDetails_Distance", "Distance: {0}");
+
+				case ESignalFlowNodeDetailParam::Attenuation:
+					return LOCTEXT("SignalFlowGraph_NodeDetails_Attenuation", "Distance/Occlusion: {0}");
+
+				case ESignalFlowNodeDetailParam::RelativeRenderCost:
+					return LOCTEXT("SignalFlowGraph_NodeDetails_RekativeRenderCost", "Rel. Render Cost: {0}");
+
+				default:
+					ensureMsgf(false, TEXT("Unrecognized ESignalFlowNodeDetailParam value"));
+					break;
+			}
+
+			return FText::GetEmpty();
+		}
+
+		const FNumberFormattingOptions* GetNodeDetailNumberFormat(const ESignalFlowNodeDetailParam ParamType)
+		{
+			switch (ParamType)
+			{
+				case ESignalFlowNodeDetailParam::Amplitude:
+					return FSlateStyle::Get().GetLinearAmplitudeFloatFormat();
+
+				case ESignalFlowNodeDetailParam::Volume:
+					return FSlateStyle::Get().GetLinearVolumeFloatFormat();
+
+				case ESignalFlowNodeDetailParam::Pitch:
+					return FSlateStyle::Get().GetPitchFloatFormat();
+
+				case ESignalFlowNodeDetailParam::LPFFreq:
+					return FSlateStyle::Get().GetFreqFloatFormat();
+
+				case ESignalFlowNodeDetailParam::HPFFreq:
+					return FSlateStyle::Get().GetFreqFloatFormat();
+
+				case ESignalFlowNodeDetailParam::Priority:
+					return FSlateStyle::Get().GetLinearVolumeFloatFormat();
+
+				case ESignalFlowNodeDetailParam::Distance:
+					return FSlateStyle::Get().GetDefaultFloatFormat();
+
+				case ESignalFlowNodeDetailParam::Attenuation:
+					return FSlateStyle::Get().GetLinearVolumeFloatFormat();
+
+				case ESignalFlowNodeDetailParam::RelativeRenderCost:
+					return FSlateStyle::Get().GetRelativeRenderCostFormat();
+
+				default:
+					ensureMsgf(false, TEXT("Unrecognized ESignalFlowNodeDetailParam value"));
+					break;
+			}
+
+			return FSlateStyle::Get().GetDefaultFloatFormat();
+		}
+	}
+} // namespace UE::Audio::Insights
+
+#undef LOCTEXT_NAMESPACE
diff --git a/Engine/Plugins/AudioInsights/Source/AudioInsights/Private/Views/SSignalFlowGraphStyle.h b/Engine/Plugins/AudioInsights/Source/AudioInsights/Private/Views/SSignalFlowGraphStyle.h
new file mode 100644
--- /dev/null
+++ b/Engine/Plugins/AudioInsights/Source/AudioInsights/Private/Views/SSignalFlowGraphStyle.h
@@ -0,0 +1,64 @@
+// Copyright Epic Games, Inc. All Rights Reserved.
+#pragma once
+
+#include "Math/MathFwd.h"
+#include "Messages/SignalFlowEntryKey.h"
+#include "Misc/Optional.h"
+#include "Types/SlateEnums.h"
+
+struct FLinearColor;
+struct FNumberFormattingOptions;
+struct FSlateColor;
+struct FSlateRoundedBoxBrush;
+
+class FText;
+
+namespace UE::Audio::Insights
+{
+	struct FSignalFlowDashboardEntry;
+
+	namespace SSignalFlowGraphStyle
+	{
+		const FSlateRoundedBoxBrush& GetNodeRoundedBoxBrush();
+		float GetNodeMaxHeight(const int32 NumParamsVisible, const bool bShowNodeDetails);
+		FSlateColor GetNodeAccentColor(const ESignalFlowEntryType NodeType);
+		FLinearColor GetConnectionColor(const float AmplitudeValue, const float WireSplineAmplitudePowerFactor, const bool bIsSourceBusConnection);
+		FVector2d GetSplineStartDirection(const EOrientation Orientation, const bool bIsSourceBusConnection);
+		FVector2d GetSplineEndDirection(const EOrientation Orientation, const bool bIsSourceBusConnection);
+		TOptional<float> GetNodeDetailParamValue(const TSharedPtr<FSignalFlowDashboardEntry>& Entry, const ESignalFlowNodeDetailParam ParamType);
+		FSlateColor GetNodeDetailParamColor(const ESignalFlowNodeDetailParam ParamType);
+		FText GetNodeDetailDisplayName(const ESignalFlowNodeDetailParam ParamType);
+		const FNumberFormattingOptions* GetNodeDetailNumberFormat(const ESignalFlowNodeDetailParam ParamType);
+
+		constexpr float NodeMaxWidth = 256.0f;
+
+		constexpr float NodeRoundedBrushRadius = 12.0f;
+
+		constexpr float SplineRadiusFlowingForwards = 80.0f;
+		constexpr float SplineRadiusFlowingBackwards = 600.0f;
+		constexpr float SplineThicknessScalarMin = 1.0f;
+		constexpr float SplineThicknessDefault = 1.5f;
+		constexpr float SplineAnimatedWireThickness = 3.0f;
+		constexpr float SplineThicknessMin = 1.0f;
+		constexpr float SplineAnimPointsPerUnit = 0.2f;
+
+		constexpr float SelectedNodeBorderThickness = 3.0f;
+		constexpr float NodeBorderThickness = 2.0f;
+
+		constexpr float NodeBorderPadding = 12.0f;
+		constexpr float NodeDetailVerticalPadding = 8.0f;
+		constexpr float NodeDetailIndentationPadding = 20.0f;
+		constexpr float LargeBetweenNodesPadding = 48.0f;
+		constexpr float SmallBetweenNodePadding = 8.0f;
+
+		constexpr float SendLabelVerticalWidth = 20.0f;
+		constexpr float SendLabelPadding = 8.0f;
+
+		constexpr float AnimationSpeed = 15.0f;
+		constexpr float DefaultZoomScale = 1.0f;
+		constexpr float ZoomStepFactor = 0.1f;
+		constexpr float MinZoomScale = 0.5f;
+		constexpr float MaxZoomScale = 2.0f;
+		constexpr float ScrollAnimInRangeThreshold = (1.e-2f);
+	}
+} // namespace UE::Audio::Insights
\ No newline at end of file
diff --git a/Engine/Plugins/AudioInsights/Source/AudioInsights/Private/Views/SignalFlowDashboardViewFactory.cpp b/Engine/Plugins/AudioInsights/Source/AudioInsights/Private/Views/SignalFlowDashboardViewFactory.cpp
--- a/Engine/Plugins/AudioInsights/Source/AudioInsights/Private/Views/SignalFlowDashboardViewFactory.cpp
+++ b/Engine/Plugins/AudioInsights/Source/AudioInsights/Private/Views/SignalFlowDashboardViewFactory.cpp
@@ -12,6 +12,7 @@
 #include "SignalFlowEditorCommands.h"
 #include "Styling/StyleColors.h"
 #include "Views/SignalFlowNodes.h"
+#include "Views/SSignalFlowGraphStyle.h"
 #include "Widgets/Colors/SColorBlock.h"
 #include "Widgets/Input/SButton.h"
 #include "Widgets/Input/SComboButton.h"
diff --git a/Engine/Plugins/AudioInsights/Source/AudioInsights/Public/Views/SSignalFlowGraph.h b/Engine/Plugins/AudioInsights/Source/AudioInsights/Public/Views/SSignalFlowGraph.h
--- a/Engine/Plugins/AudioInsights/Source/AudioInsights/Public/Views/SSignalFlowGraph.h
+++ b/Engine/Plugins/AudioInsights/Source/AudioInsights/Public/Views/SSignalFlowGraph.h
@@ -1,9 +1,9 @@
 // Copyright Epic Games, Inc. All Rights Reserved.
 #pragma once
 
-#include "Delegates/Delegate.h"
 #include "Messages/SignalFlowEntryKey.h"
 #include "Settings/SignalFlowSettings.h"
+#include "SSignalFlowGraphNode.h"
 #include "Widgets/Layout/SScrollBox.h"
 #include "Widgets/SCanvas.h"
 #include "Widgets/SCompoundWidget.h"
@@ -13,18 +13,9 @@
 #endif // WITH_EDITOR
 
 #define UE_API AUDIOINSIGHTS_API
-#define MUTE_SOLO_ENABLED WITH_EDITOR && ENABLE_AUDIO_DEBUG
 
 struct FSignalFlowNodeDetailFilterSettings;
 
-class SBorder;
-class SHorizontalBox;
-class SImage;
-class SOverlay;
-class SScrollBox;
-class STextBlock;
-class SVerticalBox;
-
 #if WITH_EDITOR
 class IMenu;
 #endif // WITH_EDITOR
@@ -34,17 +25,6 @@ namespace UE::Audio::Insights
 	struct FSignalFlowDashboardEntry;
 	class ISignalFlowNode;
 
-	DECLARE_DELEGATE_OneParam(FOnSignalFlowNodeSelected, TSharedPtr<FSignalFlowDashboardEntry>);
-
-#if WITH_EDITOR
-	DECLARE_DELEGATE_TwoParams(FOnSignalFlowNodeContextMenuOpened, const FPointerEvent&, const TSharedPtr<FSignalFlowDashboardEntry>);
-#endif // WITH_EDITOR
-
-	namespace SSignalFlowGraphStyle
-	{
-		UE_API FSlateColor GetNodeAccentColor(const ESignalFlowEntryType NodeType);
-	}
-
 	class SSignalFlowScrollBox final : public SScrollBox
 	{
 	public:
@@ -57,181 +37,6 @@ namespace UE::Audio::Insights
 		}
 	};
 
-	class SSignalFlowGraphNode : public SCompoundWidget
-	{
-	public:
-		SLATE_BEGIN_ARGS(SSignalFlowGraphNode)
-			: _Entry(nullptr),
-			_NodeDetailFilterSettings(nullptr),
-			_DisplayName(),
-			_PositionInGraph(FVector2d::ZeroVector),
-			_XPos(0.0f),
-			_YPos(0.0f),
-			_ZoomScaleFactor(1.0f),
-			_IsSelected(false),
-			_ShowNodeDetails(false),
-			_DisplayAmpPeakInDb(true),
-#if WITH_EDITOR
-			_GraphIsDraggingToMove(false),
-#endif // WITH_EDITOR
-			_GraphOrientation(EOrientation::Orient_Vertical)
-		{
-		}
-
-		SLATE_ARGUMENT(TSharedPtr<FSignalFlowDashboardEntry>, Entry)
-		SLATE_ARGUMENT(FSignalFlowNodeDetailFilterSettings*, NodeDetailFilterSettings)
-		SLATE_ARGUMENT(FText, DisplayName)
-		SLATE_ATTRIBUTE(FVector2d, PositionInGraph)
-		SLATE_ATTRIBUTE(float, XPos)
-		SLATE_ATTRIBUTE(float, YPos)
-		SLATE_ATTRIBUTE(float, ZoomScaleFactor)
-		SLATE_ATTRIBUTE(bool, IsSelected)
-		SLATE_ATTRIBUTE(bool, ShowNodeDetails)
-		SLATE_ATTRIBUTE(bool, DisplayAmpPeakInDb)
-
-#if WITH_EDITOR
-		SLATE_ATTRIBUTE(bool, GraphIsDraggingToMove)
-#endif // WITH_EDITOR
-
-		SLATE_ATTRIBUTE(EOrientation, GraphOrientation)
-		SLATE_EVENT(FOnSignalFlowNodeSelected, OnNodeSelected)
-
-#if WITH_EDITOR
-		SLATE_EVENT(FOnSignalFlowNodeContextMenuOpened, OnOpenContextMenu)
-#endif // WITH_EDITOR
-
-		SLATE_END_ARGS()
-
-		UE_API void Construct(const FArguments& InArgs);
-
-		UE_API virtual FReply OnMouseButtonDown(const FGeometry& MyGeometry, const FPointerEvent& MouseEvent) override;
-
-#if WITH_EDITOR
-		UE_API virtual FReply OnMouseButtonUp(const FGeometry& MyGeometry, const FPointerEvent& MouseEvent) override;
-#endif // WITH_EDITOR
-
-		UE_API virtual FVector2D ComputeDesiredSize(float InLayoutScaleMultiplier) const override;
-
-		UE_API FVector2D GetLocalInputPos() const;
-		UE_API FVector2D GetLocalCenteredOutputPos() const;
-		UE_API FVector2D GetLocalSendLabelPos(const FSignalFlowEntryKey& ReceiverEntryKey, const int32 LabelIndex) const;
-		UE_API int32 GetNumSendLabels() const;
-
-		UE_API void UpdateNodeLayout();
-		UE_API FVector2D GetCachedPositionInGraph() const;
-
-		UE_API void SetSendLabelWidgets(const TArray<FSignalFlowEntryKey>& ReceiverEntryKeys);
-		UE_API void AddSendLabelWidget(const FSignalFlowEntryKey& ReceiverEntryKey);
-		UE_API void RemoveSendLabelWidget(const FSignalFlowEntryKey& ReceiverEntryKey);
-		UE_API bool ContainsSendLabel(const FSignalFlowEntryKey& ReceiverEntryKey) const;
-
-		const FSignalFlowEntryKey& GetNodeEntryKey() const { return NodeEntryKey; }
-		const TSharedPtr<FSignalFlowDashboardEntry> GetEntry() const { return Entry; }
-
-		float GetXPos() const { return XPos.Get(); }
-		float GetYPos() const { return YPos.Get(); }
-
-		void SetXPos(const float Pos) { XPos = Pos; }
-		void SetYPos(const float Pos) { YPos = Pos; }
-
-	private:
-		TSharedRef<SHorizontalBox> CreateNodeName(const SSignalFlowGraphNode::FArguments& InArgs);
-		TSharedRef<SVerticalBox> CreateNodeDetails() const;
-		TSharedRef<SHorizontalBox> CreateNodeDetailsRow(const ESignalFlowNodeDetailParam ParamType) const;
-		TSharedRef<SVerticalBox> CreateSendOutputsVerticalContainer();
-		TSharedRef<SHorizontalBox> CreateSendOutputsHorizontalContainer();
-
-		bool NodeDetailsExpanderArrowIsEnabled() const;
-		bool SendLabelsAreVisible() const;
-
-#if MUTE_SOLO_ENABLED
-		void ToggleMute();
-		void ToggleSolo();
-
-		struct FMuteSoloStates
-		{
-			bool bIsMuted = false;
-			bool bIsSoloed = false;
-		};
-
-		FMuteSoloStates GetMuteSoloState() const;
-		bool IsMuteSoloAvailableForNode() const;
-#endif // MUTE_SOLO_ENABLED
-
-		TSharedPtr<FSignalFlowDashboardEntry> Entry;
-		TAttribute<FVector2d> PositionInGraph;
-		TAttribute<float> XPos;
-		TAttribute<float> YPos;
-		TAttribute<float> ZoomScaleFactor;
-		TAttribute<bool> IsSelected;
-		TAttribute<bool> ShowNodeDetails;
-		TAttribute<bool> DisplayAmpPeakInDb;
-		TAttribute<EOrientation> GraphOrientation;
-
-		TMap<FSignalFlowEntryKey, TSharedRef<class SSignalFlowGraphSendLabel>> SendLabels;
-
-		FOnSignalFlowNodeSelected OnNodeSelected;
-
-#if WITH_EDITOR
-		FOnSignalFlowNodeContextMenuOpened OnOpenContextMenu;
-		TAttribute<bool> GraphIsDraggingToMove;
-#endif // WITH_EDITOR
-
-#if MUTE_SOLO_ENABLED
-		FName AssetFName;
-		FString AssetString;
-#endif // MUTE_SOLO_ENABLED
-
-		FVector2d CachedPositionInGraph = FVector2d::Zero();
-		FSignalFlowEntryKey NodeEntryKey;
-
-		FSignalFlowNodeDetailFilterSettings* NodeDetailFilterSettings = nullptr;
-
-		TSharedPtr<SBorder> SelectedNodeBorder;
-		TSharedPtr<SBorder> NodeBorder;
-		TSharedPtr<SHorizontalBox> NodeLabelHorizontalBox;
-		TSharedPtr<SVerticalBox> SendLevelsVerticalBox;
-		TSharedPtr<SHorizontalBox> SendLevelsHorizontalBox;
-
-		bool bNodeDetailsAreExpanded = true;
-	};
-
-	class SSignalFlowGraphSendLabel : public SCompoundWidget
-	{
-	public:
-		SLATE_BEGIN_ARGS(SSignalFlowGraphSendLabel)
-			: _ReceiverEntryType(ESignalFlowEntryType::AudioDevice),
-			_SendLevel(),
-			_GraphOrientation(EOrientation::Orient_Vertical),
-			_ZoomScaleFactor(1.0f)
-		{
-		}
-
-		SLATE_ARGUMENT(ESignalFlowEntryType, ReceiverEntryType)
-
-		SLATE_ATTRIBUTE(TOptional<float>, SendLevel)
-		SLATE_ATTRIBUTE(EOrientation, GraphOrientation)
-		SLATE_ATTRIBUTE(float, ZoomScaleFactor)
-
-		SLATE_END_ARGS()
-
-		UE_API void Construct(const FArguments& InArgs);
-
-		UE_API virtual FVector2D ComputeDesiredSize(float InLayoutScaleMultiplier) const override;
-
-		UE_API bool ValueIsSet() const;
-
-	private:
-		TSharedRef<STextBlock> CreateOutputLabelText(const FText& ToolTipText) const;
-		TSharedRef<SOverlay> CreatePinIconWidget(const FText& ToolTipText, const ESignalFlowEntryType ReceiverEntryType);
-
-		TAttribute<TOptional<float>> SendLevel;
-		TAttribute<EOrientation> GraphOrientation;
-		TAttribute<float> ZoomScaleFactor;
-
-		TSharedPtr<SImage> OutputPinImage;
-	};
-
 	class SSignalFlowGraph : public SCompoundWidget
 	{
 	public:
diff --git a/Engine/Plugins/AudioInsights/Source/AudioInsights/Public/Views/SSignalFlowGraphNode.h b/Engine/Plugins/AudioInsights/Source/AudioInsights/Public/Views/SSignalFlowGraphNode.h
new file mode 100644
--- /dev/null
+++ b/Engine/Plugins/AudioInsights/Source/AudioInsights/Public/Views/SSignalFlowGraphNode.h
@@ -0,0 +1,166 @@
+// Copyright Epic Games, Inc. All Rights Reserved.
+#pragma once
+
+#include "Delegates/Delegate.h"
+#include "Messages/SignalFlowEntryKey.h"
+#include "Widgets/SCompoundWidget.h"
+
+#define UE_API AUDIOINSIGHTS_API
+#define MUTE_SOLO_ENABLED WITH_EDITOR && ENABLE_AUDIO_DEBUG
+
+struct FSignalFlowNodeDetailFilterSettings;
+
+class SBorder;
+class SHorizontalBox;
+class SVerticalBox;
+
+namespace UE::Audio::Insights
+{
+	struct FSignalFlowDashboardEntry;
+
+	DECLARE_DELEGATE_OneParam(FOnSignalFlowNodeSelected, TSharedPtr<FSignalFlowDashboardEntry>);
+
+#if WITH_EDITOR
+	DECLARE_DELEGATE_TwoParams(FOnSignalFlowNodeContextMenuOpened, const FPointerEvent&, const TSharedPtr<FSignalFlowDashboardEntry>);
+#endif // WITH_EDITOR
+
+	class SSignalFlowGraphNode : public SCompoundWidget
+	{
+	public:
+		SLATE_BEGIN_ARGS(SSignalFlowGraphNode)
+			: _Entry(nullptr),
+			_NodeDetailFilterSettings(nullptr),
+			_DisplayName(),
+			_PositionInGraph(FVector2d::ZeroVector),
+			_XPos(0.0f),
+			_YPos(0.0f),
+			_ZoomScaleFactor(1.0f),
+			_IsSelected(false),
+			_ShowNodeDetails(false),
+			_DisplayAmpPeakInDb(true),
+#if WITH_EDITOR
+			_GraphIsDraggingToMove(false),
+#endif // WITH_EDITOR
+			_GraphOrientation(EOrientation::Orient_Vertical)
+		{
+		}
+
+		SLATE_ARGUMENT(TSharedPtr<FSignalFlowDashboardEntry>, Entry)
+		SLATE_ARGUMENT(FSignalFlowNodeDetailFilterSettings*, NodeDetailFilterSettings)
+		SLATE_ARGUMENT(FText, DisplayName)
+		SLATE_ATTRIBUTE(FVector2d, PositionInGraph)
+		SLATE_ATTRIBUTE(float, XPos)
+		SLATE_ATTRIBUTE(float, YPos)
+		SLATE_ATTRIBUTE(float, ZoomScaleFactor)
+		SLATE_ATTRIBUTE(bool, IsSelected)
+		SLATE_ATTRIBUTE(bool, ShowNodeDetails)
+		SLATE_ATTRIBUTE(bool, DisplayAmpPeakInDb)
+
+#if WITH_EDITOR
+		SLATE_ATTRIBUTE(bool, GraphIsDraggingToMove)
+#endif // WITH_EDITOR
+
+		SLATE_ATTRIBUTE(EOrientation, GraphOrientation)
+		SLATE_EVENT(FOnSignalFlowNodeSelected, OnNodeSelected)
+
+#if WITH_EDITOR
+		SLATE_EVENT(FOnSignalFlowNodeContextMenuOpened, OnOpenContextMenu)
+#endif // WITH_EDITOR
+
+		SLATE_END_ARGS()
+
+		UE_API void Construct(const FArguments& InArgs);
+
+		UE_API virtual FReply OnMouseButtonDown(const FGeometry& MyGeometry, const FPointerEvent& MouseEvent) override;
+
+#if WITH_EDITOR
+		UE_API virtual FReply OnMouseButtonUp(const FGeometry& MyGeometry, const FPointerEvent& MouseEvent) override;
+#endif // WITH_EDITOR
+
+		UE_API virtual FVector2D ComputeDesiredSize(float InLayoutScaleMultiplier) const override;
+
+		UE_API FVector2D GetLocalInputPos() const;
+		UE_API FVector2D GetLocalCenteredOutputPos() const;
+		UE_API FVector2D GetLocalSendLabelPos(const FSignalFlowEntryKey& ReceiverEntryKey, const int32 LabelIndex) const;
+		UE_API int32 GetNumSendLabels() const;
+
+		UE_API void UpdateNodeLayout();
+		UE_API FVector2D GetCachedPositionInGraph() const;
+
+		UE_API void SetSendLabelWidgets(const TArray<FSignalFlowEntryKey>& ReceiverEntryKeys);
+		UE_API void AddSendLabelWidget(const FSignalFlowEntryKey& ReceiverEntryKey);
+		UE_API void RemoveSendLabelWidget(const FSignalFlowEntryKey& ReceiverEntryKey);
+		UE_API bool ContainsSendLabel(const FSignalFlowEntryKey& ReceiverEntryKey) const;
+
+		const FSignalFlowEntryKey& GetNodeEntryKey() const { return NodeEntryKey; }
+		const TSharedPtr<FSignalFlowDashboardEntry> GetEntry() const { return Entry; }
+
+		float GetXPos() const { return XPos.Get(); }
+		float GetYPos() const { return YPos.Get(); }
+
+		void SetXPos(const float Pos) { XPos = Pos; }
+		void SetYPos(const float Pos) { YPos = Pos; }
+
+	private:
+		TSharedRef<SHorizontalBox> CreateNodeName(const SSignalFlowGraphNode::FArguments& InArgs);
+		TSharedRef<SVerticalBox> CreateNodeDetails() const;
+		TSharedRef<SHorizontalBox> CreateNodeDetailsRow(const ESignalFlowNodeDetailParam ParamType) const;
+		TSharedRef<SVerticalBox> CreateSendOutputsVerticalContainer();
+		TSharedRef<SHorizontalBox> CreateSendOutputsHorizontalContainer();
+
+		bool NodeDetailsExpanderArrowIsEnabled() const;
+		bool SendLabelsAreVisible() const;
+
+#if MUTE_SOLO_ENABLED
+		void ToggleMute();
+		void ToggleSolo();
+
+		struct FMuteSoloStates
+		{
+			bool bIsMuted = false;
+			bool bIsSoloed = false;
+		};
+
+		FMuteSoloStates GetMuteSoloState() const;
+		bool IsMuteSoloAvailableForNode() const;
+#endif // MUTE_SOLO_ENABLED
+
+		TSharedPtr<FSignalFlowDashboardEntry> Entry;
+		TAttribute<FVector2d> PositionInGraph;
+		TAttribute<float> XPos;
+		TAttribute<float> YPos;
+		TAttribute<float> ZoomScaleFactor;
+		TAttribute<bool> IsSelected;
+		TAttribute<bool> ShowNodeDetails;
+		TAttribute<bool> DisplayAmpPeakInDb;
+		TAttribute<EOrientation> GraphOrientation;
+
+		TMap<FSignalFlowEntryKey, TSharedRef<class SSignalFlowGraphSendLabel>> SendLabels;
+
+		FOnSignalFlowNodeSelected OnNodeSelected;
+
+#if WITH_EDITOR
+		FOnSignalFlowNodeContextMenuOpened OnOpenContextMenu;
+		TAttribute<bool> GraphIsDraggingToMove;
+#endif // WITH_EDITOR
+
+#if MUTE_SOLO_ENABLED
+		FName AssetFName;
+		FString AssetString;
+#endif // MUTE_SOLO_ENABLED
+		FVector2d CachedPositionInGraph = FVector2d::Zero();
+		FSignalFlowEntryKey NodeEntryKey;
+
+		FSignalFlowNodeDetailFilterSettings* NodeDetailFilterSettings = nullptr;
+
+		TSharedPtr<SBorder> SelectedNodeBorder;
+		TSharedPtr<SBorder> NodeBorder;
+		TSharedPtr<SHorizontalBox> NodeLabelHorizontalBox;
+		TSharedPtr<SVerticalBox> SendLevelsVerticalBox;
+		TSharedPtr<SHorizontalBox> SendLevelsHorizontalBox;
+
+		bool bNodeDetailsAreExpanded = true;
+	};
+} // namespace UE::Audio::Insights
+
+#undef UE_API
\ No newline at end of file
diff --git a/Engine/Plugins/AudioInsights/Source/AudioInsights/Public/Views/SSignalFlowGraphSendLabel.h b/Engine/Plugins/AudioInsights/Source/AudioInsights/Public/Views/SSignalFlowGraphSendLabel.h
new file mode 100644
--- /dev/null
+++ b/Engine/Plugins/AudioInsights/Source/AudioInsights/Public/Views/SSignalFlowGraphSendLabel.h
@@ -0,0 +1,52 @@
+// Copyright Epic Games, Inc. All Rights Reserved.
+#pragma once
+
+#include "Messages/SignalFlowEntryKey.h"
+#include "Widgets/SCompoundWidget.h"
+
+#define UE_API AUDIOINSIGHTS_API
+
+class SImage;
+class SOverlay;
+class STextBlock;
+
+namespace UE::Audio::Insights
+{
+	class SSignalFlowGraphSendLabel : public SCompoundWidget
+	{
+	public:
+		SLATE_BEGIN_ARGS(SSignalFlowGraphSendLabel)
+			: _ReceiverEntryType(ESignalFlowEntryType::AudioDevice),
+			_SendLevel(),
+			_GraphOrientation(EOrientation::Orient_Vertical),
+			_ZoomScaleFactor(1.0f)
+		{
+		}
+
+		SLATE_ARGUMENT(ESignalFlowEntryType, ReceiverEntryType)
+
+		SLATE_ATTRIBUTE(TOptional<float>, SendLevel)
+		SLATE_ATTRIBUTE(EOrientation, GraphOrientation)
+		SLATE_ATTRIBUTE(float, ZoomScaleFactor)
+
+		SLATE_END_ARGS()
+
+		UE_API void Construct(const FArguments& InArgs);
+
+		UE_API virtual FVector2D ComputeDesiredSize(float InLayoutScaleMultiplier) const override;
+
+		UE_API bool ValueIsSet() const;
+
+	private:
+		TSharedRef<STextBlock> CreateOutputLabelText(const FText& ToolTipText) const;
+		TSharedRef<SOverlay> CreatePinIconWidget(const FText& ToolTipText, const ESignalFlowEntryType ReceiverEntryType);
+
+		TAttribute<TOptional<float>> SendLevel;
+		TAttribute<EOrientation> GraphOrientation;
+		TAttribute<float> ZoomScaleFactor;
+
+		TSharedPtr<SImage> OutputPinImage;
+	};
+} // namespace UE::Audio::Insights
+
+#undef UE_API
\ No newline at end of file
