[AudioInsights] - Toolbar: Support Cache Pause button

#Synced-CL 51710841

diff --git a/Engine/Plugins/AudioInsights/Source/AudioInsights/Private/Views/AudioEventLogDashboardViewFactory.cpp b/Engine/Plugins/AudioInsights/Source/AudioInsights/Private/Views/AudioEventLogDashboardViewFactory.cpp
--- a/Engine/Plugins/AudioInsights/Source/AudioInsights/Private/Views/AudioEventLogDashboardViewFactory.cpp
+++ b/Engine/Plugins/AudioInsights/Source/AudioInsights/Private/Views/AudioEventLogDashboardViewFactory.cpp
@@ -658,33 +658,6 @@ namespace UE::Audio::Insights
 		{
 			constexpr FLinearColor StatTitleColor(1.0f, 1.0f, 1.0f, 0.5f);
 			SAssignNew(SettingsAreaWidget, SHorizontalBox)
-			
-			// Cache paused icon
-			+ SHorizontalBox::Slot()
-			.VAlign(VAlign_Center)
-			.AutoWidth()
-			[
-				SNew(SButton)
-				.ButtonStyle(&FAppStyle::Get().GetWidgetStyle<FButtonStyle>("SimpleButton"))
-				.ToolTipText(LOCTEXT("AudioEventLog_CacheStats_PausedTooltip", "Caching has been paused - Audio Insights is no longer collecting data from the current session.\n\nClick/ESC to resume."))
-				.Visibility_Lambda([]()
-				{
-					FAudioInsightsModule& AudioInsightsModule = FAudioInsightsModule::GetChecked();
-					return AudioInsightsModule.GetTimingViewExtender().GetMessageCacheAndProcessingStatus() == ECacheAndProcess::None ? EVisibility::Visible : EVisibility::Hidden;
-				})
-				.OnClicked_Lambda([this]()
-				{
-					ResetInspectTimestamp();
-					return FReply::Handled();
-				})
-				[
-					SNew(SImage)
-					.Image(FSlateStyle::Get().GetBrush("AudioInsights.Icon.Pause"))
-					.ColorAndOpacity(StatTitleColor)
-					.DesiredSizeOverride(FVector2D(16.0f, 16.0f))
-				]
-				
-			]
 
 			// Clear log Button
 			+ SHorizontalBox::Slot()
diff --git a/Engine/Plugins/AudioInsights/Source/AudioInsightsEditor/Private/AudioInsightsEditorDashboardFactory.cpp b/Engine/Plugins/AudioInsights/Source/AudioInsightsEditor/Private/AudioInsightsEditorDashboardFactory.cpp
--- a/Engine/Plugins/AudioInsights/Source/AudioInsightsEditor/Private/AudioInsightsEditorDashboardFactory.cpp
+++ b/Engine/Plugins/AudioInsights/Source/AudioInsightsEditor/Private/AudioInsightsEditorDashboardFactory.cpp
@@ -59,6 +59,9 @@ namespace UE::Audio::Insights
 
 		static const FText SaveSnapshotTooltipText = LOCTEXT("AudioDashboardEditor_SnapshotButton_TooltipText", "Saves a snapshot of the most recent trace data to a .utrace file. By default snapshots are limited to 32MB and can be accessed in Unreal Insights. To specify a custom snapshot size (in MB) launch the editor executable with the -tracetailmb=<size> option.");
 
+		static const FText CachePauseTooltipText = LOCTEXT("AudioDashboardEditor_PauseButton_TooltipText", "Click to pause monitoring.\n\nDepending on your cache settings, this may also stop the cache collecting data from the current session.");
+		static const FText CacheResumeTooltipText = LOCTEXT("AudioDashboardEditor_ResumeButton_TooltipText", "Click to resume monitoring.");
+
 		static const FText SaveSnapshotNotificationText = LOCTEXT("AudioDashboardEditor_SnapshotButton_NotificationText", "Audio Insights Snapshot saved.");
 		static const FText SaveSnapshotNotificationSubText = LOCTEXT("AudioDashboardEditor_SnapshotButton_NotificationSubText", "A snapshot .utrace with the most recent events has been saved to your trace server.");
 
@@ -186,6 +189,16 @@ namespace UE::Audio::Insights
 		TraceModule.OnOnlyTraceAudioChannelsStateChanged(bInOnlyTraceAudioChannels);
 	}
 
+	void FEditorDashboardFactory::OnTimingViewTimeMarkerChanged(double InTimeMarker)
+	{
+		bIsPaused = true;
+	}
+
+	void FEditorDashboardFactory::OnTimeControlMethodReset()
+	{
+		bIsPaused = false;
+	}
+
 	void FEditorDashboardFactory::CreateTraceBookmark()
 	{
 		using namespace EditorDashboardFactoryPrivate;
@@ -308,6 +321,22 @@ namespace UE::Audio::Insights
 				OnOnlyTraceAudioChannelsPropertyChangedHandle.Reset();
 			}
 		}
+
+		{
+			FAudioInsightsTimingViewExtender& TimingViewExtender = FAudioInsightsEditorModule::GetChecked().GetTimingViewExtender();
+	
+			if (OnTimingViewTimeMarkerChangedHandle.IsValid())
+			{
+				TimingViewExtender.OnTimingViewTimeMarkerChanged.Remove(OnTimingViewTimeMarkerChangedHandle);
+				OnTimingViewTimeMarkerChangedHandle.Reset();
+			}
+	
+			if (OnTimeControlMethodResetHandle.IsValid())
+			{
+				TimingViewExtender.OnTimeControlMethodReset.Remove(OnTimeControlMethodResetHandle);
+				OnTimeControlMethodResetHandle.Reset();
+			}
+		}
 	}
 
 	void FEditorDashboardFactory::BindCommands()
@@ -769,6 +798,55 @@ namespace UE::Audio::Insights
 					]
 				]
 
+				// Cache control: Pause button
+				+ SHorizontalBox::Slot()
+				.HAlign(HAlign_Left)
+				.VAlign(VAlign_Center)
+				.AutoWidth()
+				.Padding(8.0f, 0.0f, 0.0f, 0.0f)
+				[
+					SNew(SButton)
+					.ButtonStyle(&FAppStyle::Get().GetWidgetStyle<FButtonStyle>("SimpleButton"))
+					.ToolTipText_Lambda([this]()
+					{
+						return bIsPaused ? CacheResumeTooltipText : CachePauseTooltipText;
+					})
+					.OnClicked_Lambda([this]()
+					{
+						FAudioInsightsTimingViewExtender& TimingViewExtender = FAudioInsightsEditorModule::GetChecked().GetTimingViewExtender();
+
+						if (bIsPaused)
+						{
+							TimingViewExtender.ResumeTimeMarker();
+						}
+						else
+						{
+							const double PauseTime = FAudioInsightsEditorModule::GetChecked().GetCacheManager().GetCacheEndTimeStamp();
+							TimingViewExtender.PauseTimeMarker(PauseTime, ESystemControllingTimeMarker::External);
+						}
+
+						return FReply::Handled();
+					})
+					[
+						SNew(SBox)
+						.HeightOverride(16.0f)
+						[
+							SNew(SHorizontalBox)
+							+ SHorizontalBox::Slot()
+							.VAlign(VAlign_Center)
+							.AutoWidth()
+							[
+								SNew(SImage)
+								.Image(FSlateStyle::Get().GetBrush("AudioInsights.Icon.Pause"))
+								.ColorAndOpacity_Lambda([this]()
+								{
+									return bIsPaused ? FStyleColors::AccentBlue : FSlateColor::UseForeground();
+								})
+							]
+						]
+					]
+				]
+
 				// Cache stats
 				+ SHorizontalBox::Slot()
 				.Padding(12.0f, 0.0f, 0.0f, 0.0f)
@@ -1056,6 +1134,20 @@ namespace UE::Audio::Insights
 				OnOnlyTraceAudioChannelsPropertyChangedHandle = AudioInsightsEditorSettings->OnOnlyTraceAudioChannelsPropertyChanged.AddSP(this, &FEditorDashboardFactory::OnOnlyTraceAudioChannelsPropertyChanged);
 			}
 		}
+
+		{
+			FAudioInsightsTimingViewExtender& TimingViewExtender = FAudioInsightsEditorModule::GetChecked().GetTimingViewExtender();
+	
+			if (!OnTimingViewTimeMarkerChangedHandle.IsValid())
+			{
+				OnTimingViewTimeMarkerChangedHandle = TimingViewExtender.OnTimingViewTimeMarkerChanged.AddSP(this, &FEditorDashboardFactory::OnTimingViewTimeMarkerChanged);
+			}
+	
+			if (!OnTimeControlMethodResetHandle.IsValid())
+			{
+				OnTimeControlMethodResetHandle = TimingViewExtender.OnTimeControlMethodReset.AddSP(this, &FEditorDashboardFactory::OnTimeControlMethodReset);
+			}
+		}
 	}
 
 	TSharedRef<FTabManager::FLayout> FEditorDashboardFactory::GetDefaultTabLayout()
diff --git a/Engine/Plugins/AudioInsights/Source/AudioInsightsEditor/Private/AudioInsightsEditorDashboardFactory.h b/Engine/Plugins/AudioInsights/Source/AudioInsightsEditor/Private/AudioInsightsEditorDashboardFactory.h
--- a/Engine/Plugins/AudioInsights/Source/AudioInsightsEditor/Private/AudioInsightsEditorDashboardFactory.h
+++ b/Engine/Plugins/AudioInsights/Source/AudioInsightsEditor/Private/AudioInsightsEditorDashboardFactory.h
@@ -42,6 +42,10 @@ namespace UE::Audio::Insights
 
 		void OnOnlyTraceAudioChannelsPropertyChanged(const bool bInOnlyTraceAudioChannels);
 
+		void OnTimingViewTimeMarkerChanged(double InTimeMarker);
+
+		void OnTimeControlMethodReset();
+
 		void CreateTraceBookmark();
 
 		void RefreshDeviceSelector();
@@ -76,8 +80,13 @@ namespace UE::Audio::Insights
 
 		FDelegateHandle OnOnlyTraceAudioChannelsPropertyChangedHandle;
 
+		FDelegateHandle OnTimingViewTimeMarkerChangedHandle;
+		FDelegateHandle OnTimeControlMethodResetHandle;
+
 		bool bIsLayoutReseting = false;
 
+		bool bIsPaused { false };
+
 		TSharedPtr<FTabManager> DashboardTabManager;
 		TSharedPtr<FWorkspaceItem> DashboardWorkspace;
 
