[Audio Insights] - QOL: Support tooltip hovering over column cells

#Synced-CL 51710841

diff --git a/Engine/Plugins/AudioInsights/Source/AudioInsights/Private/Views/AudioBusDashboardViewFactory.cpp b/Engine/Plugins/AudioInsights/Source/AudioInsights/Private/Views/AudioBusDashboardViewFactory.cpp
--- a/Engine/Plugins/AudioInsights/Source/AudioInsights/Private/Views/AudioBusDashboardViewFactory.cpp
+++ b/Engine/Plugins/AudioInsights/Source/AudioInsights/Private/Views/AudioBusDashboardViewFactory.cpp
@@ -212,6 +212,7 @@ namespace UE::Audio::Insights
 				[
 					SNew(STextBlock)
 					.Text(ValueText)
+					.ToolTipText(ValueText)
 					.ColorAndOpacity(AudioBusDashboardEntry.AudioBusType == EAudioBusType::CodeGenerated ? AudioBusMeterTextColor : FSlateColor::UseForeground())
 					.OnDoubleClicked_Lambda([this, InRowData](const FGeometry& MyGeometry, const FPointerEvent& PointerEvent)
 					{
diff --git a/Engine/Plugins/AudioInsights/Source/AudioInsights/Private/Views/SoundDashboardViewFactory.cpp b/Engine/Plugins/AudioInsights/Source/AudioInsights/Private/Views/SoundDashboardViewFactory.cpp
--- a/Engine/Plugins/AudioInsights/Source/AudioInsights/Private/Views/SoundDashboardViewFactory.cpp
+++ b/Engine/Plugins/AudioInsights/Source/AudioInsights/Private/Views/SoundDashboardViewFactory.cpp
@@ -1142,6 +1142,7 @@ namespace UE::Audio::Insights
 						return SoundDashboardEntry.bIsCategory ? FCoreStyle::Get().GetFontStyle("BoldFont") : FCoreStyle::Get().GetFontStyle("NormalFont");
 					})
 					.Text(InValueText)
+					.ToolTipText(InValueText)
 					.ColorAndOpacity(FSlateColor::UseForeground())
 				]
 				// Number of children text
@@ -1277,6 +1278,7 @@ namespace UE::Audio::Insights
 				[
 					SNew(STextBlock)
 					.Text(ValueText)
+					.ToolTipText(ValueText)
 				]
 				// Number of children text (if it is not leaf item)
 				+ SHorizontalBox::Slot()
@@ -1316,6 +1318,10 @@ namespace UE::Audio::Insights
 					{
 						return ColumnData.GetDisplayValue(InRowData.Get());
 					})
+					.ToolTipText_Lambda([InRowData, ColumnData]()
+					{
+						return ColumnData.GetDisplayValue(InRowData.Get());
+					})
 				];
 		}
 	}
diff --git a/Engine/Plugins/AudioInsights/Source/AudioInsights/Private/Views/SubmixDashboardViewFactory.cpp b/Engine/Plugins/AudioInsights/Source/AudioInsights/Private/Views/SubmixDashboardViewFactory.cpp
--- a/Engine/Plugins/AudioInsights/Source/AudioInsights/Private/Views/SubmixDashboardViewFactory.cpp
+++ b/Engine/Plugins/AudioInsights/Source/AudioInsights/Private/Views/SubmixDashboardViewFactory.cpp
@@ -195,6 +195,7 @@ namespace UE::Audio::Insights
 				[
 					SNew(STextBlock)
 					.Text(ValueText)
+					.ToolTipText(ValueText)
 					.OnDoubleClicked_Lambda([this, InRowData](const FGeometry& MyGeometry, const FPointerEvent& PointerEvent)
 					{
 #if WITH_EDITOR
diff --git a/Engine/Plugins/AudioInsights/Source/AudioInsights/Private/Views/TableDashboardViewFactory.cpp b/Engine/Plugins/AudioInsights/Source/AudioInsights/Private/Views/TableDashboardViewFactory.cpp
--- a/Engine/Plugins/AudioInsights/Source/AudioInsights/Private/Views/TableDashboardViewFactory.cpp
+++ b/Engine/Plugins/AudioInsights/Source/AudioInsights/Private/Views/TableDashboardViewFactory.cpp
@@ -103,6 +103,11 @@ namespace UE::Audio::Insights
 
 				return ValueText;
 			})
+			.ToolTipText_Lambda([this, InRowData, Column]()
+			{
+				const FColumnData& ColumnData = GetColumns()[Column];
+				return ColumnData.GetDisplayValue.IsSet() ? ColumnData.GetDisplayValue(InRowData.Get()) : FText::GetEmpty();
+			})
 
 			: SNullWidget::NullWidget
 		];
@@ -481,6 +486,11 @@ namespace UE::Audio::Insights
 
 				return ValueText;
 			})
+			.ToolTipText_Lambda([this, InRowData, Column]()
+			{
+				const FColumnData& ColumnData = GetColumns()[Column];
+				return ColumnData.GetDisplayValue.IsSet() ? ColumnData.GetDisplayValue(InRowData.Get()) : FText::GetEmpty();
+			})
 			.ColorAndOpacity_Lambda([this, InRowData]()
 			{
 				return GetRowColor(InRowData.ToSharedPtr());
