[Audio Insights] - Event Log: Fix Event Log Right-Click context menu options

#Synced-CL 52340544

diff --git a/Engine/Plugins/AudioInsights/Source/AudioInsights/Private/Views/AudioEventLogDashboardViewFactory.cpp b/Engine/Plugins/AudioInsights/Source/AudioInsights/Private/Views/AudioEventLogDashboardViewFactory.cpp
--- a/Engine/Plugins/AudioInsights/Source/AudioInsights/Private/Views/AudioEventLogDashboardViewFactory.cpp
+++ b/Engine/Plugins/AudioInsights/Source/AudioInsights/Private/Views/AudioEventLogDashboardViewFactory.cpp
@@ -753,21 +753,31 @@ namespace UE::Audio::Insights
 
 	TSharedPtr<SWidget> FAudioEventLogDashboardViewFactory::OnConstructContextMenu()
 	{
+		if (DataViewEntries.IsEmpty())
+		{
+			return nullptr;
+		}
+
 		constexpr bool bShouldCloseWindowAfterMenuSelection = true;
 		FMenuBuilder MenuBuilder(bShouldCloseWindowAfterMenuSelection, CommandList);
 
 #if WITH_EDITOR
-		const FAudioEventLogEditorCommands& Commands = FAudioEventLogEditorCommands::Get();
+		const bool bHasSelectedItem = FilteredEntriesListView.IsValid() && (FilteredEntriesListView->GetSelectedItems().Num() > 0);
 
-		MenuBuilder.BeginSection("EventLogDashboardActions", LOCTEXT("EventLogDashboard_Actions_HeaderText", "Editor Options"));
+		if (bHasSelectedItem)
 		{
-			MenuBuilder.AddMenuEntry(Commands.GetBrowseCommand(), NAME_None, TAttribute<FText>(), TAttribute<FText>(), FSlateStyle::Get().CreateIcon("AudioInsights.Icon.SoundDashboard.Browse"));
-			MenuBuilder.AddMenuEntry(Commands.GetEditCommand(), NAME_None, TAttribute<FText>(), TAttribute<FText>(), FSlateStyle::Get().CreateIcon("AudioInsights.Icon.SoundDashboard.Edit"));
-		}
+			const FAudioEventLogEditorCommands& Commands = FAudioEventLogEditorCommands::Get();
 
-		MenuBuilder.EndSection();
+			MenuBuilder.BeginSection("EventLogDashboardActions", LOCTEXT("EventLogDashboard_Actions_HeaderText", "Editor Options"));
+			{
+				MenuBuilder.AddMenuEntry(Commands.GetBrowseCommand(), NAME_None, TAttribute<FText>(), TAttribute<FText>(), FSlateStyle::Get().CreateIcon("AudioInsights.Icon.SoundDashboard.Browse"));
+				MenuBuilder.AddMenuEntry(Commands.GetEditCommand(), NAME_None, TAttribute<FText>(), TAttribute<FText>(), FSlateStyle::Get().CreateIcon("AudioInsights.Icon.SoundDashboard.Edit"));
+			}
 
-		MenuBuilder.AddSeparator();
+			MenuBuilder.EndSection();
+
+			MenuBuilder.AddSeparator();
+		}
 #endif // WITH_EDITOR
 
 		MenuBuilder.AddMenuEntry(
