[Audio Insights] FVirtualLoopTraceProvider::ProcessMessages: added missing message timestamp in ProcessMessageQueue<FVirtualLoopUpdateMessage>

#Synced-CL 52293379

diff --git a/Engine/Plugins/AudioInsights/Source/AudioInsights/Private/Providers/VirtualLoopTraceProvider.cpp b/Engine/Plugins/AudioInsights/Source/AudioInsights/Private/Providers/VirtualLoopTraceProvider.cpp
--- a/Engine/Plugins/AudioInsights/Source/AudioInsights/Private/Providers/VirtualLoopTraceProvider.cpp
+++ b/Engine/Plugins/AudioInsights/Source/AudioInsights/Private/Providers/VirtualLoopTraceProvider.cpp
@@ -179,6 +179,9 @@ namespace UE::Audio::Insights
 			if (OutEntry)
 			{
 				FVirtualLoopDashboardEntry& EntryRef = *OutEntry->Get();
+
+				EntryRef.Timestamp = Msg.Timestamp;
+
 				EntryRef.PlaybackTime = Msg.PlaybackTime;
 				EntryRef.TimeVirtualized = Msg.TimeVirtualized;
 				EntryRef.UpdateInterval = Msg.UpdateInterval;
