[AudioInsights] - Toolbar: Support Settings button

#Synced-CL 52377736

diff --git a/Engine/Plugins/AudioInsights/Source/AudioInsightsEditor/Private/AudioInsightsEditorDashboardFactory.cpp b/Engine/Plugins/AudioInsights/Source/AudioInsightsEditor/Private/AudioInsightsEditorDashboardFactory.cpp
--- a/Engine/Plugins/AudioInsights/Source/AudioInsightsEditor/Private/AudioInsightsEditorDashboardFactory.cpp
+++ b/Engine/Plugins/AudioInsights/Source/AudioInsightsEditor/Private/AudioInsightsEditorDashboardFactory.cpp
@@ -62,6 +62,10 @@ namespace UE::Audio::Insights
 		static const FText CachePauseTooltipText = LOCTEXT("AudioDashboardEditor_PauseButton_TooltipText", "Click to pause monitoring.\n\nDepending on your cache settings, this may also stop the cache collecting data from the current session.");
 		static const FText CacheResumeTooltipText = LOCTEXT("AudioDashboardEditor_ResumeButton_TooltipText", "Click to resume monitoring.");
 
+		static const FText CacheSettingsHeaderText = LOCTEXT("AudioDashboardEditor_CacheSettings_HeaderText", "Cache Settings");
+		static const FText CacheSettingsShowStatsText = LOCTEXT("AudioDashboardEditor_CacheSettings_ShowCacheStats", "Show Cache Stats");
+		static const FText CacheSettingsShowStatsTooltipText = LOCTEXT("AudioDashboardEditor_CacheSettings_ShowCacheStatsTooltip", "Shows/Hides Cache Stats");
+
 		static const FText SaveSnapshotNotificationText = LOCTEXT("AudioDashboardEditor_SnapshotButton_NotificationText", "Audio Insights Snapshot saved.");
 		static const FText SaveSnapshotNotificationSubText = LOCTEXT("AudioDashboardEditor_SnapshotButton_NotificationSubText", "A snapshot .utrace with the most recent events has been saved to your trace server.");
 
@@ -74,6 +78,8 @@ namespace UE::Audio::Insights
 
 		static const FText WorldFilterText = LOCTEXT("AudioDashboardEditor_WorldFilter_Text", "World Filter:");
 		static const FText WorldFilterTooltipText = LOCTEXT("AudioDashboardEditor_WorldFilter_TooltipText", "Select world(s) to monitor (worlds may share audio output).");
+
+		static const FText SettingsTooltipText = LOCTEXT("AudioDashboardEditor_Settings_TooltipText", "Settings");
 		
 		void ShowNotification(const FText& InText, const FText& InSubText)
 		{
@@ -857,6 +863,10 @@ namespace UE::Audio::Insights
 					.Justification(ETextJustify::Center)
 					.ColorAndOpacity(CacheStatsColor)
 					.Text(LOCTEXT("AudioDashboardEditor_CacheStats_Title", "Cache Size: "))
+					.Visibility_Lambda([this]()
+					{
+						return bIsCacheStatsVisible ? EVisibility::Visible : EVisibility::Collapsed;
+					})
 				]
 
 				+ SHorizontalBox::Slot()
@@ -875,6 +885,10 @@ namespace UE::Audio::Insights
 											, FText::AsNumber(CacheSizeMb, FSlateStyle::Get().GetMemoryFormat())
 											, FSlateStyle::Get().FormatMemoryAsMegabytes(MaxCacheSizeMb));
 					})
+					.Visibility_Lambda([this]()
+					{
+						return bIsCacheStatsVisible ? EVisibility::Visible : EVisibility::Collapsed;
+					})
 				]
 
 				+ SHorizontalBox::Slot()
@@ -885,6 +899,10 @@ namespace UE::Audio::Insights
 					.Justification(ETextJustify::Center)
 					.ColorAndOpacity(CacheStatsColor)
 					.Text(LOCTEXT("AudioDashboardEditor_CacheStats_Duration", " : Duration: "))
+					.Visibility_Lambda([this]()
+					{
+						return bIsCacheStatsVisible ? EVisibility::Visible : EVisibility::Collapsed;
+					})
 				]
 
 				+ SHorizontalBox::Slot()
@@ -901,6 +919,10 @@ namespace UE::Audio::Insights
 						return FText::Format(LOCTEXT("AudioDashboardEditor_CacheStats_DurationValue", "{0}s")
 											, FText::AsNumber(CacheDuration, FSlateStyle::Get().GetShortTimeFormat()));
 					})
+					.Visibility_Lambda([this]()
+					{
+						return bIsCacheStatsVisible ? EVisibility::Visible : EVisibility::Collapsed;
+					})
 				]
 
 				// World filter label
@@ -950,6 +972,15 @@ namespace UE::Audio::Insights
 						})
 					]
 				]
+
+				// Settings
+				+ SHorizontalBox::Slot()
+				.AutoWidth()
+				.Padding(12.0f, 0.0f, 0.0f, 0.0f)
+				.HAlign(HAlign_Right)
+				[
+					CreateSettingsButtonWidget()
+				]
 			];
 	}
 
@@ -1085,6 +1116,66 @@ namespace UE::Audio::Insights
 			];
 	}
 
+	TSharedRef<SWidget> FEditorDashboardFactory::CreateSettingsButtonWidget()
+	{
+		if (OnGetSettingsMenuContent() == SNullWidget::NullWidget)
+		{
+			return SNullWidget::NullWidget;
+		}
+
+		return SNew(SComboButton)
+			.ComboButtonStyle(&FAppStyle::Get().GetWidgetStyle<FComboButtonStyle>("SimpleComboButton"))
+			.OnGetMenuContent(this, &FEditorDashboardFactory::OnGetSettingsMenuContent)
+			.MenuPlacement(EMenuPlacement::MenuPlacement_ComboBoxRight)
+			.HasDownArrow(false)
+			.ToolTipText(EditorDashboardFactoryPrivate::SettingsTooltipText)
+			.ButtonContent()
+			[
+				SNew(SHorizontalBox)
+
+				+ SHorizontalBox::Slot()
+				.AutoWidth()
+				.VAlign(VAlign_Center)
+				.Padding(5.0, 0.0f)
+				[
+					SNew(SImage)
+					.ColorAndOpacity(FSlateColor::UseForeground())
+					.Image(FSlateStyle::Get().GetBrush("AudioInsights.Icon.Settings"))
+				]
+			];
+	}
+
+	TSharedRef<SWidget> FEditorDashboardFactory::OnGetSettingsMenuContent()
+	{
+		FMenuBuilder MenuBuilder(false /*bShouldCloseWindowAfterMenuSelection*/, CommandList);
+
+		MenuBuilder.BeginSection("AudioDashboardEditorSettingActions", EditorDashboardFactoryPrivate::CacheSettingsHeaderText);
+		{
+			MenuBuilder.AddMenuEntry(
+				EditorDashboardFactoryPrivate::CacheSettingsShowStatsText,
+				EditorDashboardFactoryPrivate::CacheSettingsShowStatsTooltipText,
+				FSlateIcon(),
+				FUIAction(
+					FExecuteAction::CreateLambda([this]()
+					{
+						bIsCacheStatsVisible = !bIsCacheStatsVisible;
+					}),
+					FCanExecuteAction(),
+					FIsActionChecked::CreateLambda([this]()
+					{
+						return bIsCacheStatsVisible;
+					})
+				),
+				NAME_None,
+				EUserInterfaceActionType::ToggleButton
+			);
+		}
+
+		MenuBuilder.EndSection();
+
+		return MenuBuilder.MakeWidget();
+	}
+
 	void FEditorDashboardFactory::SetActiveDeviceId(const ::Audio::FDeviceId InDeviceId)
 	{
 		const TObjectPtr<const UAudioInsightsEditorSettings> AudioInsightsEditorSettings = GetDefault<UAudioInsightsEditorSettings>();
diff --git a/Engine/Plugins/AudioInsights/Source/AudioInsightsEditor/Private/AudioInsightsEditorDashboardFactory.h b/Engine/Plugins/AudioInsights/Source/AudioInsightsEditor/Private/AudioInsightsEditorDashboardFactory.h
--- a/Engine/Plugins/AudioInsights/Source/AudioInsightsEditor/Private/AudioInsightsEditorDashboardFactory.h
+++ b/Engine/Plugins/AudioInsights/Source/AudioInsightsEditor/Private/AudioInsightsEditorDashboardFactory.h
@@ -58,6 +58,9 @@ namespace UE::Audio::Insights
 
 		TSharedRef<SWidget> MakeTraceModeButton(const ETraceMode InTraceButtonMode);
 
+		TSharedRef<SWidget> CreateSettingsButtonWidget();
+		TSharedRef<SWidget> OnGetSettingsMenuContent();
+
 		void SetActiveDeviceId(const ::Audio::FDeviceId InDeviceId);
 		void InitDelegates();
 		TSharedRef<FTabManager::FLayout> GetDefaultTabLayout();
@@ -86,6 +89,7 @@ namespace UE::Audio::Insights
 		bool bIsLayoutReseting = false;
 
 		bool bIsPaused { false };
+		bool bIsCacheStatsVisible { true };
 
 		TSharedPtr<FTabManager> DashboardTabManager;
 		TSharedPtr<FWorkspaceItem> DashboardWorkspace;
