[AudioInsights] - Toolbar (Editor): Support Editor Preferences Shortcut Button

#Synced-CL 52441802

diff --git a/Engine/Plugins/AudioInsights/Source/AudioInsightsEditor/Private/AudioInsightsEditorDashboardFactory.cpp b/Engine/Plugins/AudioInsights/Source/AudioInsightsEditor/Private/AudioInsightsEditorDashboardFactory.cpp
--- a/Engine/Plugins/AudioInsights/Source/AudioInsightsEditor/Private/AudioInsightsEditorDashboardFactory.cpp
+++ b/Engine/Plugins/AudioInsights/Source/AudioInsightsEditor/Private/AudioInsightsEditorDashboardFactory.cpp
@@ -16,6 +16,7 @@
 #include "Framework/Notifications/NotificationManager.h"
 #include "Internationalization/Text.h"
 #include "IPropertyTypeCustomization.h"
+#include "ISettingsModule.h"
 #include "Templates/SharedPointer.h"
 #include "Widgets/Input/SButton.h"
 #include "Widgets/Input/SCheckBox.h"
@@ -80,6 +81,7 @@ namespace UE::Audio::Insights
 		static const FText WorldFilterTooltipText = LOCTEXT("AudioDashboardEditor_WorldFilter_TooltipText", "Select world(s) to monitor (worlds may share audio output).");
 
 		static const FText SettingsTooltipText = LOCTEXT("AudioDashboardEditor_Settings_TooltipText", "Settings");
+		static const FText EditorPreferencesTooltipText = LOCTEXT("AudioDashboardEditor_EditorPreferences_TooltipText", "Opens Audio Insights Editor Preferences.");
 		
 		void ShowNotification(const FText& InText, const FText& InSubText)
 		{
@@ -973,11 +975,19 @@ namespace UE::Audio::Insights
 					]
 				]
 
-				// Settings
+				// Editor Preferences
 				+ SHorizontalBox::Slot()
 				.AutoWidth()
 				.Padding(12.0f, 0.0f, 0.0f, 0.0f)
 				.HAlign(HAlign_Right)
+				[
+					CreateEditorPreferencesButtonWidget()
+				]
+
+				// Settings
+				+ SHorizontalBox::Slot()
+				.AutoWidth()
+				.HAlign(HAlign_Right)
 				[
 					CreateSettingsButtonWidget()
 				]
@@ -1116,6 +1126,35 @@ namespace UE::Audio::Insights
 			];
 	}
 
+	TSharedRef<SWidget> FEditorDashboardFactory::CreateEditorPreferencesButtonWidget()
+	{
+		return SNew(SButton)
+			.ButtonStyle(&FAppStyle::Get().GetWidgetStyle<FButtonStyle>("SimpleButton"))
+			.ToolTipText(EditorDashboardFactoryPrivate::EditorPreferencesTooltipText)
+			.OnClicked_Lambda([]()
+			{
+				if (ISettingsModule* const SettingsModule = FModuleManager::GetModulePtr<ISettingsModule>("Settings"))
+				{
+					SettingsModule->ShowViewer("Editor", "Plugins", "AudioInsightsEditorSettings");
+				}
+				return FReply::Handled();
+			})
+			[
+				SNew(SHorizontalBox)
+
+				+ SHorizontalBox::Slot()
+				.AutoWidth()
+				.VAlign(VAlign_Center)
+				.Padding(5.0, 0.0f)
+				[
+					SNew(SImage)
+					.ColorAndOpacity(FSlateColor::UseForeground())
+					.Image(FAppStyle::Get().GetBrush("EditorPreferences.TabIcon"))
+					.DesiredSizeOverride(FVector2D(16.0f, 16.0f))
+				]
+			];
+	}
+
 	TSharedRef<SWidget> FEditorDashboardFactory::CreateSettingsButtonWidget()
 	{
 		if (OnGetSettingsMenuContent() == SNullWidget::NullWidget)
@@ -1141,6 +1180,7 @@ namespace UE::Audio::Insights
 					SNew(SImage)
 					.ColorAndOpacity(FSlateColor::UseForeground())
 					.Image(FSlateStyle::Get().GetBrush("AudioInsights.Icon.Settings"))
+					.DesiredSizeOverride(FVector2D(16.0f, 16.0f))
 				]
 			];
 	}
diff --git a/Engine/Plugins/AudioInsights/Source/AudioInsightsEditor/Private/AudioInsightsEditorDashboardFactory.h b/Engine/Plugins/AudioInsights/Source/AudioInsightsEditor/Private/AudioInsightsEditorDashboardFactory.h
--- a/Engine/Plugins/AudioInsights/Source/AudioInsightsEditor/Private/AudioInsightsEditorDashboardFactory.h
+++ b/Engine/Plugins/AudioInsights/Source/AudioInsightsEditor/Private/AudioInsightsEditorDashboardFactory.h
@@ -58,6 +58,7 @@ namespace UE::Audio::Insights
 
 		TSharedRef<SWidget> MakeTraceModeButton(const ETraceMode InTraceButtonMode);
 
+		TSharedRef<SWidget> CreateEditorPreferencesButtonWidget();
 		TSharedRef<SWidget> CreateSettingsButtonWidget();
 		TSharedRef<SWidget> OnGetSettingsMenuContent();
 
