[Audio Insights] FOutputMeterTraceProvider: simplified scrubbing logic since there is only one main submix per Audio Device.

#Synced-CL 52441547

diff --git a/Engine/Plugins/AudioInsights/Source/AudioInsights/Private/Providers/OutputMeterTraceProvider.cpp b/Engine/Plugins/AudioInsights/Source/AudioInsights/Private/Providers/OutputMeterTraceProvider.cpp
--- a/Engine/Plugins/AudioInsights/Source/AudioInsights/Private/Providers/OutputMeterTraceProvider.cpp
+++ b/Engine/Plugins/AudioInsights/Source/AudioInsights/Private/Providers/OutputMeterTraceProvider.cpp
@@ -1,7 +1,6 @@
 // Copyright Epic Games, Inc. All Rights Reserved.
 #include "Providers/OutputMeterTraceProvider.h"
 
-#include "Async/ParallelFor.h"
 #include "AudioInsightsModule.h"
 #include "Cache/AudioInsightsCacheManager.h"
 #include "Trace/Analyzer.h"
@@ -94,7 +93,7 @@ namespace UE::Audio::Insights
 
 	bool FOutputMeterTraceProvider::ProcessMessages()
 	{
-		auto CreateEntryFunc = [this](const FOutputMeterMessageBase& Msg)
+		auto GetOrCreateEntryFunc = [this](const FOutputMeterMessageBase& Msg)
 		{
 			TSharedPtr<FOutputMeterDashboardEntry>* ToReturn = nullptr;
 			UpdateDeviceEntry(Msg.DeviceId, Msg.SubmixId, [&ToReturn, &Msg](TSharedPtr<FOutputMeterDashboardEntry>& Entry)
@@ -113,17 +112,12 @@ namespace UE::Audio::Insights
 			return ToReturn;
 		};
 
-		auto GetEntryFunc = [this](const FOutputMeterMessageBase& Msg)
-		{
-			return FindDeviceEntry(Msg.DeviceId, Msg.SubmixId);
-		};
-
-		ProcessMessageQueue<FOutputMeterMainSubmixLoadedMessage>(TraceMessages.MainSubmixLoadedMessages, CreateEntryFunc,
+		ProcessMessageQueue<FOutputMeterMainSubmixLoadedMessage>(TraceMessages.MainSubmixLoadedMessages, GetOrCreateEntryFunc,
 		[this](const FOutputMeterMainSubmixLoadedMessage& Msg, TSharedPtr<FOutputMeterDashboardEntry>* OutEntry)
 		{
 		});
 
-		ProcessMessageQueue<FOutputMeterTruePeakMessage>(TraceMessages.TruePeakMessages, GetEntryFunc,
+		ProcessMessageQueue<FOutputMeterTruePeakMessage>(TraceMessages.TruePeakMessages, GetOrCreateEntryFunc,
 		[this](const FOutputMeterTruePeakMessage& Msg, TSharedPtr<FOutputMeterDashboardEntry>* OutEntry)
 		{
 			if (OutEntry)
@@ -133,7 +127,7 @@ namespace UE::Audio::Insights
 			}
 		});
 
-		ProcessMessageQueue<FOutputMeterLKFSValuesMessage>(TraceMessages.LKFSValuesMessages, GetEntryFunc,
+		ProcessMessageQueue<FOutputMeterLKFSValuesMessage>(TraceMessages.LKFSValuesMessages, GetOrCreateEntryFunc,
 		[this](const FOutputMeterLKFSValuesMessage& Msg, TSharedPtr<FOutputMeterDashboardEntry>* OutEntry)
 		{
 			if (OutEntry)
@@ -152,117 +146,55 @@ namespace UE::Audio::Insights
 
 	void FOutputMeterTraceProvider::OnTimingViewTimeMarkerChanged(double TimeMarker)
 	{
-		// Helper lambdas
-		auto ProcessOutputMeterMessage = [this](const FOutputMeterMainSubmixLoadedMessage& OutputMeterCachedMessage)
-		{
-			UpdateDeviceEntry(OutputMeterCachedMessage.DeviceId, OutputMeterCachedMessage.SubmixId, [&OutputMeterCachedMessage](TSharedPtr<FOutputMeterDashboardEntry>& Entry)
-			{
-				if (!Entry.IsValid())
-				{
-					Entry = MakeShared<FOutputMeterDashboardEntry>();
-					Entry->DeviceId  = OutputMeterCachedMessage.DeviceId;
-					Entry->SubmixId = OutputMeterCachedMessage.SubmixId;
-				}
-				
-				Entry->Timestamp = OutputMeterCachedMessage.Timestamp;
-			});
-		};
-
-		// Process cached messages
 		DeviceDataMap.Empty();
 
 		const FAudioInsightsCacheManager& CacheManager = FAudioInsightsModule::GetChecked().GetCacheManager();
+		const ::Audio::FDeviceId FilteredDeviceId = GetFilteredAudioDeviceID();
 
-		// Main submix loaded cached messages
-		CacheManager.IterateTo<FOutputMeterMainSubmixLoadedMessage>(OutputMeterMessageNames::MainSubmixLoaded, TimeMarker, [&ProcessOutputMeterMessage](const FOutputMeterMainSubmixLoadedMessage& OutputMeterMainSubmixLoadedMessage)
+		auto MatchesDevice = [FilteredDeviceId](const FOutputMeterMessageBase& Msg)
 		{
-			ProcessOutputMeterMessage(OutputMeterMainSubmixLoadedMessage);
-		});
-
-		CollectParamsForTimestamp(TimeMarker);
-
-		// Call parent method to update LastMessageId
-		FTraceProviderBase::OnTimingViewTimeMarkerChanged(TimeMarker);
-	}
-
-	void FOutputMeterTraceProvider::CollectParamsForTimestamp(const double InTimeMarker)
-	{
-		const FDeviceData* DeviceData = FindFilteredDeviceData();
-		if (DeviceData == nullptr)
-		{
-			return;
-		}
-
-		// Collect update messages from output meter (based on submix id)
-		struct CachedEntryInfo
-		{
-			TOptional<FOutputMeterTruePeakMessage> TruePeakMessage;
-			TOptional<FOutputMeterLKFSValuesMessage> LKFSValuesMessage;
+			return Msg.DeviceId == FilteredDeviceId;
 		};
 
-		TArray<uint32> MainSubmixIdArray;
-		DeviceData->GenerateKeyArray(MainSubmixIdArray);
-
-		TArray<CachedEntryInfo> CachedEntryInfos;
-		CachedEntryInfos.SetNum(MainSubmixIdArray.Num());
-
-		const FAudioInsightsCacheManager& CacheManager = FAudioInsightsModule::GetChecked().GetCacheManager();
-
-		ParallelFor(MainSubmixIdArray.Num(), [&MainSubmixIdArray, &CachedEntryInfos, &CacheManager, InTimeMarker, this](const int32 Index)
+		// Find closest TruePeak message for the active device at the given time marker
+		if (const FOutputMeterTruePeakMessage* TruePeakMsg = CacheManager.FindClosestMessageByPredicate<FOutputMeterTruePeakMessage>(OutputMeterMessageNames::TruePeak, TimeMarker, MatchesDevice))
 		{
-			const uint32 MainSubmixId = MainSubmixIdArray[Index];
-
-			// TruePeak
-			const FOutputMeterTruePeakMessage* FoundTruePeakCachedMessage = CacheManager.FindClosestMessage<FOutputMeterTruePeakMessage>(OutputMeterMessageNames::TruePeak, InTimeMarker, MainSubmixId);
-			if (FoundTruePeakCachedMessage)
+			UpdateDeviceEntry(TruePeakMsg->DeviceId, TruePeakMsg->SubmixId, [TruePeakMsg](TSharedPtr<FOutputMeterDashboardEntry>& Entry)
 			{
-				CachedEntryInfos[Index].TruePeakMessage = *FoundTruePeakCachedMessage;
-			}
+				if (!Entry.IsValid())
+				{
+					Entry = MakeShared<FOutputMeterDashboardEntry>();
+					Entry->DeviceId = TruePeakMsg->DeviceId;
+					Entry->SubmixId = TruePeakMsg->SubmixId;
+				}
 
-			// LKFSValues
-			const FOutputMeterLKFSValuesMessage* FoundLKFSValuesCachedMessage = CacheManager.FindClosestMessage<FOutputMeterLKFSValuesMessage>(OutputMeterMessageNames::LKFSValues, InTimeMarker, MainSubmixId);
-			if (FoundLKFSValuesCachedMessage)
-			{
-				CachedEntryInfos[Index].LKFSValuesMessage = *FoundLKFSValuesCachedMessage;
-			}
-		});
+				Entry->Timestamp          = TruePeakMsg->Timestamp;
+				Entry->TruePeakMaxValueDb = TruePeakMsg->TruePeakMaxValueDb;
+			});
+		}
 
-		// Update the device entries with the collected info
-		for (const CachedEntryInfo& CachedEntryInfo : CachedEntryInfos)
+		// Find closest LKFS message for the active device at the given time marker
+		if (const FOutputMeterLKFSValuesMessage* LKFSMsg = CacheManager.FindClosestMessageByPredicate<FOutputMeterLKFSValuesMessage>(OutputMeterMessageNames::LKFSValues, TimeMarker, MatchesDevice))
 		{
-			// TruePeak
-			if (CachedEntryInfo.TruePeakMessage.IsSet())
+			UpdateDeviceEntry(LKFSMsg->DeviceId, LKFSMsg->SubmixId, [LKFSMsg](TSharedPtr<FOutputMeterDashboardEntry>& Entry)
 			{
-				const FOutputMeterTruePeakMessage& OutputMeterTruePeakMessage = CachedEntryInfo.TruePeakMessage.GetValue();
-
-				UpdateDeviceEntry(OutputMeterTruePeakMessage.DeviceId, OutputMeterTruePeakMessage.SubmixId, [&OutputMeterTruePeakMessage](TSharedPtr<FOutputMeterDashboardEntry>& Entry)
+				if (!Entry.IsValid())
 				{
-					if (ensure(Entry.IsValid()))
-					{
-						Entry->Timestamp = OutputMeterTruePeakMessage.Timestamp;
-						Entry->TruePeakMaxValueDb = OutputMeterTruePeakMessage.TruePeakMaxValueDb;
-					}
-				});
-			}
-
-			// LKFSValues
-			if (CachedEntryInfo.LKFSValuesMessage.IsSet())
-			{
-				const FOutputMeterLKFSValuesMessage& OutputMeterLKFSValuesMessage = CachedEntryInfo.LKFSValuesMessage.GetValue();
+					Entry = MakeShared<FOutputMeterDashboardEntry>();
+					Entry->DeviceId = LKFSMsg->DeviceId;
+					Entry->SubmixId = LKFSMsg->SubmixId;
+				}
 
-				UpdateDeviceEntry(OutputMeterLKFSValuesMessage.DeviceId, OutputMeterLKFSValuesMessage.SubmixId, [&OutputMeterLKFSValuesMessage](TSharedPtr<FOutputMeterDashboardEntry>& Entry)
-				{
-					if (ensure(Entry.IsValid()))
-					{
-						Entry->Timestamp = OutputMeterLKFSValuesMessage.Timestamp;
-						Entry->LongTermLoudness  = OutputMeterLKFSValuesMessage.LongTermLoudness;
-						Entry->ShortTermLoudness = OutputMeterLKFSValuesMessage.ShortTermLoudness;
-						Entry->MomentaryLoudness = OutputMeterLKFSValuesMessage.MomentaryLoudness;
-						Entry->LoudnessRangeLowerBound = OutputMeterLKFSValuesMessage.LoudnessRangeLowerBound;
-						Entry->LoudnessRangeUpperBound = OutputMeterLKFSValuesMessage.LoudnessRangeUpperBound;
-					}
-				});
-			}
+				Entry->Timestamp              = LKFSMsg->Timestamp;
+				Entry->LongTermLoudness       = LKFSMsg->LongTermLoudness;
+				Entry->ShortTermLoudness      = LKFSMsg->ShortTermLoudness;
+				Entry->MomentaryLoudness      = LKFSMsg->MomentaryLoudness;
+				Entry->LoudnessRangeLowerBound = LKFSMsg->LoudnessRangeLowerBound;
+				Entry->LoudnessRangeUpperBound = LKFSMsg->LoudnessRangeUpperBound;
+			});
 		}
+
+		// Call parent method to update LastMessageId
+		FTraceProviderBase::OnTimingViewTimeMarkerChanged(TimeMarker);
 	}
 } // namespace UE::Audio::Insights
diff --git a/Engine/Plugins/AudioInsights/Source/AudioInsights/Private/Providers/OutputMeterTraceProvider.h b/Engine/Plugins/AudioInsights/Source/AudioInsights/Private/Providers/OutputMeterTraceProvider.h
--- a/Engine/Plugins/AudioInsights/Source/AudioInsights/Private/Providers/OutputMeterTraceProvider.h
+++ b/Engine/Plugins/AudioInsights/Source/AudioInsights/Private/Providers/OutputMeterTraceProvider.h
@@ -25,7 +25,6 @@ namespace UE::Audio::Insights
 	private:
 		virtual bool ProcessMessages() override;
 		virtual void OnTimingViewTimeMarkerChanged(double TimeMarker) override;
-		void CollectParamsForTimestamp(const double InTimeMarker);
 
 		FOutputMeterMessages TraceMessages;
 	};
