[Audio Insights] Add Details dashboard to Audio Insights for inline asset property editing

#Synced-CL 52467453

diff --git a/Engine/Plugins/AudioInsights/Source/AudioInsights/Private/AudioInsightsDetailsSelectionManager.cpp b/Engine/Plugins/AudioInsights/Source/AudioInsights/Private/AudioInsightsDetailsSelectionManager.cpp
new file mode 100644
--- /dev/null
+++ b/Engine/Plugins/AudioInsights/Source/AudioInsights/Private/AudioInsightsDetailsSelectionManager.cpp
@@ -0,0 +1,28 @@
+// Copyright Epic Games, Inc. All Rights Reserved.
+#include "AudioInsightsDetailsSelectionManager.h"
+
+#if WITH_EDITOR
+
+namespace UE::Audio::Insights
+{
+	void FAudioInsightsDetailsSelectionManager::SetSelectedAsset(TObjectPtr<UObject> InAsset)
+	{
+		if (SelectedAsset.Get() != InAsset)
+		{
+			SelectedAsset = InAsset;
+			OnAssetSelectionChanged.Broadcast(InAsset);
+		}
+	}
+
+	void FAudioInsightsDetailsSelectionManager::ClearSelection()
+	{
+		SetSelectedAsset(nullptr);
+	}
+
+	TWeakObjectPtr<UObject> FAudioInsightsDetailsSelectionManager::GetSelectedAsset() const
+	{
+		return SelectedAsset;
+	}
+} // namespace UE::Audio::Insights
+
+#endif // WITH_EDITOR
diff --git a/Engine/Plugins/AudioInsights/Source/AudioInsights/Private/AudioInsightsModule.h b/Engine/Plugins/AudioInsights/Source/AudioInsights/Private/AudioInsightsModule.h
--- a/Engine/Plugins/AudioInsights/Source/AudioInsights/Private/AudioInsightsModule.h
+++ b/Engine/Plugins/AudioInsights/Source/AudioInsights/Private/AudioInsightsModule.h
@@ -2,6 +2,7 @@
 #pragma once 
 
 #include "AudioInsightsDashboardFactory.h"
+#include "AudioInsightsDetailsSelectionManager.h"
 #include "AudioInsightsTimingViewExtender.h"
 #include "AudioInsightsTraceModule.h"
 #include "Framework/Docking/TabManager.h"
@@ -49,6 +50,10 @@ namespace UE::Audio::Insights
 		FAudioInsightsTimingViewExtender& GetTimingViewExtender() { return AudioInsightsTimingViewExtender; };
 		const FAudioInsightsTimingViewExtender& GetTimingViewExtender() const { return AudioInsightsTimingViewExtender; };
 
+#if WITH_EDITOR
+		FAudioInsightsDetailsSelectionManager& GetDetailsSelectionManager() override { return DetailsSelectionManager; };
+#endif // WITH_EDITOR
+
 		TSharedRef<FDashboardFactory> GetDashboardFactory();
 		const TSharedRef<FDashboardFactory> GetDashboardFactory() const;
 
@@ -65,5 +70,9 @@ namespace UE::Audio::Insights
 #endif // !WITH_EDITOR
 
 		FAudioInsightsTimingViewExtender AudioInsightsTimingViewExtender;
+
+#if WITH_EDITOR
+		FAudioInsightsDetailsSelectionManager DetailsSelectionManager;
+#endif // WITH_EDITOR
 	};
 } // namespace UE::Audio::Insights
diff --git a/Engine/Plugins/AudioInsights/Source/AudioInsights/Private/IAudioInsightsModule.cpp b/Engine/Plugins/AudioInsights/Source/AudioInsights/Private/IAudioInsightsModule.cpp
--- a/Engine/Plugins/AudioInsights/Source/AudioInsights/Private/IAudioInsightsModule.cpp
+++ b/Engine/Plugins/AudioInsights/Source/AudioInsights/Private/IAudioInsightsModule.cpp
@@ -21,6 +21,14 @@ UE::Audio::Insights::FAudioInsightsTimingViewExtender& IAudioInsightsModule::Get
 	return AudioInsightsModule.GetTimingViewExtender();
 }
 
+#if WITH_EDITOR
+UE::Audio::Insights::FAudioInsightsDetailsSelectionManager& IAudioInsightsModule::GetDetailsSelectionManager()
+{
+	IAudioInsightsModule& AudioInsightsModule = static_cast<IAudioInsightsModule&>(UE::Audio::Insights::FAudioInsightsModule::GetChecked());
+	return AudioInsightsModule.GetDetailsSelectionManager();
+}
+#endif // WITH_EDITOR
+
 bool IAudioInsightsModule::IsModuleLoaded()
 {
 	return UE::Audio::Insights::FAudioInsightsModule::IsModuleLoaded();
diff --git a/Engine/Plugins/AudioInsights/Source/AudioInsights/Private/Views/SoundDashboardViewFactory.cpp b/Engine/Plugins/AudioInsights/Source/AudioInsights/Private/Views/SoundDashboardViewFactory.cpp
--- a/Engine/Plugins/AudioInsights/Source/AudioInsights/Private/Views/SoundDashboardViewFactory.cpp
+++ b/Engine/Plugins/AudioInsights/Source/AudioInsights/Private/Views/SoundDashboardViewFactory.cpp
@@ -1607,10 +1607,25 @@ namespace UE::Audio::Insights
 
 	void FSoundDashboardViewFactory::OnSelectionChanged(TSharedPtr<IDashboardDataTreeViewEntry> SelectedItem, ESelectInfo::Type SelectInfo)
 	{
-		if (FilteredEntriesListView.IsValid())
+		if (!FilteredEntriesListView.IsValid())
 		{
-			OnUpdatePlotSelection.Broadcast(FilteredEntriesListView->GetSelectedItems());
+			return;
 		}
+
+		OnUpdatePlotSelection.Broadcast(FilteredEntriesListView->GetSelectedItems());
+
+#if WITH_EDITOR
+		FAudioInsightsDetailsSelectionManager& SelectionManager = IAudioInsightsModule::GetChecked().GetDetailsSelectionManager();
+		const TArray<TObjectPtr<UObject>> Assets = GetSelectedEditableAssets();
+		if (Assets.Num() > 0)
+		{
+			SelectionManager.SetSelectedAsset(Assets[0]);
+		}
+		else if (SelectInfo != ESelectInfo::Type::Direct)
+		{
+			SelectionManager.ClearSelection();
+		}
+#endif // WITH_EDITOR
 	}
 
 	FReply FSoundDashboardViewFactory::OnDataRowKeyInput(const FGeometry& InGeometry, const FKeyEvent& InKeyEvent) const
diff --git a/Engine/Plugins/AudioInsights/Source/AudioInsights/Public/AudioInsightsDetailsSelectionManager.h b/Engine/Plugins/AudioInsights/Source/AudioInsights/Public/AudioInsightsDetailsSelectionManager.h
new file mode 100644
--- /dev/null
+++ b/Engine/Plugins/AudioInsights/Source/AudioInsights/Public/AudioInsightsDetailsSelectionManager.h
@@ -0,0 +1,32 @@
+// Copyright Epic Games, Inc. All Rights Reserved.
+#pragma once
+
+#if WITH_EDITOR
+
+#include "Delegates/Delegate.h"
+#include "UObject/ObjectPtr.h"
+#include "UObject/WeakObjectPtr.h"
+
+#define UE_API AUDIOINSIGHTS_API
+
+namespace UE::Audio::Insights
+{
+	class FAudioInsightsDetailsSelectionManager
+	{
+	public:
+		UE_API void SetSelectedAsset(TObjectPtr<UObject> InAsset);
+		UE_API void ClearSelection();
+
+		UE_API TWeakObjectPtr<UObject> GetSelectedAsset() const;
+
+		DECLARE_MULTICAST_DELEGATE_OneParam(FOnAssetSelectionChanged, TObjectPtr<UObject> /*SelectedAsset*/);
+		FOnAssetSelectionChanged OnAssetSelectionChanged;
+
+	private:
+		TWeakObjectPtr<UObject> SelectedAsset;
+	};
+} // namespace UE::Audio::Insights
+
+#undef UE_API
+
+#endif // WITH_EDITOR
diff --git a/Engine/Plugins/AudioInsights/Source/AudioInsights/Public/AudioInsightsStyle.h b/Engine/Plugins/AudioInsights/Source/AudioInsights/Public/AudioInsightsStyle.h
--- a/Engine/Plugins/AudioInsights/Source/AudioInsights/Public/AudioInsightsStyle.h
+++ b/Engine/Plugins/AudioInsights/Source/AudioInsights/Public/AudioInsightsStyle.h
@@ -206,6 +206,7 @@ namespace UE::Audio::Insights
 			Set("AudioInsights.Icon.CheckBoxCheck", new IMAGE_BRUSH_SVG(TEXT("Icons/checkboxcheck"), Icon16));
 			Set("AudioInsights.Icon.CheckBoxIndeterminate", new IMAGE_BRUSH_SVG(TEXT("Icons/checkboxindeterminate"), Icon16));
 			Set("AudioInsights.Icon.Dashboard", new IMAGE_BRUSH_SVG(TEXT("Icons/audio_dashboard"), Icon16));
+			Set("AudioInsights.Icon.Details", new IMAGE_BRUSH_SVG(TEXT("Icons/details"), Icon16));
 			Set("AudioInsights.Icon.Event", new IMAGE_BRUSH_SVG(TEXT("Icons/audio_event"), Icon16));
 			Set("AudioInsights.Icon.ExpandArrowDown", new IMAGE_BRUSH_SVG(TEXT("Icons/expand_arrow_down"), Icon12));
 			Set("AudioInsights.Icon.ExpandArrowUp", new IMAGE_BRUSH_SVG(TEXT("Icons/expand_arrow_up"), Icon12));
diff --git a/Engine/Plugins/AudioInsights/Source/AudioInsights/Public/IAudioInsightsModule.h b/Engine/Plugins/AudioInsights/Source/AudioInsights/Public/IAudioInsightsModule.h
--- a/Engine/Plugins/AudioInsights/Source/AudioInsights/Public/IAudioInsightsModule.h
+++ b/Engine/Plugins/AudioInsights/Source/AudioInsights/Public/IAudioInsightsModule.h
@@ -17,6 +17,7 @@ namespace UE::Audio::Insights
 {
 	class IDashboardViewFactory;
 	class FAudioInsightsCacheManager;
+	class FAudioInsightsDetailsSelectionManager;
 	class FAudioInsightsTimingViewExtender;
 	class FTraceProviderBase;
 } // namespace UE::Audio::Insights
@@ -35,6 +36,10 @@ public:
 
 	UE_API virtual UE::Audio::Insights::FAudioInsightsTimingViewExtender& GetTimingViewExtender() override;
 
+#if WITH_EDITOR
+	UE_API virtual UE::Audio::Insights::FAudioInsightsDetailsSelectionManager& GetDetailsSelectionManager() override;
+#endif // WITH_EDITOR
+
 	virtual TSharedRef<SDockTab> CreateDashboardTabWidget(const FSpawnTabArgs& Args) = 0;
 
 	UE_API static bool IsModuleLoaded();
diff --git a/Engine/Plugins/AudioInsights/Source/AudioInsights/Public/IAudioInsightsModuleInterface.h b/Engine/Plugins/AudioInsights/Source/AudioInsights/Public/IAudioInsightsModuleInterface.h
--- a/Engine/Plugins/AudioInsights/Source/AudioInsights/Public/IAudioInsightsModuleInterface.h
+++ b/Engine/Plugins/AudioInsights/Source/AudioInsights/Public/IAudioInsightsModuleInterface.h
@@ -2,6 +2,7 @@
 #pragma once
 
 #include "AudioDefines.h"
+#include "AudioInsightsDetailsSelectionManager.h"
 #include "AudioInsightsTimingViewExtender.h"
 #include "Cache/AudioInsightsCacheManager.h"
 #include "Modules/ModuleInterface.h"
@@ -27,4 +28,8 @@ public:
 	virtual IAudioInsightsTraceModule& GetTraceModule() = 0;
 	virtual UE::Audio::Insights::FAudioInsightsCacheManager& GetCacheManager() = 0;
 	virtual UE::Audio::Insights::FAudioInsightsTimingViewExtender& GetTimingViewExtender() = 0;
+
+#if WITH_EDITOR
+	virtual UE::Audio::Insights::FAudioInsightsDetailsSelectionManager& GetDetailsSelectionManager() = 0;
+#endif // WITH_EDITOR
 };
diff --git a/Engine/Plugins/AudioInsights/Source/AudioInsightsEditor/Private/AudioInsightsEditorModule.cpp b/Engine/Plugins/AudioInsights/Source/AudioInsightsEditor/Private/AudioInsightsEditorModule.cpp
--- a/Engine/Plugins/AudioInsights/Source/AudioInsightsEditor/Private/AudioInsightsEditorModule.cpp
+++ b/Engine/Plugins/AudioInsights/Source/AudioInsightsEditor/Private/AudioInsightsEditorModule.cpp
@@ -12,6 +12,7 @@
 #include "Views/AudioBusDashboardViewFactory.h"
 #include "Views/AudioEventLogDashboardViewFactory.h"
 #include "Views/AudioMetersPanelDashboardViewFactory.h"
+#include "Views/DetailsDashboardViewFactory.h"
 #include "Views/LogDashboardViewFactory.h"
 #include "Views/SignalFlowDashboardViewFactory.h"
 #include "Views/SoundDashboardViewFactory.h"
@@ -43,6 +44,7 @@ namespace UE::Audio::Insights
 			// @TODO UE-274216: Decide what to do with the Viewport dashboard
 			//DashboardFactory->RegisterViewFactory(MakeShared<FViewportDashboardViewFactory>());
 			DashboardFactory->RegisterViewFactory(MakeShared<FLogDashboardViewFactory>());
+			DashboardFactory->RegisterViewFactory(MakeShared<FDetailsDashboardViewFactory>());
 			DashboardFactory->RegisterViewFactory(MakeShared<FAudioEventLogDashboardViewFactory>());
 			DashboardFactory->RegisterViewFactory(SoundsDashboard);
 			DashboardFactory->RegisterViewFactory(MakeShared<FVirtualLoopDashboardViewFactory>());
@@ -108,6 +110,12 @@ namespace UE::Audio::Insights
 		return InsightsModule.GetTimingViewExtender();
 	}
 
+	FAudioInsightsDetailsSelectionManager& FAudioInsightsEditorModule::GetDetailsSelectionManager()
+	{
+		IAudioInsightsModule& InsightsModule = IAudioInsightsModule::GetChecked();
+		return InsightsModule.GetDetailsSelectionManager();
+	}
+
 	TSharedPtr<FEditorDashboardFactory> FAudioInsightsEditorModule::GetDashboardFactory()
 	{
 		return DashboardFactory;
diff --git a/Engine/Plugins/AudioInsights/Source/AudioInsightsEditor/Private/AudioInsightsEditorModule.h b/Engine/Plugins/AudioInsights/Source/AudioInsightsEditor/Private/AudioInsightsEditorModule.h
--- a/Engine/Plugins/AudioInsights/Source/AudioInsightsEditor/Private/AudioInsightsEditorModule.h
+++ b/Engine/Plugins/AudioInsights/Source/AudioInsightsEditor/Private/AudioInsightsEditorModule.h
@@ -37,6 +37,7 @@ namespace UE::Audio::Insights
 		virtual IAudioInsightsTraceModule& GetTraceModule() override;
 		virtual class FAudioInsightsCacheManager& GetCacheManager() override;
 		virtual class FAudioInsightsTimingViewExtender& GetTimingViewExtender() override;
+		virtual class FAudioInsightsDetailsSelectionManager& GetDetailsSelectionManager() override;
 
 	private:
 		void RegisterMenus();
diff --git a/Engine/Plugins/AudioInsights/Source/AudioInsightsEditor/Private/Views/DetailsDashboardViewFactory.cpp b/Engine/Plugins/AudioInsights/Source/AudioInsightsEditor/Private/Views/DetailsDashboardViewFactory.cpp
new file mode 100644
--- /dev/null
+++ b/Engine/Plugins/AudioInsights/Source/AudioInsightsEditor/Private/Views/DetailsDashboardViewFactory.cpp
@@ -0,0 +1,119 @@
+// Copyright Epic Games, Inc. All Rights Reserved.
+#include "Views/DetailsDashboardViewFactory.h"
+
+#include "AudioInsightsDetailsSelectionManager.h"
+#include "AudioInsightsStyle.h"
+#include "IAudioInsightsModule.h"
+#include "IDetailsView.h"
+#include "Modules/ModuleManager.h"
+#include "PropertyEditorModule.h"
+#include "Widgets/Layout/SBox.h"
+#include "Widgets/Text/STextBlock.h"
+
+#define LOCTEXT_NAMESPACE "AudioInsights"
+
+namespace UE::Audio::Insights
+{
+	FName FDetailsDashboardViewFactory::GetName() const
+	{
+		return "Details";
+	}
+
+	FText FDetailsDashboardViewFactory::GetDisplayName() const
+	{
+		return LOCTEXT("AudioDashboard_DetailsTab_DisplayName", "Details");
+	}
+
+	EDefaultDashboardTabStack FDetailsDashboardViewFactory::GetDefaultTabStack() const
+	{
+		return EDefaultDashboardTabStack::Log;
+	}
+
+	FSlateIcon FDetailsDashboardViewFactory::GetIcon() const
+	{
+		return FSlateStyle::Get().CreateIcon("AudioInsights.Icon.Details");
+	}
+
+	TSharedRef<SWidget> FDetailsDashboardViewFactory::MakeWidget(const TSharedRef<SDockTab> OwnerTab, const FSpawnTabArgs& SpawnTabArgs)
+	{
+		if (!DashboardWidget.IsValid())
+		{
+			DashboardWidget = CreateDetailsViewWidget();
+		}
+
+		if (!AssetSelectionChangedHandle.IsValid())
+		{
+			BindOnAssetChangedDelegate();
+		}
+
+		return DashboardWidget.ToSharedRef();
+	}
+
+	TSharedPtr<SWidget> FDetailsDashboardViewFactory::CreateDetailsViewWidget()
+	{
+		FPropertyEditorModule& PropertyEditorModule = FModuleManager::LoadModuleChecked<FPropertyEditorModule>("PropertyEditor");
+
+		FDetailsViewArgs DetailsViewArgs;
+		DetailsViewArgs.bAllowSearch = true;
+		DetailsViewArgs.bHideSelectionTip = true;
+		DetailsViewArgs.bShowObjectLabel = true;
+
+		DetailsView = PropertyEditorModule.CreateDetailView(DetailsViewArgs);
+
+		DetailsViewWidget = DetailsView->AsShared();
+		DetailsViewWidget->SetVisibility(EVisibility::Collapsed);
+
+		NoSelectionMessageWidget = SNew(SBox)
+			.HAlign(HAlign_Center)
+			.VAlign(VAlign_Center)
+			[
+				SNew(STextBlock)
+					.Text(LOCTEXT("AudioDashboard_Details_NoSelection", "Select an audio asset to view its properties."))
+			];
+
+		return SNew(SOverlay)
+			+ SOverlay::Slot()
+			[
+				NoSelectionMessageWidget.ToSharedRef()
+			]
+			+ SOverlay::Slot()
+			[
+				DetailsViewWidget.ToSharedRef()
+			];
+	}
+
+	void FDetailsDashboardViewFactory::BindOnAssetChangedDelegate()
+	{
+		FAudioInsightsDetailsSelectionManager& SelectionManager = IAudioInsightsModule::GetChecked().GetDetailsSelectionManager();
+		AssetSelectionChangedHandle = SelectionManager.OnAssetSelectionChanged.AddSP(this, &FDetailsDashboardViewFactory::OnAssetSelectionChanged);
+
+		const TWeakObjectPtr<UObject> CurrentAsset = SelectionManager.GetSelectedAsset();
+		if (CurrentAsset.IsValid())
+		{
+			OnAssetSelectionChanged(CurrentAsset.Get());
+		}
+	}
+
+	void FDetailsDashboardViewFactory::OnAssetSelectionChanged(const TObjectPtr<UObject> InAsset)
+	{
+		if (!DetailsView.IsValid() || !DetailsViewWidget.IsValid() || !NoSelectionMessageWidget.IsValid() || !DashboardWidget.IsValid())
+		{
+			return;
+		}
+
+		if (InAsset)
+		{
+			DetailsView->SetObject(InAsset);
+			DetailsViewWidget->SetVisibility(EVisibility::Visible);
+			NoSelectionMessageWidget->SetVisibility(EVisibility::Collapsed);
+		}
+		else
+		{
+			DetailsView->SetObject(nullptr);
+			DetailsViewWidget->SetVisibility(EVisibility::Collapsed);
+			NoSelectionMessageWidget->SetVisibility(EVisibility::Visible);
+		}
+	}
+} // namespace UE::Audio::Insights
+
+#undef LOCTEXT_NAMESPACE
diff --git a/Engine/Plugins/AudioInsights/Source/AudioInsightsEditor/Private/Views/DetailsDashboardViewFactory.h b/Engine/Plugins/AudioInsights/Source/AudioInsightsEditor/Private/Views/DetailsDashboardViewFactory.h
new file mode 100644
--- /dev/null
+++ b/Engine/Plugins/AudioInsights/Source/AudioInsightsEditor/Private/Views/DetailsDashboardViewFactory.h
@@ -0,0 +1,40 @@
+// Copyright Epic Games, Inc. All Rights Reserved.
+#pragma once
+
+#include "Internationalization/Text.h"
+#include "Templates/SharedPointer.h"
+#include "Textures/SlateIcon.h"
+#include "UObject/NameTypes.h"
+#include "Views/DashboardViewFactory.h"
+#include "Widgets/SWidget.h"
+
+class IDetailsView;
+
+namespace UE::Audio::Insights
+{
+	class FDetailsDashboardViewFactory final : public IDashboardViewFactory, public TSharedFromThis<FDetailsDashboardViewFactory>
+	{
+	public:
+		FDetailsDashboardViewFactory() = default;
+		~FDetailsDashboardViewFactory() = default;
+
+		virtual FName GetName() const override;
+		virtual FText GetDisplayName() const override;
+		virtual EDefaultDashboardTabStack GetDefaultTabStack() const override;
+		virtual FSlateIcon GetIcon() const override;
+		virtual TSharedRef<SWidget> MakeWidget(const TSharedRef<SDockTab> OwnerTab, const FSpawnTabArgs& SpawnTabArgs) override;
+
+	private:
+		TSharedPtr<SWidget> CreateDetailsViewWidget();
+		void BindOnAssetChangedDelegate();
+
+		void OnAssetSelectionChanged(const TObjectPtr<UObject> InAsset);
+
+		TSharedPtr<IDetailsView> DetailsView;
+		TSharedPtr<SWidget> DetailsViewWidget;
+		TSharedPtr<SWidget> NoSelectionMessageWidget;
+		TSharedPtr<SWidget> DashboardWidget;
+
+		FDelegateHandle AssetSelectionChangedHandle;
+	};
+} // namespace UE::Audio::Insights
