[Audio Insights] Enable FOutputMeterDashboardViewFactory in Editor so the user can inspect cached values when scrubbing in the event log + added TimerPanelParams + added loudness meter settings persistence in standalone + enhanced SSpliter with expandable areas + misc layout adjustments.

#Synced-CL 52441547

diff --git a/Engine/Plugins/AudioInsights/Source/AudioInsights/AudioInsights.Build.cs b/Engine/Plugins/AudioInsights/Source/AudioInsights/AudioInsights.Build.cs
--- a/Engine/Plugins/AudioInsights/Source/AudioInsights/AudioInsights.Build.cs
+++ b/Engine/Plugins/AudioInsights/Source/AudioInsights/AudioInsights.Build.cs
@@ -27,6 +27,7 @@ public class AudioInsights : ModuleRules
 				"AudioMixerCore",
 				"AudioWidgetsCore",
 				"CoreUObject",
+				"DeveloperSettings",
 				"InputCore",
 				"OutputLog",
 				"SessionServices",
diff --git a/Engine/Plugins/AudioInsights/Source/AudioInsights/Private/AudioInsightsDashboardFactory.cpp b/Engine/Plugins/AudioInsights/Source/AudioInsights/Private/AudioInsightsDashboardFactory.cpp
--- a/Engine/Plugins/AudioInsights/Source/AudioInsights/Private/AudioInsightsDashboardFactory.cpp
+++ b/Engine/Plugins/AudioInsights/Source/AudioInsights/Private/AudioInsightsDashboardFactory.cpp
@@ -151,18 +151,19 @@ namespace UE::Audio::Insights
 						{
 							if (DashboardTabManager.IsValid())
 							{
-								if (TSharedPtr<SDockTab> ViewportTab = DashboardTabManager->FindExistingLiveTab(FactoryName);
-									!ViewportTab.IsValid())
+								const TSharedPtr<SDockTab> ViewportTab = DashboardTabManager->FindExistingLiveTab(FactoryName);
+
+								if (!ViewportTab.IsValid())
 								{
-									DashboardTabManager->TryInvokeTab(FactoryName);
+									const TSharedPtr<SDockTab> InvokedTab = DashboardTabManager->TryInvokeTab(FactoryName);
 
-									if (TSharedPtr<SDockTab> InvokedOutputMeterTab = DashboardTabManager->TryInvokeTab(FactoryName);
-										InvokedOutputMeterTab.IsValid() && DashboardViewFactories[FactoryName].IsValid())
+									if (InvokedTab.IsValid() && DashboardViewFactories[FactoryName].IsValid())
 									{
-										if (const EDefaultDashboardTabStack DefaultTabStack = DashboardViewFactories[FactoryName]->GetDefaultTabStack();
-											DefaultTabStack == EDefaultDashboardTabStack::AudioAnalyzerRack)
+										const EDefaultDashboardTabStack DefaultTabStack = DashboardViewFactories[FactoryName]->GetDefaultTabStack();
+
+										if (DefaultTabStack == EDefaultDashboardTabStack::OutputMetering)
 										{
-											InvokedOutputMeterTab->SetParentDockTabStackTabWellHidden(true);
+											InvokedTab->SetParentDockTabStackTabWellHidden(true);
 										}
 									}
 								}
@@ -259,7 +260,7 @@ namespace UE::Audio::Insights
 		TSharedRef<FTabManager::FStack> LogTabStack = FTabManager::NewStack();
 		TSharedRef<FTabManager::FStack> AnalysisTabStack = FTabManager::NewStack();
 		TSharedRef<FTabManager::FStack> AudioMetersTabStack = FTabManager::NewStack();
-		TSharedRef<FTabManager::FStack> AudioAnalyzerRackTabStack = FTabManager::NewStack();
+		TSharedRef<FTabManager::FStack> OutputMeteringTabStack = FTabManager::NewStack()->SetHideTabWell(true);
 		TSharedRef<FTabManager::FStack> PlotsTabStack = FTabManager::NewStack();
 
 		for (const auto& [FactoryName, Factory] : DashboardViewFactories)
@@ -286,9 +287,9 @@ namespace UE::Audio::Insights
 				}
 				break;
 
-				case EDefaultDashboardTabStack::AudioAnalyzerRack:
+				case EDefaultDashboardTabStack::OutputMetering:
 				{
-					AudioAnalyzerRackTabStack->AddTab(FactoryName, ETabState::OpenedTab);
+					OutputMeteringTabStack->AddTab(FactoryName, ETabState::OpenedTab);
 				}
 				break;
 
@@ -305,7 +306,7 @@ namespace UE::Audio::Insights
 
 		AnalysisTabStack->SetForegroundTab(FName("Sounds"));
 
-		return FTabManager::NewLayout("AudioDashboard_Layout_v2")
+		return FTabManager::NewLayout("AudioDashboard_Layout_v3")
 		->AddArea
 		(
 			FTabManager::NewPrimaryArea()
@@ -363,7 +364,7 @@ namespace UE::Audio::Insights
 				// Right column
 				->Split
 				(
-					AudioAnalyzerRackTabStack
+					OutputMeteringTabStack
 					->SetSizeCoefficient(0.25f)
 				)
 			)
diff --git a/Engine/Plugins/AudioInsights/Source/AudioInsights/Private/Views/OutputMeterDashboardViewFactory.cpp b/Engine/Plugins/AudioInsights/Source/AudioInsights/Private/Views/OutputMeterDashboardViewFactory.cpp
--- a/Engine/Plugins/AudioInsights/Source/AudioInsights/Private/Views/OutputMeterDashboardViewFactory.cpp
+++ b/Engine/Plugins/AudioInsights/Source/AudioInsights/Private/Views/OutputMeterDashboardViewFactory.cpp
@@ -4,10 +4,18 @@
 #include "Algo/Find.h"
 #include "AudioDefines.h"
 #include "AudioInsightsModule.h"
+#include "AudioInsightsSettings.h"
 #include "AudioInsightsStyle.h"
+#include "AudioInsightsTraceModule.h"
 #include "DSP/Dsp.h"
+#include "LoudnessMeterSettings.h"
+#include "Providers/AudioMeterProvider.h"
+#include "Providers/OutputMeterTraceProvider.h"
+#include "Providers/SubmixTraceProvider.h"
 #include "SAudioMeterWidget.h"
+#include "Views/SubmixDashboardViewFactory.h"
 #include "Widgets/Colors/SColorBlock.h"
+#include "Widgets/Layout/SExpandableArea.h"
 #include "Widgets/SOverlay.h"
 
 #define LOCTEXT_NAMESPACE "AudioInsights"
@@ -26,6 +34,7 @@ namespace UE::Audio::Insights
 	FOutputMeterDashboardViewFactory::FOutputMeterDashboardViewFactory(TSharedRef<FSubmixDashboardViewFactory> SubmixDashboard)
 		: FTraceDashboardViewFactoryBase()
 		, SubmixProvider(SubmixDashboard->FindProvider<FSubmixTraceProvider>())
+		, AudioMeterInfo(MakeUnique<FAudioMeterInfo>())
 	{
 		FTraceModule& AudioInsightsTraceModule = static_cast<FTraceModule&>(FAudioInsightsModule::GetChecked().GetTraceModule());
 
@@ -39,14 +48,16 @@ namespace UE::Audio::Insights
 		};
 	}
 
+	FOutputMeterDashboardViewFactory::~FOutputMeterDashboardViewFactory() = default;
+
 	EDefaultDashboardTabStack FOutputMeterDashboardViewFactory::GetDefaultTabStack() const
 	{
-		return EDefaultDashboardTabStack::AudioAnalyzerRack;
+		return EDefaultDashboardTabStack::OutputMetering;
 	}
 
 	FText FOutputMeterDashboardViewFactory::GetDisplayName() const
 	{
-		return LOCTEXT("AudioDashboard_DashboardsAnalyzerRackTab_DisplayName", "Analyzers");
+		return LOCTEXT("AudioDashboard_DashboardsOutputMeteringTab_DisplayName", "Output Metering");
 	}
 
 	FName FOutputMeterDashboardViewFactory::GetName() const
@@ -59,67 +70,109 @@ namespace UE::Audio::Insights
 		return FSlateStyle::Get().CreateIcon("AudioInsights.Icon");
 	}
 
+	FLoudnessMeterSettings& FOutputMeterDashboardViewFactory::GetLoudnessMeterSettings() const
+	{
+		return GetMutableDefault<UAudioInsightsSettings>()->OutputMeterSettings;
+	}
+
 	void FOutputMeterDashboardViewFactory::InitLoudnessMeterWidgetView()
 	{
 		using namespace AudioWidgetsCore;
 
+		auto SaveSettings = [this]()
+		{
+			const FProperty* Property = UAudioInsightsSettings::StaticClass()->FindPropertyByName(GET_MEMBER_NAME_CHECKED(UAudioInsightsSettings, OutputMeterSettings));
+
+			UAudioInsightsSettings* Settings = GetMutableDefault<UAudioInsightsSettings>();
+			Settings->UpdateSinglePropertyInConfigFile(Property, GetConfigFilename(Settings));
+		};
+
+		// TimerPanelParams
+		FLoudnessMeterWidgetView::FTimerPanelParams TimerPanelParams
+		{
+			.AnalysisTime = TAttribute<FTimespan>::CreateSP(this, &FOutputMeterDashboardViewFactory::GetAnalysisTime),
+			.OnResetButtonClicked = FOnClicked::CreateSP(this, &FOutputMeterDashboardViewFactory::HandleResetButtonClicked),
+			.bIsVisible = true
+		};
+
+		TimerPanelParams.bIsVisible.BindLambda([this]() { return GetLoudnessMeterSettings().bDisplayAnalysisTimer; });
+		TimerPanelParams.OnVisibilityToggleRequested.BindLambda([this, SaveSettings]()
+		{
+			FLoudnessMeterSettings& LoudnessMeterSettings = GetLoudnessMeterSettings();
+			LoudnessMeterSettings.bDisplayAnalysisTimer = !LoudnessMeterSettings.bDisplayAnalysisTimer;
+			SaveSettings();
+		});
+
 		// LoudnessScaleParams
 		FLoudnessMeterWidgetView::FLoudnessScaleParams LoudnessScaleParams;
 
-		LoudnessScaleParams.Range.BindLambda([this]() { return LoudnessMeterSettings.LoudnessScaleRange; });
-		LoudnessScaleParams.OnRangeValueChanged.BindLambda([this](int32 Value) { LoudnessMeterSettings.LoudnessScaleRange = Value; });
-		LoudnessScaleParams.OnRangeValueCommitted.BindLambda([this](int32 Value, ETextCommit::Type CommitType)
+		LoudnessScaleParams.Range.BindLambda([this]() { return GetLoudnessMeterSettings().LoudnessScaleRange; });
+		LoudnessScaleParams.OnRangeValueChanged.BindLambda([this](int32 Value) { GetLoudnessMeterSettings().LoudnessScaleRange = Value; });
+		LoudnessScaleParams.OnRangeValueCommitted.BindLambda([this, SaveSettings](int32 Value, ETextCommit::Type CommitType)
 		{
-			LoudnessMeterSettings.LoudnessScaleRange = Value;
+			GetLoudnessMeterSettings().LoudnessScaleRange = Value;
+			SaveSettings();
 		});
 
-		LoudnessScaleParams.Offset.BindLambda([this]() { return LoudnessMeterSettings.LoudnessScaleOffset; });
-		LoudnessScaleParams.OnOffsetValueChanged.BindLambda([this](int32 Value) { LoudnessMeterSettings.LoudnessScaleOffset = Value; });
-		LoudnessScaleParams.OnOffsetValueCommitted.BindLambda([this](int32 Value, ETextCommit::Type CommitType)
+		LoudnessScaleParams.Offset.BindLambda([this]() { return GetLoudnessMeterSettings().LoudnessScaleOffset; });
+		LoudnessScaleParams.OnOffsetValueChanged.BindLambda([this](int32 Value) { GetLoudnessMeterSettings().LoudnessScaleOffset = Value; });
+		LoudnessScaleParams.OnOffsetValueCommitted.BindLambda([this, SaveSettings](int32 Value, ETextCommit::Type CommitType)
 		{
-			LoudnessMeterSettings.LoudnessScaleOffset = Value;
+			GetLoudnessMeterSettings().LoudnessScaleOffset = Value;
+			SaveSettings();
 		});
 
-		LoudnessScaleParams.Target.BindLambda([this]() { return LoudnessMeterSettings.LoudnessScaleTarget; });
-		LoudnessScaleParams.OnTargetValueChanged.BindLambda([this](int32 Value) { LoudnessMeterSettings.LoudnessScaleTarget = Value; });
-		LoudnessScaleParams.OnTargetValueCommitted.BindLambda([this](int32 Value, ETextCommit::Type CommitType)
+		LoudnessScaleParams.Target.BindLambda([this]() { return GetLoudnessMeterSettings().LoudnessScaleTarget; });
+		LoudnessScaleParams.OnTargetValueChanged.BindLambda([this](int32 Value) { GetLoudnessMeterSettings().LoudnessScaleTarget = Value; });
+		LoudnessScaleParams.OnTargetValueCommitted.BindLambda([this, SaveSettings](int32 Value, ETextCommit::Type CommitType)
 		{
-			LoudnessMeterSettings.LoudnessScaleTarget = Value;
+			GetLoudnessMeterSettings().LoudnessScaleTarget = Value;
+			SaveSettings();
 		});
 
-		LoudnessScaleParams.TruePeakLimit.BindLambda([this]() { return LoudnessMeterSettings.TruePeakLimit; });
-		LoudnessScaleParams.OnTruePeakLimitValueChanged.BindLambda([this](float Value) { LoudnessMeterSettings.TruePeakLimit = Value; });
-		LoudnessScaleParams.OnTruePeakLimitValueCommitted.BindLambda([this](float Value, ETextCommit::Type CommitType)
+		LoudnessScaleParams.TruePeakLimit.BindLambda([this]() { return GetLoudnessMeterSettings().TruePeakLimit; });
+		LoudnessScaleParams.OnTruePeakLimitValueChanged.BindLambda([this](float Value) { GetLoudnessMeterSettings().TruePeakLimit = Value; });
+		LoudnessScaleParams.OnTruePeakLimitValueCommitted.BindLambda([this, SaveSettings](float Value, ETextCommit::Type CommitType)
 		{
-			LoudnessMeterSettings.TruePeakLimit = Value;
+			GetLoudnessMeterSettings().TruePeakLimit = Value;
+			SaveSettings();
 		});
 
-		LoudnessScaleParams.TargetColor.BindLambda([this]() { return LoudnessMeterSettings.TargetColor; });
-		LoudnessScaleParams.OnTargetColorChanged.BindLambda([this](FLinearColor NewColor) {	LoudnessMeterSettings.TargetColor = NewColor; });
-		LoudnessScaleParams.OnResetTargetColorRequested.BindLambda([this]() { LoudnessMeterSettings.TargetColor = FLoudnessMeterWidgetView::GetDefaultTargetColor(); });
+		LoudnessScaleParams.TargetColor.BindLambda([this]() { return GetLoudnessMeterSettings().TargetColor; });
+		LoudnessScaleParams.OnTargetColorChanged.BindLambda([this](FLinearColor NewColor) { GetLoudnessMeterSettings().TargetColor = NewColor; });
+		LoudnessScaleParams.OnTargetColorPickerWindowClosed.BindLambda([SaveSettings](const TSharedRef<SWindow>& Window) { SaveSettings(); });
+		LoudnessScaleParams.OnResetTargetColorRequested.BindLambda([this, SaveSettings]()
+		{
+			GetLoudnessMeterSettings().TargetColor = FLoudnessMeterWidgetView::GetDefaultTargetColor();
+			SaveSettings();
+		});
 
-		LoudnessScaleParams.bRelativeScale.BindLambda([this]() { return LoudnessMeterSettings.bUseRelativeLoudnessScale; });
-		LoudnessScaleParams.OnRelativeScaleToggleRequested.BindLambda([this]()
+		LoudnessScaleParams.bRelativeScale.BindLambda([this]() { return GetLoudnessMeterSettings().bUseRelativeLoudnessScale; });
+		LoudnessScaleParams.OnRelativeScaleToggleRequested.BindLambda([this, SaveSettings]()
 		{
+			FLoudnessMeterSettings& LoudnessMeterSettings = GetLoudnessMeterSettings();
 			LoudnessMeterSettings.bUseRelativeLoudnessScale = !LoudnessMeterSettings.bUseRelativeLoudnessScale;
+			SaveSettings();
 		});
 
 		// ValuesDisplayOrderParams
 		FLoudnessMeterWidgetView::FDisplayOrderParams ValuesDisplayOrderParams;
 
-		ValuesDisplayOrderParams.PermutationIndex.BindLambda([this]() { return LoudnessMeterSettings.ValuesOrderingPermutation; });
-		ValuesDisplayOrderParams.OnPermutationChanged.BindLambda([this](int32 NewPermutationIndex)
+		ValuesDisplayOrderParams.PermutationIndex.BindLambda([this]() { return GetLoudnessMeterSettings().ValuesOrderingPermutation; });
+		ValuesDisplayOrderParams.OnPermutationChanged.BindLambda([this, SaveSettings](int32 NewPermutationIndex)
 		{
-			LoudnessMeterSettings.ValuesOrderingPermutation = NewPermutationIndex;
+			GetLoudnessMeterSettings().ValuesOrderingPermutation = NewPermutationIndex;
+			SaveSettings();
 		});
 
 		// MetersDisplayOrderParams
 		FLoudnessMeterWidgetView::FDisplayOrderParams MetersDisplayOrderParams;
 
-		MetersDisplayOrderParams.PermutationIndex.BindLambda([this]() { return LoudnessMeterSettings.MetersOrderingPermutation; });
-		MetersDisplayOrderParams.OnPermutationChanged.BindLambda([this](int32 NewPermutationIndex)
+		MetersDisplayOrderParams.PermutationIndex.BindLambda([this]() { return GetLoudnessMeterSettings().MetersOrderingPermutation; });
+		MetersDisplayOrderParams.OnPermutationChanged.BindLambda([this, SaveSettings](int32 NewPermutationIndex)
 		{
-			LoudnessMeterSettings.MetersOrderingPermutation = NewPermutationIndex;
+			GetLoudnessMeterSettings().MetersOrderingPermutation = NewPermutationIndex;
+			SaveSettings();
 		});
 
 		// LoudnessMetrics
@@ -186,21 +239,23 @@ namespace UE::Audio::Insights
 				// Show value
 				LoudnessMetric.bShowValue.BindLambda([this, DisplayOptionsProperty]()
 				{
-					const FLoudnessMeterDisplayOptions* DisplayOptions = DisplayOptionsProperty->ContainerPtrToValuePtr<FLoudnessMeterDisplayOptions>(&LoudnessMeterSettings);
+					const FLoudnessMeterDisplayOptions* DisplayOptions = DisplayOptionsProperty->ContainerPtrToValuePtr<FLoudnessMeterDisplayOptions>(&GetLoudnessMeterSettings());
 					return DisplayOptions->bShowValue;
 				});
 
-				LoudnessMetric.OnShowValueToggleRequested.BindLambda([this, DisplayOptionsProperty]()
+				LoudnessMetric.OnShowValueToggleRequested.BindLambda([this, DisplayOptionsProperty, SaveSettings]()
 				{
-					FLoudnessMeterDisplayOptions* DisplayOptions = DisplayOptionsProperty->ContainerPtrToValuePtr<FLoudnessMeterDisplayOptions>(&LoudnessMeterSettings);
+					FLoudnessMeterDisplayOptions* DisplayOptions = DisplayOptionsProperty->ContainerPtrToValuePtr<FLoudnessMeterDisplayOptions>(&GetLoudnessMeterSettings());
 					DisplayOptions->bShowValue = !DisplayOptions->bShowValue;
 
+					SaveSettings();
+
 					LoudnessMeterWidgetView.RefreshVisibleLoudnessMetrics();
 				});
 
-				LoudnessMetric.OnShowValueToggleFromDragDropRequested.BindLambda([this, DisplayOptionsProperty]()
+				LoudnessMetric.OnShowValueToggleFromDragDropRequested.BindLambda([this, DisplayOptionsProperty, SaveSettings]()
 				{
-					FLoudnessMeterDisplayOptions* DisplayOptions = DisplayOptionsProperty->ContainerPtrToValuePtr<FLoudnessMeterDisplayOptions>(&LoudnessMeterSettings);
+					FLoudnessMeterDisplayOptions* DisplayOptions = DisplayOptionsProperty->ContainerPtrToValuePtr<FLoudnessMeterDisplayOptions>(&GetLoudnessMeterSettings());
 					DisplayOptions->bShowValue = !DisplayOptions->bShowValue;
 
 					if (DisplayOptions->bShowValue)
@@ -208,27 +263,31 @@ namespace UE::Audio::Insights
 						DisplayOptions->bHoldMaxForValue = DisplayOptions->bHoldMaxForMeter;
 					}
 
+					SaveSettings();
+
 					LoudnessMeterWidgetView.RefreshVisibleLoudnessMetrics();
 				});
 
 				// Show Meter
 				LoudnessMetric.bShowMeter.BindLambda([this, DisplayOptionsProperty]()
 				{
-					const FLoudnessMeterDisplayOptions* DisplayOptions = DisplayOptionsProperty->ContainerPtrToValuePtr<FLoudnessMeterDisplayOptions>(&LoudnessMeterSettings);
+					const FLoudnessMeterDisplayOptions* DisplayOptions = DisplayOptionsProperty->ContainerPtrToValuePtr<FLoudnessMeterDisplayOptions>(&GetLoudnessMeterSettings());
 					return DisplayOptions->bShowMeter;
 				});
 
-				LoudnessMetric.OnShowMeterToggleRequested.BindLambda([this, DisplayOptionsProperty]()
+				LoudnessMetric.OnShowMeterToggleRequested.BindLambda([this, DisplayOptionsProperty, SaveSettings]()
 				{
-					FLoudnessMeterDisplayOptions* DisplayOptions = DisplayOptionsProperty->ContainerPtrToValuePtr<FLoudnessMeterDisplayOptions>(&LoudnessMeterSettings);
+					FLoudnessMeterDisplayOptions* DisplayOptions = DisplayOptionsProperty->ContainerPtrToValuePtr<FLoudnessMeterDisplayOptions>(&GetLoudnessMeterSettings());
 					DisplayOptions->bShowMeter = !DisplayOptions->bShowMeter;
 
+					SaveSettings();
+
 					LoudnessMeterWidgetView.RefreshVisibleLoudnessMetrics();
 				});
 
-				LoudnessMetric.OnShowMeterToggleFromDragDropRequested.BindLambda([this, DisplayOptionsProperty]()
+				LoudnessMetric.OnShowMeterToggleFromDragDropRequested.BindLambda([this, DisplayOptionsProperty, SaveSettings]()
 				{
-					FLoudnessMeterDisplayOptions* DisplayOptions = DisplayOptionsProperty->ContainerPtrToValuePtr<FLoudnessMeterDisplayOptions>(&LoudnessMeterSettings);
+					FLoudnessMeterDisplayOptions* DisplayOptions = DisplayOptionsProperty->ContainerPtrToValuePtr<FLoudnessMeterDisplayOptions>(&GetLoudnessMeterSettings());
 					DisplayOptions->bShowMeter = !DisplayOptions->bShowMeter;
 
 					if (DisplayOptions->bShowMeter)
@@ -236,56 +295,70 @@ namespace UE::Audio::Insights
 						DisplayOptions->bHoldMaxForMeter = DisplayOptions->bHoldMaxForValue;
 					}
 
+					SaveSettings();
+
 					LoudnessMeterWidgetView.RefreshVisibleLoudnessMetrics();
 				});
 
 				// Hold Max for Value
 				LoudnessMetric.bHoldMaxForValue.BindLambda([this, DisplayOptionsProperty]()
 				{
-					const FLoudnessMeterDisplayOptions* DisplayOptions = DisplayOptionsProperty->ContainerPtrToValuePtr<FLoudnessMeterDisplayOptions>(&LoudnessMeterSettings);
+					const FLoudnessMeterDisplayOptions* DisplayOptions = DisplayOptionsProperty->ContainerPtrToValuePtr<FLoudnessMeterDisplayOptions>(&GetLoudnessMeterSettings());
 					return DisplayOptions->bHoldMaxForValue;
 				});
 
-				LoudnessMetric.OnHoldMaxForValueToggleRequested.BindLambda([this, DisplayOptionsProperty, LoudnessMetricName = LoudnessMetric.Name]()
+				LoudnessMetric.OnHoldMaxForValueToggleRequested.BindLambda([this, DisplayOptionsProperty, SaveSettings]()
 				{
-					FLoudnessMeterDisplayOptions* DisplayOptions = DisplayOptionsProperty->ContainerPtrToValuePtr<FLoudnessMeterDisplayOptions>(&LoudnessMeterSettings);
+					FLoudnessMeterDisplayOptions* DisplayOptions = DisplayOptionsProperty->ContainerPtrToValuePtr<FLoudnessMeterDisplayOptions>(&GetLoudnessMeterSettings());
 					DisplayOptions->bHoldMaxForValue = !DisplayOptions->bHoldMaxForValue;
+
+					SaveSettings();
 				});
 
 				// Hold Max for Meter
 				LoudnessMetric.bHoldMaxForMeter.BindLambda([this, DisplayOptionsProperty]()
 				{
-					const FLoudnessMeterDisplayOptions* DisplayOptions = DisplayOptionsProperty->ContainerPtrToValuePtr<FLoudnessMeterDisplayOptions>(&LoudnessMeterSettings);
+					const FLoudnessMeterDisplayOptions* DisplayOptions = DisplayOptionsProperty->ContainerPtrToValuePtr<FLoudnessMeterDisplayOptions>(&GetLoudnessMeterSettings());
 					return DisplayOptions->bHoldMaxForMeter;
 				});
 
-				LoudnessMetric.OnHoldMaxForMeterToggleRequested.BindLambda([this, DisplayOptionsProperty, LoudnessMetricName = LoudnessMetric.Name]()
+				LoudnessMetric.OnHoldMaxForMeterToggleRequested.BindLambda([this, DisplayOptionsProperty, SaveSettings]()
 				{
-					FLoudnessMeterDisplayOptions* DisplayOptions = DisplayOptionsProperty->ContainerPtrToValuePtr<FLoudnessMeterDisplayOptions>(&LoudnessMeterSettings);
+					FLoudnessMeterDisplayOptions* DisplayOptions = DisplayOptionsProperty->ContainerPtrToValuePtr<FLoudnessMeterDisplayOptions>(&GetLoudnessMeterSettings());
 					DisplayOptions->bHoldMaxForMeter = !DisplayOptions->bHoldMaxForMeter;
+
+					SaveSettings();
 				});
 
 				// Color
 				LoudnessMetric.Color.BindLambda([this, DisplayOptionsProperty]()
 				{
-					const FLoudnessMeterDisplayOptions* DisplayOptions = DisplayOptionsProperty->ContainerPtrToValuePtr<FLoudnessMeterDisplayOptions>(&LoudnessMeterSettings);
+					const FLoudnessMeterDisplayOptions* DisplayOptions = DisplayOptionsProperty->ContainerPtrToValuePtr<FLoudnessMeterDisplayOptions>(&GetLoudnessMeterSettings());
 					return DisplayOptions->Color;
 				});
 
 				LoudnessMetric.OnColorChanged.BindLambda([this, DisplayOptionsProperty](FLinearColor NewColor)
 				{
-					FLoudnessMeterDisplayOptions* DisplayOptions = DisplayOptionsProperty->ContainerPtrToValuePtr<FLoudnessMeterDisplayOptions>(&LoudnessMeterSettings);
+					FLoudnessMeterDisplayOptions* DisplayOptions = DisplayOptionsProperty->ContainerPtrToValuePtr<FLoudnessMeterDisplayOptions>(&GetLoudnessMeterSettings());
 					DisplayOptions->Color = NewColor;
 				});
 
-				LoudnessMetric.OnResetColorRequested.BindLambda([this, DisplayOptionsProperty]()
+				LoudnessMetric.OnColorPickerWindowClosed.BindLambda([SaveSettings](const TSharedRef<SWindow>& Window)
 				{
-					FLoudnessMeterDisplayOptions* DisplayOptions = DisplayOptionsProperty->ContainerPtrToValuePtr<FLoudnessMeterDisplayOptions>(&LoudnessMeterSettings);
+					SaveSettings();
+				});
+
+				LoudnessMetric.OnResetColorRequested.BindLambda([this, DisplayOptionsProperty, SaveSettings]()
+				{
+					FLoudnessMeterDisplayOptions* DisplayOptions = DisplayOptionsProperty->ContainerPtrToValuePtr<FLoudnessMeterDisplayOptions>(&GetLoudnessMeterSettings());
 					DisplayOptions->Color = FLoudnessMeterWidgetView::GetDefaultMeterColor();
+
+					SaveSettings();
 				});
 			}
 		}
 
+		LoudnessMeterWidgetView.InitTimerPanel(TimerPanelParams);
 		LoudnessMeterWidgetView.InitLoudnessScale(LoudnessScaleParams);
 		LoudnessMeterWidgetView.InitValuesDisplayOrder(ValuesDisplayOrderParams);
 		LoudnessMeterWidgetView.InitMetersDisplayOrder(MetersDisplayOrderParams);
@@ -312,36 +385,58 @@ namespace UE::Audio::Insights
 			]
 			+ SOverlay::Slot()
 			[
-				SNew(SSplitter)
+				SAssignNew(ContentSplitter, SSplitter)
 				.Orientation(Orient_Vertical)
 				.ResizeMode(ESplitterResizeMode::Fill)
 				.PhysicalSplitterHandleSize(4.0f)
 				.HitDetectionSplitterHandleSize(6.0f)
 				+ SSplitter::Slot()
-				.Value(0.6f)
+				.Value(LoudnessMetersSlotSize)
 				[
-					LoudnessMeterWidgetView.MakeWidget()
+					SNew(SExpandableArea)
+					.HeaderContent()
+					[
+						SNew(STextBlock)
+						.Text(LOCTEXT("OutputMeter_LoudnessMetersHeader", "Loudness Meters"))
+					]
+					.BodyContent()
+					[
+						LoudnessMeterWidgetView.MakeWidget()
+					]
+					.InitiallyCollapsed(false)
+					.OnAreaExpansionChanged(this, &FOutputMeterDashboardViewFactory::OnLoudnessMeterExpansionChanged)
 				]
 				+ SSplitter::Slot()
-				.Value(0.4f)
+				.Value(RMSMeterSlotSize)
 				[
-					SNew(SVerticalBox)
-					+ SVerticalBox::Slot()
-					.FillHeight(1.0f)
-					.HAlign(HAlign_Center)
-					.VAlign(VAlign_Fill)
+					SNew(SExpandableArea)
+					.HeaderContent()
+					[
+						SNew(STextBlock)
+						.Text(LOCTEXT("OutputMeter_RMSMeterHeader", "RMS Meter"))
+					]
+					.BodyContent()
 					[
-						SNew(SAudioMeterWidget)
-						.Orientation(EOrientation::Orient_Vertical)
-						.BackgroundColor(FLinearColor::Transparent)
-						.MeterBackgroundColor(FAudioMeterDefaultColorWidgetStyle::GetDefault().MeterBackgroundColor)
-						.MeterValueColor(FAudioMeterDefaultColorWidgetStyle::GetDefault().MeterValueColor)
-						.MeterPeakColor(FAudioMeterDefaultColorWidgetStyle::GetDefault().MeterPeakColor)
-						.MeterClippingColor(FAudioMeterDefaultColorWidgetStyle::GetDefault().MeterClippingColor)
-						.MeterScaleColor(FAudioMeterDefaultColorWidgetStyle::GetDefault().MeterScaleColor)
-						.MeterScaleLabelColor(FAudioMeterDefaultColorWidgetStyle::GetDefault().MeterScaleLabelColor)
-						.MeterChannelInfo(TAttribute<TArray<FAudioMeterChannelInfo>>::CreateSP(this, &FOutputMeterDashboardViewFactory::GetAudioMeterChannelInfo))
+						SNew(SVerticalBox)
+						+ SVerticalBox::Slot()
+						.FillHeight(1.0f)
+						.HAlign(HAlign_Center)
+						.VAlign(VAlign_Fill)
+						[
+							SNew(SAudioMeterWidget)
+							.Orientation(EOrientation::Orient_Vertical)
+							.BackgroundColor(FLinearColor::Transparent)
+							.MeterBackgroundColor(FAudioMeterDefaultColorWidgetStyle::GetDefault().MeterBackgroundColor)
+							.MeterValueColor(FAudioMeterDefaultColorWidgetStyle::GetDefault().MeterValueColor)
+							.MeterPeakColor(FAudioMeterDefaultColorWidgetStyle::GetDefault().MeterPeakColor)
+							.MeterClippingColor(FAudioMeterDefaultColorWidgetStyle::GetDefault().MeterClippingColor)
+							.MeterScaleColor(FAudioMeterDefaultColorWidgetStyle::GetDefault().MeterScaleColor)
+							.MeterScaleLabelColor(FAudioMeterDefaultColorWidgetStyle::GetDefault().MeterScaleLabelColor)
+							.MeterChannelInfo(TAttribute<TArray<FAudioMeterChannelInfo>>::CreateSP(this, &FOutputMeterDashboardViewFactory::GetAudioMeterChannelInfo))
+						]
 					]
+					.InitiallyCollapsed(false)
+					.OnAreaExpansionChanged(this, &FOutputMeterDashboardViewFactory::OnChannelMeterExpansionChanged)
 				]
 			];
 	}
@@ -389,9 +484,9 @@ namespace UE::Audio::Insights
 
 					TSharedPtr<FSubmixDashboardEntry> MainSubmixDashboardEntry = MainSubmixEntryPair ? MainSubmixEntryPair->Value : nullptr;
 
-					if (MainSubmixDashboardEntry.IsValid())
+					if (MainSubmixDashboardEntry.IsValid() && AudioMeterInfo.IsValid())
 					{
-						AudioMeterInfo = MainSubmixDashboardEntry->AudioMeterInfo.Get();
+						*AudioMeterInfo = MainSubmixDashboardEntry->AudioMeterInfo.Get();
 					}
 				}
 
@@ -405,6 +500,28 @@ namespace UE::Audio::Insights
 		if (!DataViewEntries.IsEmpty())
 		{
 			const FOutputMeterDashboardEntry& OutputMeterEntry = FOutputMeterDashboardViewFactoryPrivate::CastEntry(*DataViewEntries[0].Get());
+
+#if WITH_EDITOR
+			// Detect audio device change (e.g. PIE start/stop) and reset
+			if (LastSeenDeviceId != INDEX_NONE && OutputMeterEntry.DeviceId != LastSeenDeviceId)
+			{
+				AnalysisStartTimestamp = 0.0;
+				LatestTimestamp = 0.0;
+
+				LoudnessMeterResults    = FLoudnessMeterResults();
+				LoudnessMeterMaxResults = FLoudnessMeterMaxResults();
+			}
+
+			LastSeenDeviceId = OutputMeterEntry.DeviceId;
+#endif // WITH_EDITOR
+
+			if (AnalysisStartTimestamp == 0.0 && OutputMeterEntry.Timestamp > 0.0)
+			{
+				AnalysisStartTimestamp = OutputMeterEntry.Timestamp;
+			}
+
+			LatestTimestamp = OutputMeterEntry.Timestamp;
+
 			LoudnessMeterResults.LongTermLoudness  = OutputMeterEntry.LongTermLoudness;
 			LoudnessMeterResults.ShortTermLoudness = OutputMeterEntry.ShortTermLoudness;
 			LoudnessMeterResults.MomentaryLoudness = OutputMeterEntry.MomentaryLoudness;
@@ -415,6 +532,22 @@ namespace UE::Audio::Insights
 			LoudnessMeterMaxResults.ShortTermLoudness = FMath::Max<float>(LoudnessMeterMaxResults.ShortTermLoudness, LoudnessMeterResults.ShortTermLoudness);
 			LoudnessMeterMaxResults.MomentaryLoudness = FMath::Max<float>(LoudnessMeterMaxResults.MomentaryLoudness, LoudnessMeterResults.MomentaryLoudness);
 			LoudnessMeterMaxResults.TruePeakDb        = FMath::Max<float>(LoudnessMeterMaxResults.TruePeakDb, LoudnessMeterResults.TruePeakDb);
+
+			bWasReceivingData = true;
+		}
+		else if (bWasReceivingData)
+		{
+			// Data stopped flowing (e.g. PIE stopped, audio device destroyed) - reset timer and max values
+			AnalysisStartTimestamp = 0.0;
+			LatestTimestamp = 0.0;
+#if WITH_EDITOR
+			LastSeenDeviceId = INDEX_NONE;
+#endif // WITH_EDITOR
+
+			LoudnessMeterResults    = FLoudnessMeterResults();
+			LoudnessMeterMaxResults = FLoudnessMeterMaxResults();
+
+			bWasReceivingData = false;
 		}
 	}
 
@@ -432,25 +565,93 @@ namespace UE::Audio::Insights
 	{
 		TArray<FAudioMeterChannelInfo> MeterChannelInfo;
 
-		if (!AudioMeterInfo.EnvelopeValues.IsEmpty())
+		if (AudioMeterInfo.IsValid())
 		{
-			MeterChannelInfo.Reserve(AudioMeterInfo.EnvelopeValues.Num());
-			for (const float EnvelopeValue : AudioMeterInfo.EnvelopeValues)
+			if (!AudioMeterInfo->EnvelopeValues.IsEmpty())
 			{
-				MeterChannelInfo.Emplace(::Audio::ConvertToDecibels(EnvelopeValue), MIN_VOLUME_DECIBELS, MIN_VOLUME_DECIBELS);
+				MeterChannelInfo.Reserve(AudioMeterInfo->EnvelopeValues.Num());
+				for (const float EnvelopeValue : AudioMeterInfo->EnvelopeValues)
+				{
+					MeterChannelInfo.Emplace(::Audio::ConvertToDecibels(EnvelopeValue), MIN_VOLUME_DECIBELS, MIN_VOLUME_DECIBELS);
+				}
 			}
-		}
-		else
-		{
-			MeterChannelInfo.Reserve(AudioMeterInfo.NumChannels);
-			while (MeterChannelInfo.Num() < AudioMeterInfo.NumChannels)
+			else
 			{
-				MeterChannelInfo.Emplace(MIN_VOLUME_DECIBELS, MIN_VOLUME_DECIBELS, MIN_VOLUME_DECIBELS);
+				MeterChannelInfo.Reserve(AudioMeterInfo->NumChannels);
+				while (MeterChannelInfo.Num() < AudioMeterInfo->NumChannels)
+				{
+					MeterChannelInfo.Emplace(MIN_VOLUME_DECIBELS, MIN_VOLUME_DECIBELS, MIN_VOLUME_DECIBELS);
+				}
 			}
 		}
 
 		return MeterChannelInfo;
 	}
+
+	FTimespan FOutputMeterDashboardViewFactory::GetAnalysisTime() const
+	{
+		if (AnalysisStartTimestamp > 0.0 && LatestTimestamp > AnalysisStartTimestamp)
+		{
+			return FTimespan::FromSeconds(LatestTimestamp - AnalysisStartTimestamp);
+		}
+
+		return FTimespan::Zero();
+	}
+
+	FReply FOutputMeterDashboardViewFactory::HandleResetButtonClicked()
+	{
+		LoudnessMeterResults    = FLoudnessMeterResults();
+		LoudnessMeterMaxResults = FLoudnessMeterMaxResults();
+
+		AnalysisStartTimestamp = 0.0;
+		LatestTimestamp = 0.0;
+
+		const ::Audio::FDeviceId DeviceId = FAudioInsightsModule::GetChecked().GetDeviceId();
+		const FString ResetMainSubmixAnalyzersStr = FString::Printf(TEXT("au.ResetMainSubmixAnalyzers_AD%d"), DeviceId);
+
+		IAudioInsightsTraceModule& AudioInsightsTraceModule = IAudioInsightsModule::GetChecked().GetTraceModule();
+		AudioInsightsTraceModule.ExecuteConsoleCommand(ResetMainSubmixAnalyzersStr);
+
+		return FReply::Handled();
+	}
+
+	void FOutputMeterDashboardViewFactory::OnLoudnessMeterExpansionChanged(bool bIsExpanded)
+	{
+		if (!ContentSplitter.IsValid())
+		{
+			return;
+		}
+
+		if (bIsExpanded)
+		{
+			ContentSplitter->SlotAt(0).SetSizingRule(SSplitter::FractionOfParent);
+			ContentSplitter->SlotAt(0).SetSizeValue(LoudnessMetersSlotSize);
+		}
+		else
+		{
+			LoudnessMetersSlotSize = ContentSplitter->SlotAt(0).GetSizeValue();
+			ContentSplitter->SlotAt(0).SetSizingRule(SSplitter::SizeToContent);
+		}
+	}
+
+	void FOutputMeterDashboardViewFactory::OnChannelMeterExpansionChanged(bool bIsExpanded)
+	{
+		if (!ContentSplitter.IsValid())
+		{
+			return;
+		}
+
+		if (bIsExpanded)
+		{
+			ContentSplitter->SlotAt(1).SetSizingRule(SSplitter::FractionOfParent);
+			ContentSplitter->SlotAt(1).SetSizeValue(RMSMeterSlotSize);
+		}
+		else
+		{
+			RMSMeterSlotSize = ContentSplitter->SlotAt(1).GetSizeValue();
+			ContentSplitter->SlotAt(1).SetSizingRule(SSplitter::SizeToContent);
+		}
+	}
 } // namespace UE::Audio::Insights
 
 #undef LOCTEXT_NAMESPACE
diff --git a/Engine/Plugins/AudioInsights/Source/AudioInsights/Private/Views/OutputMeterDashboardViewFactory.h b/Engine/Plugins/AudioInsights/Source/AudioInsights/Private/Views/OutputMeterDashboardViewFactory.h
deleted file mode 100644
diff --git a/Engine/Plugins/AudioInsights/Source/AudioInsights/Public/AudioInsightsSettings.h b/Engine/Plugins/AudioInsights/Source/AudioInsights/Public/AudioInsightsSettings.h
new file mode 100644
--- /dev/null
+++ b/Engine/Plugins/AudioInsights/Source/AudioInsights/Public/AudioInsightsSettings.h
@@ -0,0 +1,22 @@
+// Copyright Epic Games, Inc. All Rights Reserved.
+#pragma once
+
+#include "Engine/DeveloperSettings.h"
+#include "LoudnessMeterSettings.h"
+
+#include "AudioInsightsSettings.generated.h"
+
+#define UE_API AUDIOINSIGHTS_API
+
+UCLASS(MinimalAPI, config = AudioInsights)
+class UAudioInsightsSettings : public UDeveloperSettings
+{
+	GENERATED_BODY()
+
+public:
+	/** Settings for the Output Meter dashboard tab */
+	UPROPERTY(EditAnywhere, config, Category = OutputMeter, meta = (ShowOnlyInnerProperties))
+	FLoudnessMeterSettings OutputMeterSettings;
+};
+
+#undef UE_API
diff --git a/Engine/Plugins/AudioInsights/Source/AudioInsights/Public/Views/DashboardViewFactory.h b/Engine/Plugins/AudioInsights/Source/AudioInsights/Public/Views/DashboardViewFactory.h
--- a/Engine/Plugins/AudioInsights/Source/AudioInsights/Public/Views/DashboardViewFactory.h
+++ b/Engine/Plugins/AudioInsights/Source/AudioInsights/Public/Views/DashboardViewFactory.h
@@ -22,8 +22,9 @@ namespace UE::Audio::Insights
 		Log,
 		Analysis,
 		AudioMeters,
-		AudioAnalyzerRack,
-		Plots
+		Plots,
+		OutputMetering,
+		AudioAnalyzerRack
 	};
 
 	class IDashboardViewFactory
@@ -95,4 +96,4 @@ namespace UE::Audio::Insights
 	};
 } // namespace UE::Audio::Insights
 
-#undef UE_API
\ No newline at end of file
+#undef UE_API
diff --git a/Engine/Plugins/AudioInsights/Source/AudioInsights/Public/Views/OutputMeterDashboardViewFactory.h b/Engine/Plugins/AudioInsights/Source/AudioInsights/Public/Views/OutputMeterDashboardViewFactory.h
new file mode 100644
--- /dev/null
+++ b/Engine/Plugins/AudioInsights/Source/AudioInsights/Public/Views/OutputMeterDashboardViewFactory.h
@@ -0,0 +1,102 @@
+// Copyright Epic Games, Inc. All Rights Reserved.
+#pragma once
+
+#include "AudioMeterChannelInfo.h"
+#include "DashboardViewFactory.h"
+#include "LoudnessMeterWidgetView.h"
+#include "Misc/Timespan.h"
+
+#define UE_API AUDIOINSIGHTS_API
+
+class SSplitter;
+struct FLoudnessMeterSettings;
+
+namespace UE::Audio::Insights
+{
+	class FOutputMeterTraceProvider;
+	class FSubmixDashboardViewFactory;
+	class FSubmixTraceProvider;
+	class IDashboardDataViewEntry;
+	struct FAudioMeterInfo;
+
+	class FOutputMeterDashboardViewFactory : public FTraceDashboardViewFactoryBase, public TSharedFromThis<FOutputMeterDashboardViewFactory>
+	{
+	public:
+		UE_API FOutputMeterDashboardViewFactory(TSharedRef<FSubmixDashboardViewFactory> SubmixDashboard);
+		UE_API virtual ~FOutputMeterDashboardViewFactory();
+
+		UE_API virtual EDefaultDashboardTabStack GetDefaultTabStack() const override;
+		UE_API virtual FText GetDisplayName() const override;
+		UE_API virtual FName GetName() const override;
+		UE_API virtual FSlateIcon GetIcon() const override;
+		UE_API virtual TSharedRef<SWidget> MakeWidget(TSharedRef<SDockTab> OwnerTab, const FSpawnTabArgs& SpawnTabArgs) override;
+
+	private:
+		struct FLoudnessMeterResults
+		{
+			float LongTermLoudness  = MIN_VOLUME_DECIBELS;
+			float ShortTermLoudness = MIN_VOLUME_DECIBELS;
+			float MomentaryLoudness = MIN_VOLUME_DECIBELS;
+			float LoudnessRangeLowerBound = 0.0f;
+			float LoudnessRangeUpperBound = 0.0f;
+
+			float TruePeakDb = MIN_VOLUME_DECIBELS;
+		};
+
+		struct FLoudnessMeterMaxResults
+		{
+			float ShortTermLoudness = MIN_VOLUME_DECIBELS;
+			float MomentaryLoudness = MIN_VOLUME_DECIBELS;
+
+			float TruePeakDb = MIN_VOLUME_DECIBELS;
+		};
+
+		FLoudnessMeterSettings& GetLoudnessMeterSettings() const;
+
+		void InitLoudnessMeterWidgetView();
+
+		void UpdateDataViewEntries();
+		void UpdateAudioMeterInfo();
+		void UpdateLoudnessValues();
+
+		virtual void ProcessEntries(FTraceDashboardViewFactoryBase::EProcessReason Reason) override;
+
+		TArray<FAudioMeterChannelInfo> GetAudioMeterChannelInfo() const;
+
+		FTimespan GetAnalysisTime() const;
+		FReply HandleResetButtonClicked();
+
+		void OnLoudnessMeterExpansionChanged(const bool bInIsExpanded);
+		void OnChannelMeterExpansionChanged(const bool bInIsExpanded);
+
+		TSharedPtr<FOutputMeterTraceProvider> OutputMeterProvider;
+		TArray<TSharedPtr<IDashboardDataViewEntry>> DataViewEntries;
+
+		// Audio meter related members
+		TSharedPtr<const FSubmixTraceProvider> SubmixProvider;
+		TOptional<uint64> ProcessedUpdateId;
+		TUniquePtr<FAudioMeterInfo> AudioMeterInfo;
+
+		// Loudness meters related members
+		AudioWidgetsCore::FLoudnessMeterWidgetView LoudnessMeterWidgetView;
+		bool bIsLoudnessMeterInitialized = false;
+
+		// Layout splitter for collapse/expand support
+		TSharedPtr<SSplitter> ContentSplitter;
+		float LoudnessMetersSlotSize = 0.75f;
+		float RMSMeterSlotSize = 0.25f;
+
+		FLoudnessMeterResults LoudnessMeterResults;
+		FLoudnessMeterMaxResults LoudnessMeterMaxResults;
+
+		double AnalysisStartTimestamp = 0.0;
+		double LatestTimestamp = 0.0;
+		bool bWasReceivingData = false;
+
+#if WITH_EDITOR
+		::Audio::FDeviceId LastSeenDeviceId = INDEX_NONE;
+#endif // WITH_EDITOR
+	};
+} // namespace UE::Audio::Insights
+
+#undef UE_API
diff --git a/Engine/Plugins/AudioInsights/Source/AudioInsightsEditor/AudioInsightsEditor.Build.cs b/Engine/Plugins/AudioInsights/Source/AudioInsightsEditor/AudioInsightsEditor.Build.cs
--- a/Engine/Plugins/AudioInsights/Source/AudioInsightsEditor/AudioInsightsEditor.Build.cs
+++ b/Engine/Plugins/AudioInsights/Source/AudioInsightsEditor/AudioInsightsEditor.Build.cs
@@ -24,6 +24,7 @@ public class AudioInsightsEditor : ModuleRules
 				"AudioInsights",
 				"AudioMixer",
 				"AudioWidgets",
+				"AudioWidgetsCore",
 				"CoreUObject",
 				"DeveloperSettings",
 				"Engine",
diff --git a/Engine/Plugins/AudioInsights/Source/AudioInsightsEditor/Private/AudioInsightsEditorDashboardFactory.cpp b/Engine/Plugins/AudioInsights/Source/AudioInsightsEditor/Private/AudioInsightsEditorDashboardFactory.cpp
--- a/Engine/Plugins/AudioInsights/Source/AudioInsightsEditor/Private/AudioInsightsEditorDashboardFactory.cpp
+++ b/Engine/Plugins/AudioInsights/Source/AudioInsightsEditor/Private/AudioInsightsEditorDashboardFactory.cpp
@@ -502,18 +502,19 @@ namespace UE::Audio::Insights
 						{
 							if (DashboardTabManager.IsValid())
 							{
-								if (TSharedPtr<SDockTab> ViewportTab = DashboardTabManager->FindExistingLiveTab(FactoryName);
-									!ViewportTab.IsValid())
+								const TSharedPtr<SDockTab> ViewportTab = DashboardTabManager->FindExistingLiveTab(FactoryName);
+
+								if (!ViewportTab.IsValid())
 								{
-									DashboardTabManager->TryInvokeTab(FactoryName);
+									const TSharedPtr<SDockTab> InvokedTab = DashboardTabManager->TryInvokeTab(FactoryName);
 
-									if (TSharedPtr<SDockTab> InvokedOutputMeterTab = DashboardTabManager->TryInvokeTab(FactoryName);
-										InvokedOutputMeterTab.IsValid() && DashboardViewFactories[FactoryName].IsValid())
+									if (InvokedTab.IsValid() && DashboardViewFactories[FactoryName].IsValid())
 									{
-										if (const EDefaultDashboardTabStack DefaultTabStack = DashboardViewFactories[FactoryName]->GetDefaultTabStack();
-											DefaultTabStack == EDefaultDashboardTabStack::AudioAnalyzerRack)
+										const EDefaultDashboardTabStack DefaultTabStack = DashboardViewFactories[FactoryName]->GetDefaultTabStack();
+
+										if (DefaultTabStack == EDefaultDashboardTabStack::AudioAnalyzerRack || DefaultTabStack == EDefaultDashboardTabStack::OutputMetering)
 										{
-											InvokedOutputMeterTab->SetParentDockTabStackTabWellHidden(true);
+											InvokedTab->SetParentDockTabStackTabWellHidden(true);
 										}
 									}
 								}
@@ -1290,7 +1291,8 @@ namespace UE::Audio::Insights
 		TSharedRef<FTabManager::FStack> AnalysisTabStack = FTabManager::NewStack();
 		TSharedRef<FTabManager::FStack> AudioMetersTabStack = FTabManager::NewStack();
 		TSharedRef<FTabManager::FStack> PlotsTabStack = FTabManager::NewStack();
-		TSharedRef<FTabManager::FStack> AudioAnalyzerRackTabStack = FTabManager::NewStack()->SetHideTabWell(true)->SetSizeCoefficient(0.15f);
+		TSharedRef<FTabManager::FStack> OutputMeteringTabStack = FTabManager::NewStack()->SetHideTabWell(true);
+		TSharedRef<FTabManager::FStack> AudioAnalyzerRackTabStack = FTabManager::NewStack()->SetHideTabWell(true);
 
 		for (const auto& KVP : DashboardViewFactories)
 		{
@@ -1330,6 +1332,12 @@ namespace UE::Audio::Insights
 				}
 				break;
 
+				case EDefaultDashboardTabStack::OutputMetering:
+				{
+					OutputMeteringTabStack->AddTab(FactoryName, ETabState::OpenedTab);
+				}
+				break;
+
 				case EDefaultDashboardTabStack::AudioAnalyzerRack:
 				{
 					AudioAnalyzerRackTabStack->AddTab(FactoryName, ETabState::OpenedTab);
@@ -1343,7 +1351,7 @@ namespace UE::Audio::Insights
 
 		AnalysisTabStack->SetForegroundTab(FName("Sounds"));
 
-		return FTabManager::NewLayout("AudioDashboard_Editor_Layout_v2")
+		return FTabManager::NewLayout("AudioDashboard_Editor_Layout_v3")
 		->AddArea
 		(
 			FTabManager::NewPrimaryArea()
@@ -1355,10 +1363,10 @@ namespace UE::Audio::Insights
 				->SetOrientation(Orient_Horizontal)
 				->Split
 				(
-					// Top
 					FTabManager::NewSplitter()
 					->SetOrientation(Orient_Vertical)
-					->SetSizeCoefficient(0.25f) // Column width
+					->SetSizeCoefficient(0.2f) // Column width
+					// Top
 					->Split
 					(
 						ViewportTabStack
@@ -1375,10 +1383,10 @@ namespace UE::Audio::Insights
 				// Middle column
 				->Split
 				(
-					// Top
 					FTabManager::NewSplitter()
 					->SetOrientation(Orient_Vertical)
 					->SetSizeCoefficient(0.6f) // Column width
+					// Top
 					->Split
 					(
 						FTabManager::NewSplitter()
@@ -1409,10 +1417,23 @@ namespace UE::Audio::Insights
 						)
 					)
 				)
+
 				// Right column
 				->Split
 				(
-					AudioAnalyzerRackTabStack
+					FTabManager::NewSplitter()
+					->SetOrientation(Orient_Vertical)
+					->SetSizeCoefficient(0.2f) // Column width
+					->Split
+					(
+						OutputMeteringTabStack
+						->SetSizeCoefficient(0.8f)
+					)
+					->Split
+					(
+						AudioAnalyzerRackTabStack
+						->SetSizeCoefficient(0.2f)
+					)
 				)
 			)
 		);
diff --git a/Engine/Plugins/AudioInsights/Source/AudioInsightsEditor/Private/AudioInsightsEditorModule.cpp b/Engine/Plugins/AudioInsights/Source/AudioInsightsEditor/Private/AudioInsightsEditorModule.cpp
--- a/Engine/Plugins/AudioInsights/Source/AudioInsightsEditor/Private/AudioInsightsEditorModule.cpp
+++ b/Engine/Plugins/AudioInsights/Source/AudioInsightsEditor/Private/AudioInsightsEditorModule.cpp
@@ -14,6 +14,7 @@
 #include "Views/AudioMetersPanelDashboardViewFactory.h"
 #include "Views/DetailsDashboardViewFactory.h"
 #include "Views/LogDashboardViewFactory.h"
+#include "Views/OutputMeterDashboardViewFactory.h"
 #include "Views/SignalFlowDashboardViewFactory.h"
 #include "Views/SoundDashboardViewFactory.h"
 #include "Views/SoundPlotsDashboardViewFactory.h"
@@ -40,6 +41,7 @@ namespace UE::Audio::Insights
 			
 			TSharedRef<FSoundDashboardViewFactory> SoundsDashboard = MakeShared<FSoundDashboardViewFactory>();
 			TSharedRef<FSoundPlotsDashboardViewFactory> PlotsDashboard = MakeShared<FSoundPlotsDashboardViewFactory>();
+			TSharedRef<FSubmixDashboardViewFactory> SubmixDashboard = MakeShared<FSubmixDashboardViewFactory>();
 
 			// @TODO UE-274216: Decide what to do with the Viewport dashboard
 			//DashboardFactory->RegisterViewFactory(MakeShared<FViewportDashboardViewFactory>());
@@ -48,11 +50,12 @@ namespace UE::Audio::Insights
 			DashboardFactory->RegisterViewFactory(MakeShared<FAudioEventLogDashboardViewFactory>());
 			DashboardFactory->RegisterViewFactory(SoundsDashboard);
 			DashboardFactory->RegisterViewFactory(MakeShared<FVirtualLoopDashboardViewFactory>());
-			DashboardFactory->RegisterViewFactory(MakeShared<FSubmixDashboardViewFactory>());
+			DashboardFactory->RegisterViewFactory(SubmixDashboard);
 			DashboardFactory->RegisterViewFactory(MakeShared<FAudioBusDashboardViewFactory>());
 			DashboardFactory->RegisterViewFactory(MakeShared<FAudioMetersPanelDashboardViewFactory>());
 			DashboardFactory->RegisterViewFactory(MakeShared<FSignalFlowDashboardViewFactory>());
 			DashboardFactory->RegisterViewFactory(PlotsDashboard);
+			DashboardFactory->RegisterViewFactory(MakeShared<FOutputMeterDashboardViewFactory>(SubmixDashboard));
 			DashboardFactory->RegisterViewFactory(MakeShared<FAudioAnalyzerRackDashboardViewFactory>());
 
 			PlotsDashboard->InitPlots(SoundsDashboard);
