[Audio Insights] LoudnessMeters and RMSMeter slot sizes and collapsed states are now persistent settings.

#Synced-CL 52505914

#Synced-CL 52505914

diff --git a/Engine/Plugins/AudioInsights/Source/AudioInsights/Private/Views/OutputMeterDashboardViewFactory.cpp b/Engine/Plugins/AudioInsights/Source/AudioInsights/Private/Views/OutputMeterDashboardViewFactory.cpp
--- a/Engine/Plugins/AudioInsights/Source/AudioInsights/Private/Views/OutputMeterDashboardViewFactory.cpp
+++ b/Engine/Plugins/AudioInsights/Source/AudioInsights/Private/Views/OutputMeterDashboardViewFactory.cpp
@@ -75,6 +75,19 @@ namespace UE::Audio::Insights
 		return GetMutableDefault<UAudioInsightsSettings>()->OutputMeterSettings;
 	}
 
+	FOutputMeterDashboardSettings& FOutputMeterDashboardViewFactory::GetOutputMeterDashboardSettings()
+	{
+		return GetMutableDefault<UAudioInsightsSettings>()->OutputMeterDashboardSettings;
+	}
+
+	void FOutputMeterDashboardViewFactory::SaveOutputMeterDashboardSettings()
+	{
+		const FProperty* Property = UAudioInsightsSettings::StaticClass()->FindPropertyByName(GET_MEMBER_NAME_CHECKED(UAudioInsightsSettings, OutputMeterDashboardSettings));
+
+		UAudioInsightsSettings* AudioInsightsSettings = GetMutableDefault<UAudioInsightsSettings>();
+		AudioInsightsSettings->UpdateSinglePropertyInConfigFile(Property, GetConfigFilename(AudioInsightsSettings));
+	}
+
 	void FOutputMeterDashboardViewFactory::InitLoudnessMeterWidgetView()
 	{
 		using namespace AudioWidgetsCore;
@@ -377,7 +390,7 @@ namespace UE::Audio::Insights
 			bIsLoudnessMeterInitialized = true;
 		}
 
-		return SNew(SOverlay)
+		TSharedRef<SWidget> Widget = SNew(SOverlay)
 			+ SOverlay::Slot()
 			[
 				SNew(SColorBlock)
@@ -390,8 +403,9 @@ namespace UE::Audio::Insights
 				.ResizeMode(ESplitterResizeMode::Fill)
 				.PhysicalSplitterHandleSize(4.0f)
 				.HitDetectionSplitterHandleSize(6.0f)
+				.OnSplitterFinishedResizing(this, &FOutputMeterDashboardViewFactory::OnSplitterFinishedResizing)
 				+ SSplitter::Slot()
-				.Value(LoudnessMetersSlotSize)
+				.Value(GetOutputMeterDashboardSettings().LoudnessMetersSlotSize)
 				[
 					SNew(SExpandableArea)
 					.HeaderContent()
@@ -403,11 +417,11 @@ namespace UE::Audio::Insights
 					[
 						LoudnessMeterWidgetView.MakeWidget()
 					]
-					.InitiallyCollapsed(false)
+					.InitiallyCollapsed(GetOutputMeterDashboardSettings().bLoudnessMetersCollapsed)
 					.OnAreaExpansionChanged(this, &FOutputMeterDashboardViewFactory::OnLoudnessMeterExpansionChanged)
 				]
 				+ SSplitter::Slot()
-				.Value(RMSMeterSlotSize)
+				.Value(GetOutputMeterDashboardSettings().RMSMeterSlotSize)
 				[
 					SNew(SExpandableArea)
 					.HeaderContent()
@@ -435,10 +449,25 @@ namespace UE::Audio::Insights
 							.MeterChannelInfo(TAttribute<TArray<FAudioMeterChannelInfo>>::CreateSP(this, &FOutputMeterDashboardViewFactory::GetAudioMeterChannelInfo))
 						]
 					]
-					.InitiallyCollapsed(false)
+					.InitiallyCollapsed(GetOutputMeterDashboardSettings().bRMSMeterCollapsed)
 					.OnAreaExpansionChanged(this, &FOutputMeterDashboardViewFactory::OnChannelMeterExpansionChanged)
 				]
 			];
+
+		// Restore collapsed sizing rules from settings
+		const FOutputMeterDashboardSettings& OutputMeterDashboardSettings = GetOutputMeterDashboardSettings();
+
+		if (OutputMeterDashboardSettings.bLoudnessMetersCollapsed)
+		{
+			ContentSplitter->SlotAt(0).SetSizingRule(SSplitter::SizeToContent);
+		}
+
+		if (OutputMeterDashboardSettings.bRMSMeterCollapsed)
+		{
+			ContentSplitter->SlotAt(1).SetSizingRule(SSplitter::SizeToContent);
+		}
+
+		return Widget;
 	}
 
 	void FOutputMeterDashboardViewFactory::UpdateDataViewEntries()
@@ -615,42 +644,76 @@ namespace UE::Audio::Insights
 		return FReply::Handled();
 	}
 
-	void FOutputMeterDashboardViewFactory::OnLoudnessMeterExpansionChanged(bool bIsExpanded)
+	void FOutputMeterDashboardViewFactory::OnLoudnessMeterExpansionChanged(const bool bIsExpanded)
 	{
 		if (!ContentSplitter.IsValid())
 		{
 			return;
 		}
 
+		FOutputMeterDashboardSettings& OutputMeterDashboardSettings = GetOutputMeterDashboardSettings();
+
+		OutputMeterDashboardSettings.bLoudnessMetersCollapsed = !bIsExpanded;
+
 		if (bIsExpanded)
 		{
 			ContentSplitter->SlotAt(0).SetSizingRule(SSplitter::FractionOfParent);
-			ContentSplitter->SlotAt(0).SetSizeValue(LoudnessMetersSlotSize);
+			ContentSplitter->SlotAt(0).SetSizeValue(OutputMeterDashboardSettings.LoudnessMetersSlotSize);
 		}
 		else
 		{
-			LoudnessMetersSlotSize = ContentSplitter->SlotAt(0).GetSizeValue();
+			OutputMeterDashboardSettings.LoudnessMetersSlotSize = ContentSplitter->SlotAt(0).GetSizeValue();
 			ContentSplitter->SlotAt(0).SetSizingRule(SSplitter::SizeToContent);
 		}
+
+		SaveOutputMeterDashboardSettings();
 	}
 
-	void FOutputMeterDashboardViewFactory::OnChannelMeterExpansionChanged(bool bIsExpanded)
+	void FOutputMeterDashboardViewFactory::OnChannelMeterExpansionChanged(const bool bIsExpanded)
 	{
 		if (!ContentSplitter.IsValid())
 		{
 			return;
 		}
 
+		FOutputMeterDashboardSettings& OutputMeterDashboardSettings = GetOutputMeterDashboardSettings();
+
+		OutputMeterDashboardSettings.bRMSMeterCollapsed = !bIsExpanded;
+
 		if (bIsExpanded)
 		{
 			ContentSplitter->SlotAt(1).SetSizingRule(SSplitter::FractionOfParent);
-			ContentSplitter->SlotAt(1).SetSizeValue(RMSMeterSlotSize);
+			ContentSplitter->SlotAt(1).SetSizeValue(OutputMeterDashboardSettings.RMSMeterSlotSize);
 		}
 		else
 		{
-			RMSMeterSlotSize = ContentSplitter->SlotAt(1).GetSizeValue();
+			OutputMeterDashboardSettings.RMSMeterSlotSize = ContentSplitter->SlotAt(1).GetSizeValue();
 			ContentSplitter->SlotAt(1).SetSizingRule(SSplitter::SizeToContent);
 		}
+
+		SaveOutputMeterDashboardSettings();
+	}
+
+	void FOutputMeterDashboardViewFactory::OnSplitterFinishedResizing()
+	{
+		if (!ContentSplitter.IsValid())
+		{
+			return;
+		}
+
+		FOutputMeterDashboardSettings& OutputMeterDashboardSettings = GetOutputMeterDashboardSettings();
+
+		if (!OutputMeterDashboardSettings.bLoudnessMetersCollapsed)
+		{
+			OutputMeterDashboardSettings.LoudnessMetersSlotSize = ContentSplitter->SlotAt(0).GetSizeValue();
+		}
+
+		if (!OutputMeterDashboardSettings.bRMSMeterCollapsed)
+		{
+			OutputMeterDashboardSettings.RMSMeterSlotSize = ContentSplitter->SlotAt(1).GetSizeValue();
+		}
+
+		SaveOutputMeterDashboardSettings();
 	}
 } // namespace UE::Audio::Insights
 
diff --git a/Engine/Plugins/AudioInsights/Source/AudioInsights/Public/AudioInsightsSettings.h b/Engine/Plugins/AudioInsights/Source/AudioInsights/Public/AudioInsightsSettings.h
--- a/Engine/Plugins/AudioInsights/Source/AudioInsights/Public/AudioInsightsSettings.h
+++ b/Engine/Plugins/AudioInsights/Source/AudioInsights/Public/AudioInsightsSettings.h
@@ -3,6 +3,7 @@
 
 #include "Engine/DeveloperSettings.h"
 #include "LoudnessMeterSettings.h"
+#include "Settings/OutputMeterDashboardSettings.h"
 
 #include "AudioInsightsSettings.generated.h"
 
@@ -14,9 +15,13 @@ class UAudioInsightsSettings : public UDeveloperSettings
 	GENERATED_BODY()
 
 public:
-	/** Settings for the Output Meter dashboard tab */
+	/** Settings for the Output Meters */
 	UPROPERTY(EditAnywhere, config, Category = OutputMeter, meta = (ShowOnlyInnerProperties))
 	FLoudnessMeterSettings OutputMeterSettings;
+
+	/** Settings for the Output Meter dashboard tab (not visible to the user) */
+	UPROPERTY(config)
+	FOutputMeterDashboardSettings OutputMeterDashboardSettings;
 };
 
 #undef UE_API
diff --git a/Engine/Plugins/AudioInsights/Source/AudioInsights/Public/Settings/OutputMeterDashboardSettings.h b/Engine/Plugins/AudioInsights/Source/AudioInsights/Public/Settings/OutputMeterDashboardSettings.h
new file mode 100644
--- /dev/null
+++ b/Engine/Plugins/AudioInsights/Source/AudioInsights/Public/Settings/OutputMeterDashboardSettings.h
@@ -0,0 +1,22 @@
+// Copyright Epic Games, Inc. All Rights Reserved.
+#pragma once
+
+#include "OutputMeterDashboardSettings.generated.h"
+
+USTRUCT()
+struct FOutputMeterDashboardSettings
+{
+	GENERATED_BODY()
+
+	UPROPERTY(config)
+	float LoudnessMetersSlotSize = 0.75f;
+
+	UPROPERTY(config)
+	float RMSMeterSlotSize = 0.25f;
+
+	UPROPERTY(config)
+	bool bLoudnessMetersCollapsed = false;
+
+	UPROPERTY(config)
+	bool bRMSMeterCollapsed = false;
+};
diff --git a/Engine/Plugins/AudioInsights/Source/AudioInsights/Public/Views/OutputMeterDashboardViewFactory.h b/Engine/Plugins/AudioInsights/Source/AudioInsights/Public/Views/OutputMeterDashboardViewFactory.h
--- a/Engine/Plugins/AudioInsights/Source/AudioInsights/Public/Views/OutputMeterDashboardViewFactory.h
+++ b/Engine/Plugins/AudioInsights/Source/AudioInsights/Public/Views/OutputMeterDashboardViewFactory.h
@@ -10,6 +10,7 @@
 
 class SSplitter;
 struct FLoudnessMeterSettings;
+struct FOutputMeterDashboardSettings;
 
 namespace UE::Audio::Insights
 {
@@ -53,6 +54,9 @@ namespace UE::Audio::Insights
 
 		FLoudnessMeterSettings& GetLoudnessMeterSettings() const;
 
+		FOutputMeterDashboardSettings& GetOutputMeterDashboardSettings();
+		void SaveOutputMeterDashboardSettings();
+
 		void InitLoudnessMeterWidgetView();
 
 		void UpdateDataViewEntries();
@@ -66,8 +70,9 @@ namespace UE::Audio::Insights
 		FTimespan GetAnalysisTime() const;
 		FReply HandleResetButtonClicked();
 
-		void OnLoudnessMeterExpansionChanged(const bool bInIsExpanded);
-		void OnChannelMeterExpansionChanged(const bool bInIsExpanded);
+		void OnLoudnessMeterExpansionChanged(const bool bIsExpanded);
+		void OnChannelMeterExpansionChanged(const bool bIsExpanded);
+		void OnSplitterFinishedResizing();
 
 		TSharedPtr<FOutputMeterTraceProvider> OutputMeterProvider;
 		TArray<TSharedPtr<IDashboardDataViewEntry>> DataViewEntries;
@@ -83,8 +88,6 @@ namespace UE::Audio::Insights
 
 		// Layout splitter for collapse/expand support
 		TSharedPtr<SSplitter> ContentSplitter;
-		float LoudnessMetersSlotSize = 0.75f;
-		float RMSMeterSlotSize = 0.25f;
 
 		FLoudnessMeterResults LoudnessMeterResults;
 		FLoudnessMeterMaxResults LoudnessMeterMaxResults;
