[Audio Insights] Add asset icon and name header to Audio Insights Details dashboard. Displays the selected asset class icon and name at the top of the Details panel using FSlateIconFinder for generic class icon resolution.

#Synced-CL 52502761

diff --git a/Engine/Plugins/AudioInsights/Source/AudioInsightsEditor/Private/Views/DetailsDashboardViewFactory.cpp b/Engine/Plugins/AudioInsights/Source/AudioInsightsEditor/Private/Views/DetailsDashboardViewFactory.cpp
--- a/Engine/Plugins/AudioInsights/Source/AudioInsightsEditor/Private/Views/DetailsDashboardViewFactory.cpp
+++ b/Engine/Plugins/AudioInsights/Source/AudioInsightsEditor/Private/Views/DetailsDashboardViewFactory.cpp
@@ -7,6 +7,9 @@
 #include "IDetailsView.h"
 #include "Modules/ModuleManager.h"
 #include "PropertyEditorModule.h"
+#include "Styling/AppStyle.h"
+#include "Styling/SlateIconFinder.h"
+#include "Widgets/Images/SImage.h"
 #include "Widgets/Layout/SBox.h"
 #include "Widgets/Text/STextBlock.h"
 
@@ -34,7 +37,7 @@ namespace UE::Audio::Insights
 		return FSlateStyle::Get().CreateIcon("AudioInsights.Icon.Details");
 	}
 
-	TSharedRef<SWidget> FDetailsDashboardViewFactory::MakeWidget(const TSharedRef<SDockTab> OwnerTab, const FSpawnTabArgs& SpawnTabArgs)
+	TSharedRef<SWidget> FDetailsDashboardViewFactory::MakeWidget(TSharedRef<SDockTab> OwnerTab, const FSpawnTabArgs& SpawnTabArgs)
 	{
 		if (!DashboardWidget.IsValid())
 		{
@@ -49,26 +52,64 @@ namespace UE::Audio::Insights
 		return DashboardWidget.ToSharedRef();
 	}
 
-	TSharedPtr<SWidget> FDetailsDashboardViewFactory::CreateDetailsViewWidget()
+	TSharedRef<SWidget> FDetailsDashboardViewFactory::CreateAssetHeaderWidget()
+	{
+		return SNew(SBox)
+			.Padding(8.0f, 4.0f)
+			[
+				SNew(SHorizontalBox)
+
+				+ SHorizontalBox::Slot()
+				.AutoWidth()
+				.VAlign(VAlign_Center)
+				.Padding(0.0f, 0.0f, 4.0f, 0.0f)
+				[
+					SAssignNew(AssetIconImage, SImage)
+				]
+
+				+ SHorizontalBox::Slot()
+				.FillWidth(1.0f)
+				.VAlign(VAlign_Center)
+				[
+					SAssignNew(AssetNameText, STextBlock)
+					.TextStyle(&FAppStyle::Get().GetWidgetStyle<FTextBlockStyle>("DetailsView.ConstantTextBlockStyle"))
+					.OverflowPolicy(ETextOverflowPolicy::Ellipsis)
+				]
+			];
+	}
+
+	TSharedRef<SWidget> FDetailsDashboardViewFactory::CreateDetailsViewWidget()
 	{
 		FPropertyEditorModule& PropertyEditorModule = FModuleManager::LoadModuleChecked<FPropertyEditorModule>("PropertyEditor");
 
 		FDetailsViewArgs DetailsViewArgs;
 		DetailsViewArgs.bAllowSearch = true;
 		DetailsViewArgs.bHideSelectionTip = true;
-		DetailsViewArgs.bShowObjectLabel = true;
+		DetailsViewArgs.bShowObjectLabel = false;
 
 		DetailsView = PropertyEditorModule.CreateDetailView(DetailsViewArgs);
 
-		DetailsViewWidget = DetailsView->AsShared();
-		DetailsViewWidget->SetVisibility(EVisibility::Collapsed);
+		DetailsViewWidget = SNew(SVerticalBox)
+			.Visibility(EVisibility::Collapsed)
+
+			+ SVerticalBox::Slot()
+			.AutoHeight()
+			[
+				CreateAssetHeaderWidget()
+			]
+
+			+ SVerticalBox::Slot()
+			.FillHeight(1.0f)
+			[
+				DetailsView->AsShared()
+			];
 
 		NoSelectionMessageWidget = SNew(SBox)
 			.HAlign(HAlign_Center)
 			.VAlign(VAlign_Center)
 			[
 				SNew(STextBlock)
-					.Text(LOCTEXT("AudioDashboard_Details_NoSelection", "Select an audio asset to view its properties."))
+				.Text(LOCTEXT("AudioDashboard_Details_NoSelection", "Select an audio asset to view its properties."))
 			];
 
 		return SNew(SOverlay)
@@ -96,7 +137,7 @@ namespace UE::Audio::Insights
 
 	void FDetailsDashboardViewFactory::OnAssetSelectionChanged(const TObjectPtr<UObject> InAsset)
 	{
-		if (!DetailsView.IsValid() || !DetailsViewWidget.IsValid() || !NoSelectionMessageWidget.IsValid() || !DashboardWidget.IsValid())
+		if (!DetailsView.IsValid() || !DetailsViewWidget.IsValid() || !NoSelectionMessageWidget.IsValid() || !DashboardWidget.IsValid() || !AssetIconImage.IsValid() || !AssetNameText.IsValid())
 		{
 			return;
 		}
@@ -104,12 +145,20 @@ namespace UE::Audio::Insights
 		if (InAsset)
 		{
 			DetailsView->SetObject(InAsset);
+
+			AssetIconImage->SetImage(FSlateIconFinder::FindIconForClass(InAsset->GetClass()).GetIcon());
+			AssetNameText->SetText(FText::FromName(InAsset->GetFName()));
+
 			DetailsViewWidget->SetVisibility(EVisibility::Visible);
 			NoSelectionMessageWidget->SetVisibility(EVisibility::Collapsed);
 		}
 		else
 		{
 			DetailsView->SetObject(nullptr);
+
+			AssetIconImage->SetImage(nullptr);
+			AssetNameText->SetText(FText::GetEmpty());
+
 			DetailsViewWidget->SetVisibility(EVisibility::Collapsed);
 			NoSelectionMessageWidget->SetVisibility(EVisibility::Visible);
 		}
diff --git a/Engine/Plugins/AudioInsights/Source/AudioInsightsEditor/Private/Views/DetailsDashboardViewFactory.h b/Engine/Plugins/AudioInsights/Source/AudioInsightsEditor/Private/Views/DetailsDashboardViewFactory.h
--- a/Engine/Plugins/AudioInsights/Source/AudioInsightsEditor/Private/Views/DetailsDashboardViewFactory.h
+++ b/Engine/Plugins/AudioInsights/Source/AudioInsightsEditor/Private/Views/DetailsDashboardViewFactory.h
@@ -9,6 +9,8 @@
 #include "Widgets/SWidget.h"
 
 class IDetailsView;
+class SImage;
+class STextBlock;
 
 namespace UE::Audio::Insights
 {
@@ -22,10 +24,11 @@ namespace UE::Audio::Insights
 		virtual FText GetDisplayName() const override;
 		virtual EDefaultDashboardTabStack GetDefaultTabStack() const override;
 		virtual FSlateIcon GetIcon() const override;
-		virtual TSharedRef<SWidget> MakeWidget(const TSharedRef<SDockTab> OwnerTab, const FSpawnTabArgs& SpawnTabArgs) override;
+		virtual TSharedRef<SWidget> MakeWidget(TSharedRef<SDockTab> OwnerTab, const FSpawnTabArgs& SpawnTabArgs) override;
 
 	private:
-		TSharedPtr<SWidget> CreateDetailsViewWidget();
+		TSharedRef<SWidget> CreateAssetHeaderWidget();
+		TSharedRef<SWidget> CreateDetailsViewWidget();
 		void BindOnAssetChangedDelegate();
 
 		void OnAssetSelectionChanged(const TObjectPtr<UObject> InAsset);
@@ -35,6 +38,9 @@ namespace UE::Audio::Insights
 		TSharedPtr<SWidget> NoSelectionMessageWidget;
 		TSharedPtr<SWidget> DashboardWidget;
 
+		TSharedPtr<SImage> AssetIconImage;
+		TSharedPtr<STextBlock> AssetNameText;
+
 		FDelegateHandle AssetSelectionChangedHandle;
 	};
 } // namespace UE::Audio::Insights
