[Audio Insights] Expand Details panel selection to all relevant dashboards. Implements OnSelectionChanged in FTraceObjectTableDashboardViewFactory base class with TryLoad fallback for unloaded assets. Adds GetObjectPath to IObjectDashboardEntry. Covers AudioBus, Submix, VirtualLoop, AudioEventLog, AudioModulation, and SignalFlow dashboards.

#Synced-CL 52502761

diff --git a/Engine/Plugins/AudioInsights/Source/AudioInsights/Private/Views/AudioEventLogDashboardViewFactory.cpp b/Engine/Plugins/AudioInsights/Source/AudioInsights/Private/Views/AudioEventLogDashboardViewFactory.cpp
--- a/Engine/Plugins/AudioInsights/Source/AudioInsights/Private/Views/AudioEventLogDashboardViewFactory.cpp
+++ b/Engine/Plugins/AudioInsights/Source/AudioInsights/Private/Views/AudioEventLogDashboardViewFactory.cpp
@@ -843,6 +843,8 @@ namespace UE::Audio::Insights
 	{
 		using namespace AudioEventLogPrivate;
 
+		FTraceObjectTableDashboardViewFactory::OnSelectionChanged(SelectedItem, SelectInfo);
+
 		if (!FilteredEntriesListView.IsValid())
 		{
 			return;
diff --git a/Engine/Plugins/AudioInsights/Source/AudioInsights/Private/Views/SignalFlowDashboardViewFactory.cpp b/Engine/Plugins/AudioInsights/Source/AudioInsights/Private/Views/SignalFlowDashboardViewFactory.cpp
--- a/Engine/Plugins/AudioInsights/Source/AudioInsights/Private/Views/SignalFlowDashboardViewFactory.cpp
+++ b/Engine/Plugins/AudioInsights/Source/AudioInsights/Private/Views/SignalFlowDashboardViewFactory.cpp
@@ -20,6 +20,7 @@
 #include "Widgets/Layout/SSeparator.h"
 
 #if WITH_EDITOR
+#include "AudioInsightsDetailsSelectionManager.h"
 #include "Editor.h"
 #include "Subsystems/AssetEditorSubsystem.h"
 #else
@@ -1263,6 +1264,30 @@ namespace UE::Audio::Insights
 		}
 
 		SetNewSelection(InSelectedItem);
+
+#if WITH_EDITOR
+		FAudioInsightsDetailsSelectionManager& SelectionManager = IAudioInsightsModule::GetChecked().GetDetailsSelectionManager();
+		TObjectPtr<UObject> Object = InSelectedItem.IsValid() ? InSelectedItem->GetObject() : nullptr;
+
+		// If we could not find a loaded asset for this object, try to load one using the object path
+		if (Object == nullptr && InSelectedItem.IsValid())
+		{
+			const FString ObjectPath = InSelectedItem->GetObjectPath();
+			if (!ObjectPath.IsEmpty())
+			{
+				Object = FSoftObjectPath(ObjectPath).TryLoad();
+			}
+		}
+
+		if (Object && Object->IsAsset())
+		{
+			SelectionManager.SetSelectedAsset(Object);
+		}
+		else
+		{
+			SelectionManager.ClearSelection();
+		}
+#endif // WITH_EDITOR
 	}
 
 	void FSignalFlowDashboardViewFactory::SetNewSelection(const TSharedPtr<FSignalFlowDashboardEntry>& InSelectedItem)
diff --git a/Engine/Plugins/AudioInsights/Source/AudioInsights/Private/Views/TableDashboardViewFactory.cpp b/Engine/Plugins/AudioInsights/Source/AudioInsights/Private/Views/TableDashboardViewFactory.cpp
--- a/Engine/Plugins/AudioInsights/Source/AudioInsights/Private/Views/TableDashboardViewFactory.cpp
+++ b/Engine/Plugins/AudioInsights/Source/AudioInsights/Private/Views/TableDashboardViewFactory.cpp
@@ -14,6 +14,8 @@
 #if WITH_EDITOR
 #include "AudioDeviceManager.h"
 #include "AudioInsightsDashboardAssetCommands.h"
+#include "AudioInsightsDetailsSelectionManager.h"
+#include "AudioInsightsModule.h"
 #include "Editor.h"
 #include "Editor/EditorEngine.h"
 #include "Subsystems/AssetEditorSubsystem.h"
@@ -654,6 +656,44 @@ namespace UE::Audio::Insights
 		return false;
 	}
 #endif // WITH_EDITOR
+
+	void FTraceObjectTableDashboardViewFactory::OnSelectionChanged(TSharedPtr<IDashboardDataViewEntry> SelectedItem, ESelectInfo::Type SelectInfo)
+	{
+#if WITH_EDITOR
+		// Notify the details panel that a new asset has been selected
+		FAudioInsightsDetailsSelectionManager& SelectionManager = IAudioInsightsModule::GetChecked().GetDetailsSelectionManager();
+		TObjectPtr<UObject> SelectedAssetPtr = nullptr;
+
+		const TArray<UObject*> Assets = GetSelectedEditableAssets();
+		if (Assets.Num() > 0)
+		{
+			SelectedAssetPtr = Assets[0];
+		}
+		else if (SelectedItem.IsValid())
+		{
+			// If we could not find a loaded asset for this object, try to load one using the object path
+			IObjectDashboardEntry& RowData = *StaticCastSharedPtr<IObjectDashboardEntry>(SelectedItem).Get();
+			const FString ObjectPath = RowData.GetObjectPath();
+			if (!ObjectPath.IsEmpty())
+			{
+				const TObjectPtr<UObject> LoadedObject = FSoftObjectPath(ObjectPath).TryLoad();
+				if (LoadedObject && LoadedObject->IsAsset())
+				{
+					SelectedAssetPtr = LoadedObject;
+				}
+			}
+		}
+
+		if (SelectedAssetPtr)
+		{
+			SelectionManager.SetSelectedAsset(SelectedAssetPtr);
+		}
+		else if (SelectInfo != ESelectInfo::Type::Direct)
+		{
+			SelectionManager.ClearSelection();
+		}
+#endif // WITH_EDITOR
+	}
 } // namespace UE::Audio::Insights
 
 #undef LOCTEXT_NAMESPACE
\ No newline at end of file
diff --git a/Engine/Plugins/AudioInsights/Source/AudioInsights/Public/Messages/AudioBusTraceMessages.h b/Engine/Plugins/AudioInsights/Source/AudioInsights/Public/Messages/AudioBusTraceMessages.h
--- a/Engine/Plugins/AudioInsights/Source/AudioInsights/Public/Messages/AudioBusTraceMessages.h
+++ b/Engine/Plugins/AudioInsights/Source/AudioInsights/Public/Messages/AudioBusTraceMessages.h
@@ -117,6 +117,7 @@ namespace UE::Audio::Insights
 		FAudioBusDashboardEntry& operator=(const FAudioBusDashboardEntry& Other);
 
 		virtual FText GetDisplayName() const override { return FText::FromString(AudioBusType == EAudioBusType::AssetBased ? FSoftObjectPath(Name).GetAssetName() : Name); }
+		virtual FString GetObjectPath() const override { return Name; }
 		virtual const UObject* GetObject() const override { return FSoftObjectPath(Name).ResolveObject(); }
 		virtual UObject* GetObject() override { return FSoftObjectPath(Name).ResolveObject(); }
 		virtual bool IsValid() const override { return AudioBusId != static_cast<uint32>(INDEX_NONE); }
diff --git a/Engine/Plugins/AudioInsights/Source/AudioInsights/Public/Messages/SignalFlowTraceMessages.h b/Engine/Plugins/AudioInsights/Source/AudioInsights/Public/Messages/SignalFlowTraceMessages.h
--- a/Engine/Plugins/AudioInsights/Source/AudioInsights/Public/Messages/SignalFlowTraceMessages.h
+++ b/Engine/Plugins/AudioInsights/Source/AudioInsights/Public/Messages/SignalFlowTraceMessages.h
@@ -132,6 +132,7 @@ namespace UE::Audio::Insights
 			return DisplayName.IsEmpty() ? FText::FromString(Name) : FText::FromString(DisplayName);
 		}
 
+		virtual FString GetObjectPath() const override { return Name; }
 		virtual const UObject* GetObject() const override { return FSoftObjectPath(Name).ResolveObject(); }
 		virtual UObject* GetObject() override { return FSoftObjectPath(Name).ResolveObject(); }
 		virtual bool IsValid() const override { return SignalFlowNodeID != static_cast<uint32>(INDEX_NONE); }
diff --git a/Engine/Plugins/AudioInsights/Source/AudioInsights/Public/Messages/SubmixTraceMessages.h b/Engine/Plugins/AudioInsights/Source/AudioInsights/Public/Messages/SubmixTraceMessages.h
--- a/Engine/Plugins/AudioInsights/Source/AudioInsights/Public/Messages/SubmixTraceMessages.h
+++ b/Engine/Plugins/AudioInsights/Source/AudioInsights/Public/Messages/SubmixTraceMessages.h
@@ -125,6 +125,7 @@ namespace UE::Audio::Insights
 		FSubmixDashboardEntry& operator=(const FSubmixDashboardEntry& Other);
 
 		virtual FText GetDisplayName() const override { return FText::FromString(FSoftObjectPath(Name).GetAssetName()); }
+		virtual FString GetObjectPath() const override { return Name; }
 		virtual const UObject* GetObject() const override { return FSoftObjectPath(Name).ResolveObject(); }
 		virtual UObject* GetObject() override { return FSoftObjectPath(Name).ResolveObject(); }
 		virtual bool IsValid() const override { return SubmixId != static_cast<uint32>(INDEX_NONE); }
diff --git a/Engine/Plugins/AudioInsights/Source/AudioInsights/Public/Views/TableDashboardViewFactory.h b/Engine/Plugins/AudioInsights/Source/AudioInsights/Public/Views/TableDashboardViewFactory.h
--- a/Engine/Plugins/AudioInsights/Source/AudioInsights/Public/Views/TableDashboardViewFactory.h
+++ b/Engine/Plugins/AudioInsights/Source/AudioInsights/Public/Views/TableDashboardViewFactory.h
@@ -41,6 +41,7 @@ namespace UE::Audio::Insights
 		virtual ~IObjectDashboardEntry() = default;
 
 		virtual FText GetDisplayName() const = 0;
+		virtual FString GetObjectPath() const { return FString(); }
 		virtual const UObject* GetObject() const = 0;
 		virtual UObject* GetObject() = 0;
 	};
@@ -50,6 +51,7 @@ namespace UE::Audio::Insights
 		virtual ~FSoundAssetDashboardEntry() = default;
 
 		UE_API virtual FText GetDisplayName() const override;
+		virtual FString GetObjectPath() const override { return Name; }
 		UE_API virtual const UObject* GetObject() const override;
 		UE_API virtual UObject* GetObject() override;
 		UE_API virtual bool IsValid() const override;
@@ -245,6 +247,7 @@ namespace UE::Audio::Insights
 		UE_API virtual TSharedRef<SWidget> GenerateWidgetForColumn(TSharedRef<IDashboardDataViewEntry> InRowData, const FName& Column) override;
 		UE_API virtual FReply OnDataRowKeyInput(const FGeometry& Geometry, const FKeyEvent& KeyEvent) const override;
 		UE_API virtual TSharedRef<SWidget> MakeWidget(TSharedRef<SDockTab> OwnerTab, const FSpawnTabArgs& SpawnTabArgs) override;
+		UE_API virtual void OnSelectionChanged(TSharedPtr<IDashboardDataViewEntry> SelectedItem, ESelectInfo::Type SelectInfo) override;
 
 	protected:
 #if WITH_EDITOR
diff --git a/Engine/Plugins/Developer/AudioModulationInsights/Source/AudioModulationInsights/Private/Views/AudioModulationContributorsDashboardViewFactory.cpp b/Engine/Plugins/Developer/AudioModulationInsights/Source/AudioModulationInsights/Private/Views/AudioModulationContributorsDashboardViewFactory.cpp
--- a/Engine/Plugins/Developer/AudioModulationInsights/Source/AudioModulationInsights/Private/Views/AudioModulationContributorsDashboardViewFactory.cpp
+++ b/Engine/Plugins/Developer/AudioModulationInsights/Source/AudioModulationInsights/Private/Views/AudioModulationContributorsDashboardViewFactory.cpp
@@ -208,6 +208,8 @@ namespace AudioModulationInsights
 
 	void FAudioModulationContributorsDashboardViewFactory::OnSelectionChanged(TSharedPtr<UE::Audio::Insights::IDashboardDataViewEntry> SelectedItem, ESelectInfo::Type SelectInfo)
 	{
+		UE::Audio::Insights::FTraceObjectTableDashboardViewFactory::OnSelectionChanged(SelectedItem, SelectInfo);
+
 		switch (SelectInfo)
 		{
 			// Note: Only call ResetSelectedModulatorId() for user interactions not via non-user calls with ESelectInfo::Type::Direct
diff --git a/Engine/Plugins/Developer/AudioModulationInsights/Source/AudioModulationInsights/Private/Views/AudioModulationDashboardViewFactory.cpp b/Engine/Plugins/Developer/AudioModulationInsights/Source/AudioModulationInsights/Private/Views/AudioModulationDashboardViewFactory.cpp
--- a/Engine/Plugins/Developer/AudioModulationInsights/Source/AudioModulationInsights/Private/Views/AudioModulationDashboardViewFactory.cpp
+++ b/Engine/Plugins/Developer/AudioModulationInsights/Source/AudioModulationInsights/Private/Views/AudioModulationDashboardViewFactory.cpp
@@ -545,6 +545,8 @@ namespace AudioModulationInsights
 
 	void FAudioModulationDashboardViewFactory::OnSelectionChanged(TSharedPtr<UE::Audio::Insights::IDashboardDataViewEntry> SelectedItem, ESelectInfo::Type SelectInfo)
 	{
+		UE::Audio::Insights::FTraceObjectTableDashboardViewFactory::OnSelectionChanged(SelectedItem, SelectInfo);
+
 		switch (SelectInfo)
 		{
 			// Note: Only call ResetSelectedModulatorId() for user interactions not via non-user calls with ESelectInfo::Type::Direct
