[AudioInsights] - Toolbar: Support Persistent ShowCacheStats Settings

#Synced-CL 52537012

diff --git a/Engine/Plugins/AudioInsights/Source/AudioInsightsEditor/Private/AudioInsightsEditorDashboardFactory.cpp b/Engine/Plugins/AudioInsights/Source/AudioInsightsEditor/Private/AudioInsightsEditorDashboardFactory.cpp
--- a/Engine/Plugins/AudioInsights/Source/AudioInsightsEditor/Private/AudioInsightsEditorDashboardFactory.cpp
+++ b/Engine/Plugins/AudioInsights/Source/AudioInsightsEditor/Private/AudioInsightsEditorDashboardFactory.cpp
@@ -608,6 +608,12 @@ namespace UE::Audio::Insights
 
 		constexpr FLinearColor CacheStatsColor(1.0f, 1.0f, 1.0f, 0.5f);
 
+		const auto CacheStatsVisibilityLambda = []()
+		{
+			const TObjectPtr<const UAudioInsightsEditorSettings> Settings = GetDefault<UAudioInsightsEditorSettings>();
+			return (Settings && Settings->bShowCacheStats) ? EVisibility::Visible : EVisibility::Collapsed;
+		};
+
 		return SNew(SBorder)
 			.BorderImage(&BackgroundColorBrush)
 			[
@@ -866,10 +872,7 @@ namespace UE::Audio::Insights
 					.Justification(ETextJustify::Center)
 					.ColorAndOpacity(CacheStatsColor)
 					.Text(LOCTEXT("AudioDashboardEditor_CacheStats_Title", "Cache Size: "))
-					.Visibility_Lambda([this]()
-					{
-						return bIsCacheStatsVisible ? EVisibility::Visible : EVisibility::Collapsed;
-					})
+					.Visibility_Lambda(CacheStatsVisibilityLambda)
 				]
 
 				+ SHorizontalBox::Slot()
@@ -888,10 +891,7 @@ namespace UE::Audio::Insights
 											, FText::AsNumber(CacheSizeMb, FSlateStyle::Get().GetMemoryFormat())
 											, FSlateStyle::Get().FormatMemoryAsMegabytes(MaxCacheSizeMb));
 					})
-					.Visibility_Lambda([this]()
-					{
-						return bIsCacheStatsVisible ? EVisibility::Visible : EVisibility::Collapsed;
-					})
+					.Visibility_Lambda(CacheStatsVisibilityLambda)
 				]
 
 				+ SHorizontalBox::Slot()
@@ -902,10 +902,7 @@ namespace UE::Audio::Insights
 					.Justification(ETextJustify::Center)
 					.ColorAndOpacity(CacheStatsColor)
 					.Text(LOCTEXT("AudioDashboardEditor_CacheStats_Duration", " : Duration: "))
-					.Visibility_Lambda([this]()
-					{
-						return bIsCacheStatsVisible ? EVisibility::Visible : EVisibility::Collapsed;
-					})
+					.Visibility_Lambda(CacheStatsVisibilityLambda)
 				]
 
 				+ SHorizontalBox::Slot()
@@ -922,10 +919,7 @@ namespace UE::Audio::Insights
 						return FText::Format(LOCTEXT("AudioDashboardEditor_CacheStats_DurationValue", "{0}s")
 											, FText::AsNumber(CacheDuration, FSlateStyle::Get().GetShortTimeFormat()));
 					})
-					.Visibility_Lambda([this]()
-					{
-						return bIsCacheStatsVisible ? EVisibility::Visible : EVisibility::Collapsed;
-					})
+					.Visibility_Lambda(CacheStatsVisibilityLambda)
 				]
 
 				// World filter label
@@ -1197,14 +1191,19 @@ namespace UE::Audio::Insights
 				EditorDashboardFactoryPrivate::CacheSettingsShowStatsTooltipText,
 				FSlateIcon(),
 				FUIAction(
-					FExecuteAction::CreateLambda([this]()
+					FExecuteAction::CreateLambda([]()
 					{
-						bIsCacheStatsVisible = !bIsCacheStatsVisible;
+						if (TObjectPtr<UAudioInsightsEditorSettings> Settings = GetMutableDefault<UAudioInsightsEditorSettings>())
+						{
+							Settings->bShowCacheStats = !Settings->bShowCacheStats;
+							Settings->SaveConfig();
+						}
 					}),
 					FCanExecuteAction(),
-					FIsActionChecked::CreateLambda([this]()
+					FIsActionChecked::CreateLambda([]()
 					{
-						return bIsCacheStatsVisible;
+						const TObjectPtr<const UAudioInsightsEditorSettings> Settings = GetDefault<UAudioInsightsEditorSettings>();
+						return Settings && Settings->bShowCacheStats;
 					})
 				),
 				NAME_None,
diff --git a/Engine/Plugins/AudioInsights/Source/AudioInsightsEditor/Private/AudioInsightsEditorDashboardFactory.h b/Engine/Plugins/AudioInsights/Source/AudioInsightsEditor/Private/AudioInsightsEditorDashboardFactory.h
--- a/Engine/Plugins/AudioInsights/Source/AudioInsightsEditor/Private/AudioInsightsEditorDashboardFactory.h
+++ b/Engine/Plugins/AudioInsights/Source/AudioInsightsEditor/Private/AudioInsightsEditorDashboardFactory.h
@@ -90,7 +90,6 @@ namespace UE::Audio::Insights
 		bool bIsLayoutReseting = false;
 
 		bool bIsPaused { false };
-		bool bIsCacheStatsVisible { true };
 
 		TSharedPtr<FTabManager> DashboardTabManager;
 		TSharedPtr<FWorkspaceItem> DashboardWorkspace;
diff --git a/Engine/Plugins/AudioInsights/Source/AudioInsightsEditor/Public/AudioInsightsEditorSettings.h b/Engine/Plugins/AudioInsights/Source/AudioInsightsEditor/Public/AudioInsightsEditorSettings.h
--- a/Engine/Plugins/AudioInsights/Source/AudioInsightsEditor/Public/AudioInsightsEditorSettings.h
+++ b/Engine/Plugins/AudioInsights/Source/AudioInsightsEditor/Public/AudioInsightsEditorSettings.h
@@ -39,6 +39,10 @@ public:
 	UPROPERTY(Config, EditAnywhere, Category = "Trace Channels")
 	bool bOnlyTraceAudioChannels = true;
 
+	/** Whether to show cache stats in the Audio Insights toolbar. */
+	UPROPERTY(Config, EditAnywhere, Category = "Cache")
+	bool bShowCacheStats = true;
+
 	/** Settings for analyzer rack spectrogram widget */
 	UPROPERTY(EditAnywhere, config, Category = Spectrogram, meta = (ShowOnlyInnerProperties))
 	FSpectrogramRackUnitSettings SpectrogramSettings;
