[AudioInsights] - Cache: Unify Cache Container (TCircularBuffer to TArray)

#Synced-CL 52537012

diff --git a/Engine/Plugins/AudioInsights/Source/AudioInsights/Private/Cache/AudioInsightsCacheManager.cpp b/Engine/Plugins/AudioInsights/Source/AudioInsights/Private/Cache/AudioInsightsCacheManager.cpp
--- a/Engine/Plugins/AudioInsights/Source/AudioInsights/Private/Cache/AudioInsightsCacheManager.cpp
+++ b/Engine/Plugins/AudioInsights/Source/AudioInsights/Private/Cache/AudioInsightsCacheManager.cpp
@@ -172,10 +172,6 @@ namespace UE::Audio::Insights
 		{
 			if (Chunk->IsChunkFull())
 			{
-#if WITH_EDITOR
-				// Move to next chunk
-				WriteIndex = Cache.GetNextIndex(WriteIndex);
-#else
 				const uint32 NextWriteIndex = Cache.GetNextIndex(WriteIndex);
 
 				// If our cache is not limited, we may need to reserve more chunks
@@ -185,7 +181,6 @@ namespace UE::Audio::Insights
 				}
 
 				WriteIndex = NextWriteIndex;
-#endif // WITH_EDITOR
 
 				Chunk = &Cache[WriteIndex];
 
@@ -208,11 +203,10 @@ namespace UE::Audio::Insights
 	}
 
 	FAudioInsightsCache::FAudioInsightsCache()
-#if WITH_EDITOR
-		: EditorCache(CacheManagerPrivate::NumChunks)
-#endif // WITH_EDITOR
 	{
-#if !WITH_EDITOR
+#if WITH_EDITOR
+		bLimitCacheSize = true;
+#else
 		IUnrealInsightsModule& UnrealInsightsModule = FModuleManager::LoadModuleChecked<IUnrealInsightsModule>("TraceInsights");
 		if (UnrealInsightsModule.GetAnalysisSession().IsValid())
 		{
@@ -227,7 +221,7 @@ namespace UE::Audio::Insights
 				InsightsManager->GetSessionChangedEvent().AddRaw(this, &FAudioInsightsCache::OnSessionChangedEvent);
 			}
 		}
-#endif // !WITH_EDITOR
+#endif // WITH_EDITOR
 	}
 
 	FAudioInsightsCache::~FAudioInsightsCache()
@@ -246,70 +240,44 @@ namespace UE::Audio::Insights
 
 	void FAudioInsightsCache::ResetCache(TraceServices::FSlabAllocator& Allocator)
 	{
-#if !WITH_EDITOR
-		StandaloneCache.Empty();
-#endif // !WITH_EDITOR
+		Cache.Empty();
 
 		ReserveNewChunks(Allocator, CacheManagerPrivate::NumChunks);
 	}
 
 	void FAudioInsightsCache::ReserveNewChunks(TraceServices::FSlabAllocator& Allocator, const uint32 NumToReserve)
 	{
-#if !WITH_EDITOR
-		StandaloneCache.Reserve(StandaloneCache.Num() + NumToReserve);
-#endif // !WITH_EDITOR
+		Cache.Reserve(Cache.Num() + NumToReserve);
 
 		for (uint32 Index = 0u; Index < NumToReserve; ++Index)
 		{
-#if WITH_EDITOR
-			EditorCache[Index] = MakeUnique<FAudioCachedMessageChunk>(Allocator, CacheManagerPrivate::MaxChunkSize);
-#else
-			StandaloneCache.Add(MakeUnique<FAudioCachedMessageChunk>(Allocator, CacheManagerPrivate::MaxChunkSize));
-#endif // WITH_EDITOR
+			Cache.Add(MakeUnique<FAudioCachedMessageChunk>(Allocator, CacheManagerPrivate::MaxChunkSize));
 		}
 	}
 
 	uint32 FAudioInsightsCache::Capacity() const
 	{
-#if WITH_EDITOR
-		return EditorCache.Capacity();
-#else
-		return static_cast<uint32>(StandaloneCache.Num());
-#endif // WITH_EDITOR
+		return static_cast<uint32>(Cache.Num());
 	}
 
 	uint32 FAudioInsightsCache::GetNextIndex(uint32 ReadIndex) const
 	{
-#if WITH_EDITOR
-		return EditorCache.GetNextIndex(ReadIndex);
-#else
-		if (bLimitStandaloneCacheSize)
+		ReadIndex++;
+		if (bLimitCacheSize && ReadIndex >= Capacity())
 		{
-			ReadIndex++;
-			if (ReadIndex >= Capacity())
-			{
-				ReadIndex = 0u;
-			}
-		}
-		else
-		{
-			ReadIndex++;
+			ReadIndex = 0u;
 		}
 
 		return ReadIndex;
-#endif // WITH_EDITOR
 	}
 
 	uint32 FAudioInsightsCache::GetPreviousIndex(uint32 ReadIndex) const
 	{
-#if WITH_EDITOR
-		return EditorCache.GetPreviousIndex(ReadIndex);
-#else
 		if (ReadIndex > 0u)
 		{
 			ReadIndex--;
 		}
-		else if (bLimitStandaloneCacheSize)
+		else if (bLimitCacheSize)
 		{
 			ReadIndex = Capacity() - 1u;
 		}
@@ -319,7 +287,6 @@ namespace UE::Audio::Insights
 		}
 
 		return ReadIndex;
-#endif // WITH_EDITOR
 	}
 
 #if !WITH_EDITOR
@@ -329,7 +296,7 @@ namespace UE::Audio::Insights
 		const TSharedPtr<FAudioInsightsComponent> AudioInsightsComponent = AudioInsightsModule.GetAudioInsightsComponent();
 		if (AudioInsightsComponent.IsValid() && AudioInsightsComponent->GetIsLiveSession())
 		{
-			bLimitStandaloneCacheSize = true;
+			bLimitCacheSize = true;
 		}
 	}
 
diff --git a/Engine/Plugins/AudioInsights/Source/AudioInsights/Public/Cache/AudioInsightsCacheManager.h b/Engine/Plugins/AudioInsights/Source/AudioInsights/Public/Cache/AudioInsightsCacheManager.h
--- a/Engine/Plugins/AudioInsights/Source/AudioInsights/Public/Cache/AudioInsightsCacheManager.h
+++ b/Engine/Plugins/AudioInsights/Source/AudioInsights/Public/Cache/AudioInsightsCacheManager.h
@@ -3,7 +3,6 @@
 
 #include "Cache/AudioCachedMessageChunk.h"
 #include "Cache/IAudioCachedMessage.h"
-#include "Containers/CircularBuffer.h"
 #include "Containers/Ticker.h"
 #include "Messages/AnalyzerMessageQueue.h"
 #include "Misc/Optional.h"
@@ -13,11 +12,11 @@
 
 namespace UE::Audio::Insights
 {
-	// Helper class that provides a similar API in Editor and Standalone builds to the cache container
-	// In Editor, the cache is a limited circular buffer
-	// In standalone, the cache is an array that may or may not be limited in size
-	// - When connected to a live session (either via direct or live trace), the array will be behave like a circular buffer
-	// - When loading in a .utrace file, the array will grow to hold all of the trace file data in memory
+	// Helper class that manages the cache container
+	// The cache may or may not be limited in size based on the bLimitCacheSize flag
+	// - In Editor, the cache size is always limited, behaving like a circular buffer
+	// - In Standalone, when connected to a live session (either via direct or live trace), the cache size is limited, behaving like a circular buffer
+	// - In Standalone, when loading in a .utrace file, the cache will grow to hold all of the trace file data in memory
 	class FAudioInsightsCache
 	{
 	public:
@@ -26,20 +25,12 @@ namespace UE::Audio::Insights
 
 		FAudioCachedMessageChunk& operator[](uint32 Index)
 		{
-#if WITH_EDITOR
-			return *EditorCache[Index].Get();
-#else
-			return *StandaloneCache[Index].Get();
-#endif // WITH_EDITOR
+			return *Cache[Index].Get();
 		}
 
 		const FAudioCachedMessageChunk& operator[](uint32 Index) const
 		{
-#if WITH_EDITOR
-			return *EditorCache[Index].Get();
-#else
-			return *StandaloneCache[Index].Get();
-#endif // WITH_EDITOR
+			return *Cache[Index].Get();
 		}
 
 		void ResetCache(TraceServices::FSlabAllocator& Allocator);
@@ -55,12 +46,8 @@ namespace UE::Audio::Insights
 		void OnSessionChangedEvent();
 #endif // !WITH_EDITOR
 
-#if WITH_EDITOR
-		TCircularBuffer<TUniquePtr<FAudioCachedMessageChunk>> EditorCache;
-#else
-		TArray<TUniquePtr<FAudioCachedMessageChunk>> StandaloneCache;
-		bool bLimitStandaloneCacheSize = false;
-#endif // WITH_EDITOR
+		TArray<TUniquePtr<FAudioCachedMessageChunk>> Cache;
+		bool bLimitCacheSize = false;
 	};
 
 	class FAudioInsightsCacheManager
@@ -143,12 +130,10 @@ namespace UE::Audio::Insights
 				else
 				{
 					Index = Cache.GetPreviousIndex(Index);
-#if !WITH_EDITOR
 					if (Index == INDEX_NONE)
 					{
 						break;
 					}
-#endif // !WITH_EDITOR
 				}
 			}
 
diff --git a/Engine/Plugins/AudioInsights/Source/AudioInsights/Public/Messages/SignalFlowTraceMessages.h b/Engine/Plugins/AudioInsights/Source/AudioInsights/Public/Messages/SignalFlowTraceMessages.h
--- a/Engine/Plugins/AudioInsights/Source/AudioInsights/Public/Messages/SignalFlowTraceMessages.h
+++ b/Engine/Plugins/AudioInsights/Source/AudioInsights/Public/Messages/SignalFlowTraceMessages.h
@@ -2,6 +2,7 @@
 #pragma once
 
 #include "Cache/IAudioCachedMessage.h"
+#include "Containers/CircularBuffer.h"
 #include "Messages/AnalyzerMessageQueue.h"
 #include "Messages/AudioBusTraceMessages.h"
 #include "Messages/SignalFlowEntryKey.h"
