[Audio Insights] AudioInsightsSettings: fix for settings to be defined editor wide and not per project.

#Synced-CL 52505914

diff --git a/Engine/Plugins/AudioInsights/Source/AudioInsights/Private/AudioInsightsModule.cpp b/Engine/Plugins/AudioInsights/Source/AudioInsights/Private/AudioInsightsModule.cpp
--- a/Engine/Plugins/AudioInsights/Source/AudioInsights/Private/AudioInsightsModule.cpp
+++ b/Engine/Plugins/AudioInsights/Source/AudioInsights/Private/AudioInsightsModule.cpp
@@ -15,6 +15,7 @@
 
 #if WITH_EDITOR
 #include "AssetEditorCommands.h"
+#include "AudioInsightsSettings.h"
 #endif // WITH_EDITOR
 
 #if !WITH_EDITOR
@@ -41,6 +42,11 @@ namespace UE::Audio::Insights
 		// Don't run providers in any commandlet to avoid additional, unnecessary overhead as audio insights is dormant.
 		if (!IsRunningCommandlet())
 		{
+#if WITH_EDITOR
+			// Re-load settings from engine-wide config so they persist across projects
+			GetMutableDefault<UAudioInsightsSettings>()->LoadConfig(nullptr, *UAudioInsightsSettings::GetAudioInsightsConfigFilename());
+#endif
+
 			TraceModule = MakeUnique<FTraceModule>();
 			RewindDebuggerExtension = MakeUnique<FRewindDebuggerAudioInsightsRuntime>();
 
diff --git a/Engine/Plugins/AudioInsights/Source/AudioInsights/Private/AudioInsightsSettings.cpp b/Engine/Plugins/AudioInsights/Source/AudioInsights/Private/AudioInsightsSettings.cpp
new file mode 100644
--- /dev/null
+++ b/Engine/Plugins/AudioInsights/Source/AudioInsights/Private/AudioInsightsSettings.cpp
@@ -0,0 +1,17 @@
+// Copyright Epic Games, Inc. All Rights Reserved.
+#include "AudioInsightsSettings.h"
+
+#if WITH_EDITOR
+#include "Misc/ConfigCacheIni.h"
+#include "Misc/Paths.h"
+#endif // WITH_EDITOR
+
+FString UAudioInsightsSettings::GetAudioInsightsConfigFilename()
+{
+#if WITH_EDITOR
+	const FName ConfigName = UAudioInsightsSettings::StaticClass()->ClassConfigName;
+	return FConfigCacheIni::GetDestIniFilename(*ConfigName.ToString(), nullptr, *FPaths::EngineEditorSettingsDir());
+#else
+	return GetConfigFilename(GetMutableDefault<UAudioInsightsSettings>());
+#endif
+}
diff --git a/Engine/Plugins/AudioInsights/Source/AudioInsights/Private/Views/OutputMeterDashboardViewFactory.cpp b/Engine/Plugins/AudioInsights/Source/AudioInsights/Private/Views/OutputMeterDashboardViewFactory.cpp
--- a/Engine/Plugins/AudioInsights/Source/AudioInsights/Private/Views/OutputMeterDashboardViewFactory.cpp
+++ b/Engine/Plugins/AudioInsights/Source/AudioInsights/Private/Views/OutputMeterDashboardViewFactory.cpp
@@ -85,7 +85,7 @@ namespace UE::Audio::Insights
 		const FProperty* Property = UAudioInsightsSettings::StaticClass()->FindPropertyByName(GET_MEMBER_NAME_CHECKED(UAudioInsightsSettings, OutputMeterDashboardSettings));
 
 		UAudioInsightsSettings* AudioInsightsSettings = GetMutableDefault<UAudioInsightsSettings>();
-		AudioInsightsSettings->UpdateSinglePropertyInConfigFile(Property, GetConfigFilename(AudioInsightsSettings));
+		AudioInsightsSettings->UpdateSinglePropertyInConfigFile(Property, UAudioInsightsSettings::GetAudioInsightsConfigFilename());
 	}
 
 	void FOutputMeterDashboardViewFactory::InitLoudnessMeterWidgetView()
@@ -97,7 +97,7 @@ namespace UE::Audio::Insights
 			const FProperty* Property = UAudioInsightsSettings::StaticClass()->FindPropertyByName(GET_MEMBER_NAME_CHECKED(UAudioInsightsSettings, OutputMeterSettings));
 
 			UAudioInsightsSettings* Settings = GetMutableDefault<UAudioInsightsSettings>();
-			Settings->UpdateSinglePropertyInConfigFile(Property, GetConfigFilename(Settings));
+			Settings->UpdateSinglePropertyInConfigFile(Property, UAudioInsightsSettings::GetAudioInsightsConfigFilename());
 		};
 
 		// TimerPanelParams
diff --git a/Engine/Plugins/AudioInsights/Source/AudioInsights/Public/AudioInsightsSettings.h b/Engine/Plugins/AudioInsights/Source/AudioInsights/Public/AudioInsightsSettings.h
--- a/Engine/Plugins/AudioInsights/Source/AudioInsights/Public/AudioInsightsSettings.h
+++ b/Engine/Plugins/AudioInsights/Source/AudioInsights/Public/AudioInsightsSettings.h
@@ -22,6 +22,8 @@ public:
 	/** Settings for the Output Meter dashboard tab (not visible to the user) */
 	UPROPERTY(config)
 	FOutputMeterDashboardSettings OutputMeterDashboardSettings;
+
+	static FString GetAudioInsightsConfigFilename();
 };
 
 #undef UE_API
