[AudioInsights] - Toolbar: Support Cache Controls and Settings in Standalone

#Synced-CL 52611640

diff --git a/Engine/Plugins/AudioInsights/Source/AudioInsights/Private/AudioInsightsDashboardFactory.cpp b/Engine/Plugins/AudioInsights/Source/AudioInsights/Private/AudioInsightsDashboardFactory.cpp
--- a/Engine/Plugins/AudioInsights/Source/AudioInsights/Private/AudioInsightsDashboardFactory.cpp
+++ b/Engine/Plugins/AudioInsights/Source/AudioInsights/Private/AudioInsightsDashboardFactory.cpp
@@ -11,6 +11,7 @@
 #include "Widgets/Input/SButton.h"
 #include "Widgets/Input/SCheckBox.h"
 #include "Widgets/Layout/SBox.h"
+#include "Widgets/Layout/SSeparator.h"
 
 #define LOCTEXT_NAMESPACE "AudioInsights"
 
@@ -45,6 +46,8 @@ namespace UE::Audio::Insights
 
 		DashboardTabManager = FGlobalTabmanager::Get()->NewTabManager(DockTab);
 
+		ToolbarWidgets = MakeShared<FToolbarWidgets>();
+
 		TabLayout = GetDefaultTabLayout();
 
 		RegisterTabSpawners();
@@ -101,6 +104,7 @@ namespace UE::Audio::Insights
 		DockTab->SetOnTabClosed(SDockTab::FOnTabClosedCallback::CreateLambda([this](TSharedRef<SDockTab> TabClosed)
 		{
 			UnregisterTabSpawners();
+			ToolbarWidgets.Reset();
 		}));
 
 		return DockTab;
@@ -208,6 +212,57 @@ namespace UE::Audio::Insights
 			.Padding(12.0f, 0.0f, 0.0f, 0.0f)
 			[
 				MakeSnapshotButtonWidget()
+			]
+			// Cache toolbar widgets (only visible during live monitoring)
+			+ SHorizontalBox::Slot()
+			.AutoWidth()
+			[
+				SNew(SBox)
+				.Visibility_Lambda([]()
+				{
+					return IAudioInsightsModule::IsLiveMonitoringSession() ? EVisibility::Visible : EVisibility::Collapsed;
+				})
+				[
+					SNew(SHorizontalBox)
+					// Separator
+					+ SHorizontalBox::Slot()
+					.HAlign(HAlign_Left)
+					.VAlign(VAlign_Center)
+					.AutoWidth()
+					.Padding(0.0f, 0.0f, 0.0f, 0.0f)
+					[
+						SNew(SBox)
+						.HeightOverride(38.0f)
+						[
+							SNew(SSeparator)
+							.Orientation(Orient_Vertical)
+							.Thickness(2.0f)
+						]
+					]
+					// Cache control: Pause button
+					+ SHorizontalBox::Slot()
+					.HAlign(HAlign_Left)
+					.VAlign(VAlign_Center)
+					.AutoWidth()
+					.Padding(8.0f, 0.0f, 0.0f, 0.0f)
+					[
+						ToolbarWidgets->MakePauseButtonWidget()
+					]
+					// Cache stats
+					+ SHorizontalBox::Slot()
+					.VAlign(VAlign_Center)
+					.AutoWidth()
+					[
+						ToolbarWidgets->MakeCacheStatsWidget()
+					]
+				]
+			]
+			// Settings
+			+ SHorizontalBox::Slot()
+			.AutoWidth()
+			.HAlign(HAlign_Right)
+			[
+				ToolbarWidgets->MakeSettingsButtonWidget()
 			];
 	}
 
diff --git a/Engine/Plugins/AudioInsights/Source/AudioInsights/Private/AudioInsightsSettings.cpp b/Engine/Plugins/AudioInsights/Source/AudioInsights/Private/AudioInsightsSettings.cpp
--- a/Engine/Plugins/AudioInsights/Source/AudioInsights/Private/AudioInsightsSettings.cpp
+++ b/Engine/Plugins/AudioInsights/Source/AudioInsights/Private/AudioInsightsSettings.cpp
@@ -15,3 +15,37 @@ FString UAudioInsightsSettings::GetAudioInsightsConfigFilename()
 	return GetConfigFilename(GetMutableDefault<UAudioInsightsSettings>());
 #endif
 }
+
+#if WITH_EDITOR
+void UAudioInsightsSettings::PostEditChangeChainProperty(FPropertyChangedChainEvent& PropertyChangedEvent)
+{
+	Super::PostEditChangeChainProperty(PropertyChangedEvent);
+
+	const TDoubleLinkedList<FProperty*>::TDoubleLinkedListNode* const ActiveMemberNode = PropertyChangedEvent.PropertyChain.GetActiveMemberNode();
+	if (ActiveMemberNode == nullptr || ActiveMemberNode->GetValue() == nullptr)
+	{
+		return;
+	}
+
+	SavePropertyToConfigFile(ActiveMemberNode->GetValue());
+}
+#endif
+
+void UAudioInsightsSettings::SaveCacheSettings()
+{
+	const FProperty* const CacheSettingProperty = UAudioInsightsSettings::StaticClass()->FindPropertyByName(GET_MEMBER_NAME_CHECKED(UAudioInsightsSettings, CacheSettings));
+	ensure(CacheSettingProperty);
+
+	SavePropertyToConfigFile(CacheSettingProperty);
+}
+
+void UAudioInsightsSettings::SavePropertyToConfigFile(const FProperty* const PropertyToSave)
+{
+	if (PropertyToSave == nullptr)
+	{
+		return;
+	}
+	
+	// Save setting to config file
+	UpdateSinglePropertyInConfigFile(PropertyToSave, GetAudioInsightsConfigFilename());
+}
diff --git a/Engine/Plugins/AudioInsights/Source/AudioInsights/Private/AudioInsightsToolbarWidgets.cpp b/Engine/Plugins/AudioInsights/Source/AudioInsights/Private/AudioInsightsToolbarWidgets.cpp
new file mode 100644
--- /dev/null
+++ b/Engine/Plugins/AudioInsights/Source/AudioInsights/Private/AudioInsightsToolbarWidgets.cpp
@@ -0,0 +1,288 @@
+// Copyright Epic Games, Inc. All Rights Reserved.
+#include "AudioInsightsToolbarWidgets.h"
+
+#include "AudioInsightsSettings.h"
+#include "AudioInsightsStyle.h"
+#include "AudioInsightsTimingViewExtender.h"
+#include "Cache/AudioInsightsCacheManager.h"
+#include "Framework/MultiBox/MultiBoxBuilder.h"
+#include "IAudioInsightsModule.h"
+#include "Internationalization/Text.h"
+#include "Settings/CacheSettings.h"
+#include "Styling/StyleColors.h"
+#include "Widgets/Images/SImage.h"
+#include "Widgets/Input/SButton.h"
+#include "Widgets/Input/SComboButton.h"
+#include "Widgets/Layout/SBox.h"
+#include "Widgets/Text/STextBlock.h"
+
+#define LOCTEXT_NAMESPACE "AudioInsights"
+
+namespace UE::Audio::Insights
+{
+	namespace ToolbarWidgetsPrivate
+	{
+		static const FText CachePauseTooltipText = LOCTEXT("AudioInsightsToolbar_CachePause_TooltipText", "Click to pause monitoring.\n\nDepending on your cache settings, this may also stop the cache collecting data from the current session.");
+		static const FText CacheResumeTooltipText = LOCTEXT("AudioInsightsToolbar_CacheResume_TooltipText", "Click to resume monitoring.");
+
+		static const FText CacheStatsTitleText = LOCTEXT("AudioInsightsToolbar_CacheStats_Title", "Cache Size: ");
+		static const FText CacheStatsMemoryText = LOCTEXT("AudioInsightsToolbar_CacheStats_Memory", "{0} / {1}");
+		static const FText CacheStatsDurationText = LOCTEXT("AudioInsightsToolbar_CacheStats_Duration", " : Duration: ");
+		static const FText CacheStatsDurationValueText = LOCTEXT("AudioInsightsToolbar_CacheStats_DurationValue", "{0}s");
+
+		static const FText CacheSettingsHeaderText = LOCTEXT("AudioInsightsToolbar_CacheSettings_HeaderText", "Cache Settings");
+		static const FText CacheSettingsShowStatsText = LOCTEXT("AudioInsightsToolbar_CacheSettings_ShowCacheStats", "Show Cache Stats");
+		static const FText CacheSettingsShowStatsTooltipText = LOCTEXT("AudioInsightsToolbar_CacheSettings_ShowCacheStatsTooltip", "Shows/Hides Cache Stats");
+
+		static const FText SettingsTooltipText = LOCTEXT("AudioInsightsToolbar_Settings_TooltipText", "Settings");
+
+	} // namespace ToolbarWidgetsPrivate
+
+	FToolbarWidgets::FToolbarWidgets()
+	{
+		InitializeDelegates();
+	}
+
+	FToolbarWidgets::~FToolbarWidgets()
+	{
+		DeinitializeDelegates();
+	}
+
+	void FToolbarWidgets::InitializeDelegates()
+	{
+		FAudioInsightsTimingViewExtender& TimingViewExtender = IAudioInsightsModule::GetChecked().GetTimingViewExtender();
+
+		OnTimingViewTimeMarkerChangedHandle = TimingViewExtender.OnTimingViewTimeMarkerChanged.AddRaw(this, &FToolbarWidgets::OnTimingViewTimeMarkerChanged);
+		OnTimeControlMethodResetHandle = TimingViewExtender.OnTimeControlMethodReset.AddRaw(this, &FToolbarWidgets::OnTimeControlMethodReset);
+	}
+
+	void FToolbarWidgets::DeinitializeDelegates()
+	{
+		FAudioInsightsTimingViewExtender& TimingViewExtender = IAudioInsightsModule::GetChecked().GetTimingViewExtender();
+
+		if (OnTimingViewTimeMarkerChangedHandle.IsValid())
+		{
+			TimingViewExtender.OnTimingViewTimeMarkerChanged.Remove(OnTimingViewTimeMarkerChangedHandle);
+			OnTimingViewTimeMarkerChangedHandle.Reset();
+		}
+
+		if (OnTimeControlMethodResetHandle.IsValid())
+		{
+			TimingViewExtender.OnTimeControlMethodReset.Remove(OnTimeControlMethodResetHandle);
+			OnTimeControlMethodResetHandle.Reset();
+		}
+	}
+
+	void FToolbarWidgets::OnTimingViewTimeMarkerChanged(double InTimeMarker)
+	{
+		bIsPaused = true;
+	}
+
+	void FToolbarWidgets::OnTimeControlMethodReset()
+	{
+		bIsPaused = false;
+	}
+
+	TSharedRef<SWidget> FToolbarWidgets::MakePauseButtonWidget()
+	{
+		using namespace ToolbarWidgetsPrivate;
+
+		return SNew(SButton)
+			.ButtonStyle(&FAppStyle::Get().GetWidgetStyle<FButtonStyle>("SimpleButton"))
+			.ToolTipText_Lambda([this]()
+			{
+				return bIsPaused ? CacheResumeTooltipText : CachePauseTooltipText;
+			})
+			.OnClicked_Lambda([this]()
+			{
+				FAudioInsightsTimingViewExtender& TimingViewExtender = IAudioInsightsModule::GetChecked().GetTimingViewExtender();
+
+				if (bIsPaused)
+				{
+					TimingViewExtender.ResumeTimeMarker();
+				}
+				else
+				{
+					const double PauseTime = IAudioInsightsModule::GetChecked().GetCacheManager().GetCacheEndTimeStamp();
+					TimingViewExtender.PauseTimeMarker(PauseTime, ESystemControllingTimeMarker::External);
+				}
+
+				return FReply::Handled();
+			})
+			[
+				SNew(SBox)
+				.HeightOverride(16.0f)
+				[
+					SNew(SHorizontalBox)
+					+ SHorizontalBox::Slot()
+					.VAlign(VAlign_Center)
+					.AutoWidth()
+					[
+						SNew(SImage)
+						.Image(FSlateStyle::Get().GetBrush("AudioInsights.Icon.Pause"))
+						.ColorAndOpacity_Lambda([this]()
+						{
+							return bIsPaused ? FStyleColors::AccentBlue : FSlateColor::UseForeground();
+						})
+					]
+				]
+			];
+	}
+
+	TSharedRef<SWidget> FToolbarWidgets::MakeCacheStatsWidget()
+	{
+		using namespace ToolbarWidgetsPrivate;
+
+		constexpr FLinearColor CacheStatsColor(1.0f, 1.0f, 1.0f, 0.5f);
+
+		const auto CacheStatsVisibilityLambda = []()
+		{
+			if (const TObjectPtr<const UAudioInsightsSettings> Settings = GetDefault<UAudioInsightsSettings>())
+			{
+				const FCacheSettings& CacheSettings = Settings->CacheSettings;
+				return CacheSettings.bShowCacheStats ? EVisibility::Visible : EVisibility::Collapsed;
+			}
+
+			return EVisibility::Collapsed;
+		};
+
+		return SNew(SHorizontalBox)
+			+ SHorizontalBox::Slot()
+			.Padding(12.0f, 0.0f, 0.0f, 0.0f)
+			.VAlign(VAlign_Center)
+			.AutoWidth()
+			[
+				SNew(STextBlock)
+				.Justification(ETextJustify::Center)
+				.ColorAndOpacity(CacheStatsColor)
+				.Text(CacheStatsTitleText)
+				.Visibility_Lambda(CacheStatsVisibilityLambda)
+			]
+
+			+ SHorizontalBox::Slot()
+			.VAlign(VAlign_Center)
+			.AutoWidth()
+			[
+				SNew(STextBlock)
+				.Justification(ETextJustify::Center)
+				.Text_Lambda([]()
+				{
+					FAudioInsightsCacheManager& CacheManager = IAudioInsightsModule::GetChecked().GetCacheManager();
+					const float CacheSizeMb = static_cast<float>(CacheManager.GetUsedCacheSize()) / 1000000.0f;
+					const float MaxCacheSizeMb = static_cast<float>(CacheManager.GetMaxCacheSize()) / 1000000.0f;
+
+					return FText::Format(CacheStatsMemoryText
+										, FText::AsNumber(CacheSizeMb, FSlateStyle::Get().GetMemoryFormat())
+										, FSlateStyle::Get().FormatMemoryAsMegabytes(MaxCacheSizeMb));
+				})
+				.Visibility_Lambda(CacheStatsVisibilityLambda)
+			]
+
+			+ SHorizontalBox::Slot()
+			.VAlign(VAlign_Center)
+			.AutoWidth()
+			[
+				SNew(STextBlock)
+				.Justification(ETextJustify::Center)
+				.ColorAndOpacity(CacheStatsColor)
+				.Text(CacheStatsDurationText)
+				.Visibility_Lambda(CacheStatsVisibilityLambda)
+			]
+
+			+ SHorizontalBox::Slot()
+			.VAlign(VAlign_Center)
+			.AutoWidth()
+			[
+				SNew(STextBlock)
+				.Justification(ETextJustify::Center)
+				.Text_Lambda([]()
+				{
+					FAudioInsightsCacheManager& CacheManager = IAudioInsightsModule::GetChecked().GetCacheManager();
+					const float CacheDuration = CacheManager.GetCacheDuration();
+
+					return FText::Format(CacheStatsDurationValueText
+										, FText::AsNumber(CacheDuration, FSlateStyle::Get().GetShortTimeFormat()));
+				})
+				.Visibility_Lambda(CacheStatsVisibilityLambda)
+			];
+	}
+
+	TSharedRef<SWidget> FToolbarWidgets::MakeSettingsMenuContent()
+	{
+		using namespace ToolbarWidgetsPrivate;
+
+		FMenuBuilder MenuBuilder(false /*bShouldCloseWindowAfterMenuSelection*/, TSharedPtr<FUICommandList>());
+
+		if (IAudioInsightsModule::IsLiveMonitoringSession())
+		{
+			MenuBuilder.BeginSection("AudioInsightsToolbarSettingActions", CacheSettingsHeaderText);
+			{
+				MenuBuilder.AddMenuEntry(
+					CacheSettingsShowStatsText,
+					CacheSettingsShowStatsTooltipText,
+					FSlateIcon(),
+					FUIAction(
+						FExecuteAction::CreateLambda([]()
+						{
+							if (TObjectPtr<UAudioInsightsSettings> Settings = GetMutableDefault<UAudioInsightsSettings>())
+							{
+								FCacheSettings& CacheSettings = Settings->CacheSettings;
+								CacheSettings.bShowCacheStats = !CacheSettings.bShowCacheStats;
+								Settings->SaveCacheSettings();
+							}
+						}),
+						FCanExecuteAction(),
+						FIsActionChecked::CreateLambda([]()
+						{
+							if (const TObjectPtr<const UAudioInsightsSettings> Settings = GetDefault<UAudioInsightsSettings>())
+							{
+								const FCacheSettings& CacheSettings = Settings->CacheSettings;
+								return CacheSettings.bShowCacheStats;
+							}
+
+							return false;
+						})
+					),
+					NAME_None,
+					EUserInterfaceActionType::ToggleButton
+				);
+			}
+
+			MenuBuilder.EndSection();
+		}
+
+		return MenuBuilder.MakeWidget();
+	}
+
+	TSharedRef<SWidget> FToolbarWidgets::MakeSettingsButtonWidget()
+	{
+		using namespace ToolbarWidgetsPrivate;
+
+		return SNew(SComboButton)
+			.ComboButtonStyle(&FAppStyle::Get().GetWidgetStyle<FComboButtonStyle>("SimpleComboButton"))
+			.OnGetMenuContent_Lambda([this]()
+			{
+				return MakeSettingsMenuContent();
+			})
+			.MenuPlacement(EMenuPlacement::MenuPlacement_ComboBoxRight)
+			.HasDownArrow(false)
+			.ToolTipText(SettingsTooltipText)
+			.ButtonContent()
+			[
+				SNew(SHorizontalBox)
+
+				+ SHorizontalBox::Slot()
+				.AutoWidth()
+				.VAlign(VAlign_Center)
+				.Padding(5.0, 0.0f)
+				[
+					SNew(SImage)
+					.ColorAndOpacity(FSlateColor::UseForeground())
+					.Image(FSlateStyle::Get().GetBrush("AudioInsights.Icon.Settings"))
+					.DesiredSizeOverride(FVector2D(16.0f, 16.0f))
+				]
+			];
+	}
+} // namespace UE::Audio::Insights
+
+#undef LOCTEXT_NAMESPACE
diff --git a/Engine/Plugins/AudioInsights/Source/AudioInsights/Private/IAudioInsightsModule.cpp b/Engine/Plugins/AudioInsights/Source/AudioInsights/Private/IAudioInsightsModule.cpp
--- a/Engine/Plugins/AudioInsights/Source/AudioInsights/Private/IAudioInsightsModule.cpp
+++ b/Engine/Plugins/AudioInsights/Source/AudioInsights/Private/IAudioInsightsModule.cpp
@@ -1,6 +1,7 @@
 // Copyright Epic Games, Inc. All Rights Reserved.
 #include "IAudioInsightsModule.h"
 
+#include "AudioInsightsComponent.h"
 #include "AudioInsightsModule.h"
 
 IAudioInsightsTraceModule& IAudioInsightsModule::GetTraceModule()
@@ -34,6 +35,21 @@ bool IAudioInsightsModule::IsModuleLoaded()
 	return UE::Audio::Insights::FAudioInsightsModule::IsModuleLoaded();
 }
 
+bool IAudioInsightsModule::IsLiveMonitoringSession()
+{
+#if !WITH_EDITOR
+	UE::Audio::Insights::FAudioInsightsModule* const ModulePtr = UE::Audio::Insights::FAudioInsightsModule::GetModulePtr();
+	if (ModulePtr)
+	{
+		const TSharedPtr<UE::Audio::Insights::FAudioInsightsComponent> Component = ModulePtr->GetAudioInsightsComponent();
+		return Component.IsValid() && Component->GetIsLiveSession();
+	}
+	return false;
+#else
+	return true;
+#endif // !WITH_EDITOR
+}
+
 IAudioInsightsModule& IAudioInsightsModule::GetChecked()
 {
 	return static_cast<IAudioInsightsModule&>(UE::Audio::Insights::FAudioInsightsModule::GetChecked());
diff --git a/Engine/Plugins/AudioInsights/Source/AudioInsights/Public/AudioInsightsDashboardFactory.h b/Engine/Plugins/AudioInsights/Source/AudioInsights/Public/AudioInsightsDashboardFactory.h
--- a/Engine/Plugins/AudioInsights/Source/AudioInsights/Public/AudioInsightsDashboardFactory.h
+++ b/Engine/Plugins/AudioInsights/Source/AudioInsights/Public/AudioInsightsDashboardFactory.h
@@ -2,6 +2,7 @@
 #pragma once
 
 #include "AudioDefines.h"
+#include "AudioInsightsToolbarWidgets.h"
 #include "Containers/Map.h"
 #include "Framework/Docking/TabManager.h"
 #include "IAudioInsightsDashboardFactory.h"
@@ -39,6 +40,8 @@ namespace UE::Audio::Insights
 		FReply ToggleAutoEnableAudioTraces();
 #endif
 
+		TSharedPtr<FToolbarWidgets> ToolbarWidgets;
+
 		TSharedPtr<FTabManager> DashboardTabManager;
 		TSharedPtr<FWorkspaceItem> DashboardWorkspace;
 		TSharedPtr<FTabManager::FLayout> TabLayout;
diff --git a/Engine/Plugins/AudioInsights/Source/AudioInsights/Public/AudioInsightsSettings.h b/Engine/Plugins/AudioInsights/Source/AudioInsights/Public/AudioInsightsSettings.h
--- a/Engine/Plugins/AudioInsights/Source/AudioInsights/Public/AudioInsightsSettings.h
+++ b/Engine/Plugins/AudioInsights/Source/AudioInsights/Public/AudioInsightsSettings.h
@@ -3,6 +3,7 @@
 
 #include "Engine/DeveloperSettings.h"
 #include "LoudnessMeterSettings.h"
+#include "Settings/CacheSettings.h"
 #include "Settings/OutputMeterDashboardSettings.h"
 
 #include "AudioInsightsSettings.generated.h"
@@ -19,11 +20,24 @@ public:
 	UPROPERTY(EditAnywhere, config, Category = OutputMeter, meta = (ShowOnlyInnerProperties))
 	FLoudnessMeterSettings OutputMeterSettings;
 
+	/** Settings for the Cache */
+	UPROPERTY(EditAnywhere, config, Category = Cache, meta = (ShowOnlyInnerProperties))
+	FCacheSettings CacheSettings;
+
 	/** Settings for the Output Meter dashboard tab (not visible to the user) */
 	UPROPERTY(config)
 	FOutputMeterDashboardSettings OutputMeterDashboardSettings;
 
 	static FString GetAudioInsightsConfigFilename();
+
+#if WITH_EDITOR
+	virtual void PostEditChangeChainProperty(struct FPropertyChangedChainEvent& PropertyChangedEvent) override;
+#endif
+
+	void SaveCacheSettings();
+
+private:
+	void SavePropertyToConfigFile(const FProperty* const PropertyToSave);
 };
 
 #undef UE_API
diff --git a/Engine/Plugins/AudioInsights/Source/AudioInsights/Public/AudioInsightsToolbarWidgets.h b/Engine/Plugins/AudioInsights/Source/AudioInsights/Public/AudioInsightsToolbarWidgets.h
new file mode 100644
--- /dev/null
+++ b/Engine/Plugins/AudioInsights/Source/AudioInsights/Public/AudioInsightsToolbarWidgets.h
@@ -0,0 +1,39 @@
+// Copyright Epic Games, Inc. All Rights Reserved.
+#pragma once
+
+#include "Delegates/Delegate.h"
+#include "Templates/SharedPointer.h"
+
+#define UE_API AUDIOINSIGHTS_API
+
+class SWidget;
+
+namespace UE::Audio::Insights
+{
+	class UE_API FToolbarWidgets final
+	{
+	public:
+		FToolbarWidgets();
+		~FToolbarWidgets();
+
+		TSharedRef<SWidget> MakePauseButtonWidget();
+		TSharedRef<SWidget> MakeCacheStatsWidget();
+		TSharedRef<SWidget> MakeSettingsButtonWidget();
+
+	private:
+		void InitializeDelegates();
+		void DeinitializeDelegates();
+
+		TSharedRef<SWidget> MakeSettingsMenuContent();
+
+		void OnTimingViewTimeMarkerChanged(double InTimeMarker);
+		void OnTimeControlMethodReset();
+
+		bool bIsPaused = false;
+
+		FDelegateHandle OnTimingViewTimeMarkerChangedHandle;
+		FDelegateHandle OnTimeControlMethodResetHandle;
+	};
+} // namespace UE::Audio::Insights
+
+#undef UE_API
diff --git a/Engine/Plugins/AudioInsights/Source/AudioInsights/Public/IAudioInsightsModule.h b/Engine/Plugins/AudioInsights/Source/AudioInsights/Public/IAudioInsightsModule.h
--- a/Engine/Plugins/AudioInsights/Source/AudioInsights/Public/IAudioInsightsModule.h
+++ b/Engine/Plugins/AudioInsights/Source/AudioInsights/Public/IAudioInsightsModule.h
@@ -44,6 +44,8 @@ public:
 
 	UE_API static bool IsModuleLoaded();
 
+	static UE_API bool IsLiveMonitoringSession();
+
 	static UE_API IAudioInsightsModule& GetChecked();
 
 #ifdef WITH_EDITOR
diff --git a/Engine/Plugins/AudioInsights/Source/AudioInsights/Public/Settings/CacheSettings.h b/Engine/Plugins/AudioInsights/Source/AudioInsights/Public/Settings/CacheSettings.h
new file mode 100644
--- /dev/null
+++ b/Engine/Plugins/AudioInsights/Source/AudioInsights/Public/Settings/CacheSettings.h
@@ -0,0 +1,14 @@
+// Copyright Epic Games, Inc. All Rights Reserved.
+#pragma once
+
+#include "CacheSettings.generated.h"
+
+USTRUCT()
+struct FCacheSettings
+{
+	GENERATED_BODY()
+
+	/** Whether to show cache stats in the Audio Insights toolbar. */
+	UPROPERTY(EditAnywhere, config, Category = Cache)
+	bool bShowCacheStats = true;
+};
diff --git a/Engine/Plugins/AudioInsights/Source/AudioInsightsEditor/Private/AudioInsightsEditorDashboardFactory.cpp b/Engine/Plugins/AudioInsights/Source/AudioInsightsEditor/Private/AudioInsightsEditorDashboardFactory.cpp
--- a/Engine/Plugins/AudioInsights/Source/AudioInsightsEditor/Private/AudioInsightsEditorDashboardFactory.cpp
+++ b/Engine/Plugins/AudioInsights/Source/AudioInsightsEditor/Private/AudioInsightsEditorDashboardFactory.cpp
@@ -60,13 +60,6 @@ namespace UE::Audio::Insights
 
 		static const FText SaveSnapshotTooltipText = LOCTEXT("AudioDashboardEditor_SnapshotButton_TooltipText", "Saves a snapshot of the most recent trace data to a .utrace file. By default snapshots are limited to 32MB and can be accessed in Unreal Insights. To specify a custom snapshot size (in MB) launch the editor executable with the -tracetailmb=<size> option.");
 
-		static const FText CachePauseTooltipText = LOCTEXT("AudioDashboardEditor_PauseButton_TooltipText", "Click to pause monitoring.\n\nDepending on your cache settings, this may also stop the cache collecting data from the current session.");
-		static const FText CacheResumeTooltipText = LOCTEXT("AudioDashboardEditor_ResumeButton_TooltipText", "Click to resume monitoring.");
-
-		static const FText CacheSettingsHeaderText = LOCTEXT("AudioDashboardEditor_CacheSettings_HeaderText", "Cache Settings");
-		static const FText CacheSettingsShowStatsText = LOCTEXT("AudioDashboardEditor_CacheSettings_ShowCacheStats", "Show Cache Stats");
-		static const FText CacheSettingsShowStatsTooltipText = LOCTEXT("AudioDashboardEditor_CacheSettings_ShowCacheStatsTooltip", "Shows/Hides Cache Stats");
-
 		static const FText SaveSnapshotNotificationText = LOCTEXT("AudioDashboardEditor_SnapshotButton_NotificationText", "Audio Insights Snapshot saved.");
 		static const FText SaveSnapshotNotificationSubText = LOCTEXT("AudioDashboardEditor_SnapshotButton_NotificationSubText", "A snapshot .utrace with the most recent events has been saved to your trace server.");
 
@@ -80,7 +73,6 @@ namespace UE::Audio::Insights
 		static const FText WorldFilterText = LOCTEXT("AudioDashboardEditor_WorldFilter_Text", "World Filter:");
 		static const FText WorldFilterTooltipText = LOCTEXT("AudioDashboardEditor_WorldFilter_TooltipText", "Select world(s) to monitor (worlds may share audio output).");
 
-		static const FText SettingsTooltipText = LOCTEXT("AudioDashboardEditor_Settings_TooltipText", "Settings");
 		static const FText EditorPreferencesTooltipText = LOCTEXT("AudioDashboardEditor_EditorPreferences_TooltipText", "Opens Audio Insights Editor Preferences.");
 		
 		void ShowNotification(const FText& InText, const FText& InSubText)
@@ -197,16 +189,6 @@ namespace UE::Audio::Insights
 		TraceModule.OnOnlyTraceAudioChannelsStateChanged(bInOnlyTraceAudioChannels);
 	}
 
-	void FEditorDashboardFactory::OnTimingViewTimeMarkerChanged(double InTimeMarker)
-	{
-		bIsPaused = true;
-	}
-
-	void FEditorDashboardFactory::OnTimeControlMethodReset()
-	{
-		bIsPaused = false;
-	}
-
 	void FEditorDashboardFactory::CreateTraceBookmark()
 	{
 		using namespace EditorDashboardFactoryPrivate;
@@ -330,21 +312,6 @@ namespace UE::Audio::Insights
 			}
 		}
 
-		{
-			FAudioInsightsTimingViewExtender& TimingViewExtender = FAudioInsightsEditorModule::GetChecked().GetTimingViewExtender();
-	
-			if (OnTimingViewTimeMarkerChangedHandle.IsValid())
-			{
-				TimingViewExtender.OnTimingViewTimeMarkerChanged.Remove(OnTimingViewTimeMarkerChangedHandle);
-				OnTimingViewTimeMarkerChangedHandle.Reset();
-			}
-	
-			if (OnTimeControlMethodResetHandle.IsValid())
-			{
-				TimingViewExtender.OnTimeControlMethodReset.Remove(OnTimeControlMethodResetHandle);
-				OnTimeControlMethodResetHandle.Reset();
-			}
-		}
 	}
 
 	void FEditorDashboardFactory::BindCommands()
@@ -391,6 +358,8 @@ namespace UE::Audio::Insights
 			}
 		}));
 
+		ToolbarWidgets = MakeShared<FToolbarWidgets>();
+
 		InitDelegates();
 
 		RegisterTabSpawners();
@@ -433,6 +402,7 @@ namespace UE::Audio::Insights
 			}
 
 			ResetDelegates();
+			ToolbarWidgets.Reset();
 			UnregisterTabSpawners();
 			SaveLayoutToConfig();
 
@@ -606,14 +576,6 @@ namespace UE::Audio::Insights
 		BackgroundColorBrush.TintColor = FSlateColor(FLinearColor(0.018f, 0.018f, 0.018f, 1.0f));
 		BackgroundColorBrush.DrawAs = ESlateBrushDrawType::Box;
 
-		constexpr FLinearColor CacheStatsColor(1.0f, 1.0f, 1.0f, 0.5f);
-
-		const auto CacheStatsVisibilityLambda = []()
-		{
-			const TObjectPtr<const UAudioInsightsEditorSettings> Settings = GetDefault<UAudioInsightsEditorSettings>();
-			return (Settings && Settings->bShowCacheStats) ? EVisibility::Visible : EVisibility::Collapsed;
-		};
-
 		return SNew(SBorder)
 			.BorderImage(&BackgroundColorBrush)
 			[
@@ -820,106 +782,15 @@ namespace UE::Audio::Insights
 				.AutoWidth()
 				.Padding(8.0f, 0.0f, 0.0f, 0.0f)
 				[
-					SNew(SButton)
-					.ButtonStyle(&FAppStyle::Get().GetWidgetStyle<FButtonStyle>("SimpleButton"))
-					.ToolTipText_Lambda([this]()
-					{
-						return bIsPaused ? CacheResumeTooltipText : CachePauseTooltipText;
-					})
-					.OnClicked_Lambda([this]()
-					{
-						FAudioInsightsTimingViewExtender& TimingViewExtender = FAudioInsightsEditorModule::GetChecked().GetTimingViewExtender();
-
-						if (bIsPaused)
-						{
-							TimingViewExtender.ResumeTimeMarker();
-						}
-						else
-						{
-							const double PauseTime = FAudioInsightsEditorModule::GetChecked().GetCacheManager().GetCacheEndTimeStamp();
-							TimingViewExtender.PauseTimeMarker(PauseTime, ESystemControllingTimeMarker::External);
-						}
-
-						return FReply::Handled();
-					})
-					[
-						SNew(SBox)
-						.HeightOverride(16.0f)
-						[
-							SNew(SHorizontalBox)
-							+ SHorizontalBox::Slot()
-							.VAlign(VAlign_Center)
-							.AutoWidth()
-							[
-								SNew(SImage)
-								.Image(FSlateStyle::Get().GetBrush("AudioInsights.Icon.Pause"))
-								.ColorAndOpacity_Lambda([this]()
-								{
-									return bIsPaused ? FStyleColors::AccentBlue : FSlateColor::UseForeground();
-								})
-							]
-						]
-					]
+					ToolbarWidgets->MakePauseButtonWidget()
 				]
 
 				// Cache stats
 				+ SHorizontalBox::Slot()
-				.Padding(12.0f, 0.0f, 0.0f, 0.0f)
 				.VAlign(VAlign_Center)
 				.AutoWidth()
 				[
-					SNew(STextBlock)
-					.Justification(ETextJustify::Center)
-					.ColorAndOpacity(CacheStatsColor)
-					.Text(LOCTEXT("AudioDashboardEditor_CacheStats_Title", "Cache Size: "))
-					.Visibility_Lambda(CacheStatsVisibilityLambda)
-				]
-
-				+ SHorizontalBox::Slot()
-				.VAlign(VAlign_Center)
-				.AutoWidth()
-				[
-					SNew(STextBlock)
-					.Justification(ETextJustify::Center)
-					.Text_Lambda([]()
-					{
-						FAudioInsightsCacheManager& CacheManager = FAudioInsightsEditorModule::GetChecked().GetCacheManager();
-						const float CacheSizeMb = static_cast<float>(CacheManager.GetUsedCacheSize()) / 1000000.0f;
-						const float MaxCacheSizeMb = static_cast<float>(CacheManager.GetMaxCacheSize()) / 1000000.0f;
-
-						return FText::Format(LOCTEXT("AudioDashboardEditor_CacheStats_Memory", "{0} / {1}")
-											, FText::AsNumber(CacheSizeMb, FSlateStyle::Get().GetMemoryFormat())
-											, FSlateStyle::Get().FormatMemoryAsMegabytes(MaxCacheSizeMb));
-					})
-					.Visibility_Lambda(CacheStatsVisibilityLambda)
-				]
-
-				+ SHorizontalBox::Slot()
-				.VAlign(VAlign_Center)
-				.AutoWidth()
-				[
-					SNew(STextBlock)
-					.Justification(ETextJustify::Center)
-					.ColorAndOpacity(CacheStatsColor)
-					.Text(LOCTEXT("AudioDashboardEditor_CacheStats_Duration", " : Duration: "))
-					.Visibility_Lambda(CacheStatsVisibilityLambda)
-				]
-
-				+ SHorizontalBox::Slot()
-				.VAlign(VAlign_Center)
-				.AutoWidth()
-				[
-					SNew(STextBlock)
-					.Justification(ETextJustify::Center)
-					.Text_Lambda([]()
-					{
-						FAudioInsightsCacheManager& CacheManager = FAudioInsightsEditorModule::GetChecked().GetCacheManager();
-						const float CacheDuration = CacheManager.GetCacheDuration();
-
-						return FText::Format(LOCTEXT("AudioDashboardEditor_CacheStats_DurationValue", "{0}s")
-											, FText::AsNumber(CacheDuration, FSlateStyle::Get().GetShortTimeFormat()));
-					})
-					.Visibility_Lambda(CacheStatsVisibilityLambda)
+					ToolbarWidgets->MakeCacheStatsWidget()
 				]
 
 				// World filter label
@@ -984,7 +855,7 @@ namespace UE::Audio::Insights
 				.AutoWidth()
 				.HAlign(HAlign_Right)
 				[
-					CreateSettingsButtonWidget()
+					ToolbarWidgets->MakeSettingsButtonWidget()
 				]
 			];
 	}
@@ -1150,72 +1021,6 @@ namespace UE::Audio::Insights
 			];
 	}
 
-	TSharedRef<SWidget> FEditorDashboardFactory::CreateSettingsButtonWidget()
-	{
-		if (OnGetSettingsMenuContent() == SNullWidget::NullWidget)
-		{
-			return SNullWidget::NullWidget;
-		}
-
-		return SNew(SComboButton)
-			.ComboButtonStyle(&FAppStyle::Get().GetWidgetStyle<FComboButtonStyle>("SimpleComboButton"))
-			.OnGetMenuContent(this, &FEditorDashboardFactory::OnGetSettingsMenuContent)
-			.MenuPlacement(EMenuPlacement::MenuPlacement_ComboBoxRight)
-			.HasDownArrow(false)
-			.ToolTipText(EditorDashboardFactoryPrivate::SettingsTooltipText)
-			.ButtonContent()
-			[
-				SNew(SHorizontalBox)
-
-				+ SHorizontalBox::Slot()
-				.AutoWidth()
-				.VAlign(VAlign_Center)
-				.Padding(5.0, 0.0f)
-				[
-					SNew(SImage)
-					.ColorAndOpacity(FSlateColor::UseForeground())
-					.Image(FSlateStyle::Get().GetBrush("AudioInsights.Icon.Settings"))
-					.DesiredSizeOverride(FVector2D(16.0f, 16.0f))
-				]
-			];
-	}
-
-	TSharedRef<SWidget> FEditorDashboardFactory::OnGetSettingsMenuContent()
-	{
-		FMenuBuilder MenuBuilder(false /*bShouldCloseWindowAfterMenuSelection*/, CommandList);
-
-		MenuBuilder.BeginSection("AudioDashboardEditorSettingActions", EditorDashboardFactoryPrivate::CacheSettingsHeaderText);
-		{
-			MenuBuilder.AddMenuEntry(
-				EditorDashboardFactoryPrivate::CacheSettingsShowStatsText,
-				EditorDashboardFactoryPrivate::CacheSettingsShowStatsTooltipText,
-				FSlateIcon(),
-				FUIAction(
-					FExecuteAction::CreateLambda([]()
-					{
-						if (TObjectPtr<UAudioInsightsEditorSettings> Settings = GetMutableDefault<UAudioInsightsEditorSettings>())
-						{
-							Settings->bShowCacheStats = !Settings->bShowCacheStats;
-							Settings->SaveConfig();
-						}
-					}),
-					FCanExecuteAction(),
-					FIsActionChecked::CreateLambda([]()
-					{
-						const TObjectPtr<const UAudioInsightsEditorSettings> Settings = GetDefault<UAudioInsightsEditorSettings>();
-						return Settings && Settings->bShowCacheStats;
-					})
-				),
-				NAME_None,
-				EUserInterfaceActionType::ToggleButton
-			);
-		}
-
-		MenuBuilder.EndSection();
-
-		return MenuBuilder.MakeWidget();
-	}
-
 	void FEditorDashboardFactory::SetActiveDeviceId(const ::Audio::FDeviceId InDeviceId)
 	{
 		const TObjectPtr<const UAudioInsightsEditorSettings> AudioInsightsEditorSettings = GetDefault<UAudioInsightsEditorSettings>();
@@ -1266,19 +1071,6 @@ namespace UE::Audio::Insights
 			}
 		}
 
-		{
-			FAudioInsightsTimingViewExtender& TimingViewExtender = FAudioInsightsEditorModule::GetChecked().GetTimingViewExtender();
-	
-			if (!OnTimingViewTimeMarkerChangedHandle.IsValid())
-			{
-				OnTimingViewTimeMarkerChangedHandle = TimingViewExtender.OnTimingViewTimeMarkerChanged.AddSP(this, &FEditorDashboardFactory::OnTimingViewTimeMarkerChanged);
-			}
-	
-			if (!OnTimeControlMethodResetHandle.IsValid())
-			{
-				OnTimeControlMethodResetHandle = TimingViewExtender.OnTimeControlMethodReset.AddSP(this, &FEditorDashboardFactory::OnTimeControlMethodReset);
-			}
-		}
 	}
 
 	TSharedRef<FTabManager::FLayout> FEditorDashboardFactory::GetDefaultTabLayout()
diff --git a/Engine/Plugins/AudioInsights/Source/AudioInsightsEditor/Private/AudioInsightsEditorDashboardFactory.h b/Engine/Plugins/AudioInsights/Source/AudioInsightsEditor/Private/AudioInsightsEditorDashboardFactory.h
--- a/Engine/Plugins/AudioInsights/Source/AudioInsightsEditor/Private/AudioInsightsEditorDashboardFactory.h
+++ b/Engine/Plugins/AudioInsights/Source/AudioInsightsEditor/Private/AudioInsightsEditorDashboardFactory.h
@@ -2,6 +2,7 @@
 #pragma once
 
 #include "AudioDefines.h"
+#include "AudioInsightsToolbarWidgets.h"
 #include "Containers/Map.h"
 #include "Engine/World.h"
 #include "Framework/Docking/TabManager.h"
@@ -42,10 +43,6 @@ namespace UE::Audio::Insights
 
 		void OnOnlyTraceAudioChannelsPropertyChanged(const bool bInOnlyTraceAudioChannels);
 
-		void OnTimingViewTimeMarkerChanged(double InTimeMarker);
-
-		void OnTimeControlMethodReset();
-
 		void CreateTraceBookmark();
 
 		void RefreshDeviceSelector();
@@ -59,8 +56,6 @@ namespace UE::Audio::Insights
 		TSharedRef<SWidget> MakeTraceModeButton(const ETraceMode InTraceButtonMode);
 
 		TSharedRef<SWidget> CreateEditorPreferencesButtonWidget();
-		TSharedRef<SWidget> CreateSettingsButtonWidget();
-		TSharedRef<SWidget> OnGetSettingsMenuContent();
 
 		void SetActiveDeviceId(const ::Audio::FDeviceId InDeviceId);
 		void InitDelegates();
@@ -84,12 +79,9 @@ namespace UE::Audio::Insights
 
 		FDelegateHandle OnOnlyTraceAudioChannelsPropertyChangedHandle;
 
-		FDelegateHandle OnTimingViewTimeMarkerChangedHandle;
-		FDelegateHandle OnTimeControlMethodResetHandle;
-
 		bool bIsLayoutReseting = false;
 
-		bool bIsPaused { false };
+		TSharedPtr<FToolbarWidgets> ToolbarWidgets;
 
 		TSharedPtr<FTabManager> DashboardTabManager;
 		TSharedPtr<FWorkspaceItem> DashboardWorkspace;
diff --git a/Engine/Plugins/AudioInsights/Source/AudioInsightsEditor/Public/AudioInsightsEditorSettings.h b/Engine/Plugins/AudioInsights/Source/AudioInsightsEditor/Public/AudioInsightsEditorSettings.h
--- a/Engine/Plugins/AudioInsights/Source/AudioInsightsEditor/Public/AudioInsightsEditorSettings.h
+++ b/Engine/Plugins/AudioInsights/Source/AudioInsightsEditor/Public/AudioInsightsEditorSettings.h
@@ -39,10 +39,6 @@ public:
 	UPROPERTY(Config, EditAnywhere, Category = "Trace Channels")
 	bool bOnlyTraceAudioChannels = true;
 
-	/** Whether to show cache stats in the Audio Insights toolbar. */
-	UPROPERTY(Config, EditAnywhere, Category = "Cache")
-	bool bShowCacheStats = true;
-
 	/** Settings for analyzer rack spectrogram widget */
 	UPROPERTY(EditAnywhere, config, Category = Spectrogram, meta = (ShowOnlyInnerProperties))
 	FSpectrogramRackUnitSettings SpectrogramSettings;
