[Audio Insights] Cleanup: Removed old LoudnessMeterSettings from UAudioInsightsEditorSettings + make OutputMeterSettings visible only via cog + show AudioInsightsSettings in Editor settings (missing code change that should have been part of 52569349)

#Synced-CL 52673159

diff --git a/Engine/Plugins/AudioInsights/Source/AudioInsights/Public/AudioInsightsSettings.h b/Engine/Plugins/AudioInsights/Source/AudioInsights/Public/AudioInsightsSettings.h
--- a/Engine/Plugins/AudioInsights/Source/AudioInsights/Public/AudioInsightsSettings.h
+++ b/Engine/Plugins/AudioInsights/Source/AudioInsights/Public/AudioInsightsSettings.h
@@ -16,14 +16,14 @@ class UAudioInsightsSettings : public UDeveloperSettings
 	GENERATED_BODY()
 
 public:
-	/** Settings for the Output Meters */
-	UPROPERTY(EditAnywhere, config, Category = OutputMeter, meta = (ShowOnlyInnerProperties))
-	FLoudnessMeterSettings OutputMeterSettings;
-
 	/** Settings for the Cache */
 	UPROPERTY(EditAnywhere, config, Category = Cache, meta = (ShowOnlyInnerProperties))
 	FCacheSettings CacheSettings;
 
+	/** Settings for the Output Meters (only visible via settings cog)*/
+	UPROPERTY(config)
+	FLoudnessMeterSettings OutputMeterSettings;
+
 	/** Settings for the Output Meter dashboard tab (not visible to the user) */
 	UPROPERTY(config)
 	FOutputMeterDashboardSettings OutputMeterDashboardSettings;
@@ -37,6 +37,10 @@ public:
 	void SaveCacheSettings();
 
 private:
+	FName GetContainerName() const override { return TEXT("Editor"); }
+	FName GetCategoryName()  const override { return TEXT("Plugins"); }
+	FName GetSectionName()   const override { return TEXT("Audio Insights"); }
+
 	void SavePropertyToConfigFile(const FProperty* const PropertyToSave);
 };
 
diff --git a/Engine/Plugins/AudioInsights/Source/AudioInsightsEditor/Public/AudioInsightsEditorSettings.h b/Engine/Plugins/AudioInsights/Source/AudioInsightsEditor/Public/AudioInsightsEditorSettings.h
--- a/Engine/Plugins/AudioInsights/Source/AudioInsightsEditor/Public/AudioInsightsEditorSettings.h
+++ b/Engine/Plugins/AudioInsights/Source/AudioInsightsEditor/Public/AudioInsightsEditorSettings.h
@@ -47,10 +47,6 @@ public:
 	UPROPERTY(EditAnywhere, config, Category = SpectrumAnalyzer, meta = (ShowOnlyInnerProperties))
 	FSpectrumAnalyzerRackUnitSettings SpectrumAnalyzerSettings;
 
-	/** Settings for analyzer rack loudness meter widget */
-	UPROPERTY(EditAnywhere, config, Category = LoudnessMeter, meta = (ShowOnlyInnerProperties))
-	FLoudnessMeterRackUnitSettings LoudnessMeterSettings;
-
 	/** Settings for the Event Log */
 	UPROPERTY(EditAnywhere, config, Category = EventLog, meta = (ShowOnlyInnerProperties))
 	FAudioEventLogSettings EventLogSettings;
