[Audio Insights] [AudioModulationInsights] Write a trace file based on the data contained in the Audio Inisights dashboard cache, the old Save Trace Snapshot functionality has been replaced with this new one.

#Synced-CL 52673159

diff --git a/Engine/Plugins/AudioInsights/Source/AudioInsights/AudioInsights.Build.cs b/Engine/Plugins/AudioInsights/Source/AudioInsights/AudioInsights.Build.cs
--- a/Engine/Plugins/AudioInsights/Source/AudioInsights/AudioInsights.Build.cs
+++ b/Engine/Plugins/AudioInsights/Source/AudioInsights/AudioInsights.Build.cs
@@ -26,6 +26,7 @@ public class AudioInsights : ModuleRules
 			{
 				"AudioMixerCore",
 				"AudioWidgetsCore",
+				"BuildSettings",
 				"CoreUObject",
 				"DeveloperSettings",
 				"InputCore",
diff --git a/Engine/Plugins/AudioInsights/Source/AudioInsights/Private/AudioInsightsDashboardFactory.cpp b/Engine/Plugins/AudioInsights/Source/AudioInsights/Private/AudioInsightsDashboardFactory.cpp
--- a/Engine/Plugins/AudioInsights/Source/AudioInsights/Private/AudioInsightsDashboardFactory.cpp
+++ b/Engine/Plugins/AudioInsights/Source/AudioInsights/Private/AudioInsightsDashboardFactory.cpp
@@ -1,6 +1,7 @@
 // Copyright Epic Games, Inc. All Rights Reserved.
 #include "AudioInsightsDashboardFactory.h"
 
+#include "AudioInsightsComponent.h"
 #include "AudioInsightsModule.h"
 #include "AudioInsightsStyle.h"
 #include "Framework/Docking/TabManager.h"
@@ -28,8 +29,8 @@ namespace UE::Audio::Insights
 		static const FText EnableThemNowText = LOCTEXT("AudioDashboard_EnableNowText", "Enable them now?");
 		static const FText TraceControllerUnavailabledWarning = LOCTEXT("AudioDashboard_TraceControllerUnavailableWarning", "The Trace Controller API is currently unavailable.");
 		static const FText TryEnablingMessagingText = LOCTEXT("AudioDashboard_TryEnablingMessagingText", "Make sure you have launched this package with the -Messaging command line argument.");
-		static const FText SaveSnapshotText = LOCTEXT("AudioDashboard_SnapshotButton_Text", "Save Snapshot");
-		static const FText SaveSnapshotTooltipText = LOCTEXT("AudioDashboard_SnapshotButton_TooltipText", "Saves a snapshot of the most recent trace data to a .utrace file. By default snapshots are limited to 4MB and can be accessed in Unreal Insights. To specify a custom snapshot size (in MB) launch the game executable with the -tracetailmb=<size> option.");
+		static const FText SaveCacheSnapshotText = LOCTEXT("AudioDashboard_SaveCacheSnapshotButton_Text", "Save Cache Snapshot");
+		static const FText SaveCacheSnapshotTooltipText = LOCTEXT("AudioDashboard_SaveCacheSnapshotButton_TooltipText", "Saves a snapshot of the Audio Insights cached data to a .utrace file containing only audio trace events.");
 	}
 
 	::Audio::FDeviceId FDashboardFactory::GetDeviceId() const
@@ -204,14 +205,14 @@ namespace UE::Audio::Insights
 	TSharedRef<SWidget> FDashboardFactory::MakeMainToolbarWidget()
 	{
 		return SNew(SHorizontalBox)
-			// Save snapshot button
+			// Save cache snapshot button
 			+ SHorizontalBox::Slot()
 			.HAlign(HAlign_Left)
 			.VAlign(VAlign_Center)
 			.AutoWidth()
 			.Padding(12.0f, 0.0f, 0.0f, 0.0f)
 			[
-				MakeSnapshotButtonWidget()
+				MakeCacheSnapshotButtonWidget()
 			]
 			// Cache toolbar widgets (only visible during live monitoring)
 			+ SHorizontalBox::Slot()
@@ -266,17 +267,24 @@ namespace UE::Audio::Insights
 			];
 	}
 
-	TSharedRef<SWidget> FDashboardFactory::MakeSnapshotButtonWidget()
+	TSharedRef<SWidget> FDashboardFactory::MakeCacheSnapshotButtonWidget()
 	{
 		using namespace DashboardFactoryPrivate;
 
 		return SNew(SButton)
 			.ButtonStyle(&FAppStyle::Get().GetWidgetStyle<FButtonStyle>("SimpleButton"))
-			.ToolTipText(SaveSnapshotTooltipText)
+			.ToolTipText(SaveCacheSnapshotTooltipText)
+#if !WITH_EDITOR
+			.IsEnabled_Lambda([]()
+			{
+				const TSharedPtr<FAudioInsightsComponent> Component = FAudioInsightsModule::GetChecked().GetAudioInsightsComponent();
+				return Component.IsValid() && Component->GetIsLiveSession();
+			})
+#endif // !WITH_EDITOR
 			.OnClicked_Lambda([]()
 			{
 				FTraceModule& AudioInsightsTraceModule = static_cast<FTraceModule&>(FAudioInsightsModule::GetChecked().GetTraceModule());
-				AudioInsightsTraceModule.SaveTraceSnapshot();
+				AudioInsightsTraceModule.SaveCacheSnapshot();
 
 				return FReply::Handled();
 			})
@@ -285,16 +293,16 @@ namespace UE::Audio::Insights
 				.HeightOverride(16.0f)
 				[
 					SNew(SHorizontalBox)
-					// Trace snapshot icon
+					// Save cache snapshot icon
 					+ SHorizontalBox::Slot()
 					.VAlign(VAlign_Center)
 					.AutoWidth()
 					[
 						SNew(SImage)
-						.Image(FSlateStyle::Get().GetBrush("AudioInsights.Icon.Trace.Snapshot"))
+						.Image(FSlateStyle::Get().GetBrush("AudioInsights.Icon.Trace.CacheSnapshot"))
 						.ColorAndOpacity(FSlateColor::UseForeground())
 					]
-					// Trace snapshot text
+					// Save cache snapshot text
 					+ SHorizontalBox::Slot()
 					.VAlign(VAlign_Center)
 					.AutoWidth()
@@ -302,7 +310,7 @@ namespace UE::Audio::Insights
 					[
 						SNew(STextBlock)
 						.Justification(ETextJustify::Center)
-						.Text(SaveSnapshotText)
+						.Text(SaveCacheSnapshotText)
 					]
 				]
 			];
diff --git a/Engine/Plugins/AudioInsights/Source/AudioInsights/Private/AudioInsightsTraceModule.cpp b/Engine/Plugins/AudioInsights/Source/AudioInsights/Private/AudioInsightsTraceModule.cpp
--- a/Engine/Plugins/AudioInsights/Source/AudioInsights/Private/AudioInsightsTraceModule.cpp
+++ b/Engine/Plugins/AudioInsights/Source/AudioInsights/Private/AudioInsightsTraceModule.cpp
@@ -15,6 +15,7 @@
 #include "Trace/Trace.h"
 #include "TraceServices/Model/Channel.h"
 #include "TraceServices/Model/Diagnostics.h"
+#include "TraceWriter/AudioInsightsCacheTraceWriter.h"
 
 namespace UE::Audio::Insights
 {
@@ -23,6 +24,47 @@ namespace UE::Audio::Insights
 		static const TCHAR* AudioChannelName = TEXT("Audio");
 		static const TCHAR* AudioMixerChannelName = TEXT("AudioMixer");
 		static const TCHAR* BookmarkChannelName = TEXT("Bookmark");
+
+		FString GetTraceStoreFilePath(const FString& FileName)
+		{
+			auto GetTraceHostAddress = []() -> FString
+			{
+				TSharedPtr<FInternetAddr> RecorderAddress;
+
+				if (ISocketSubsystem* const Sockets = ISocketSubsystem::Get())
+				{
+					bool bCanBindAll = false;
+					RecorderAddress = Sockets->GetLocalHostAddr(*GLog, bCanBindAll);
+				}
+
+				if (RecorderAddress.IsValid())
+				{
+					constexpr bool bAppendPort = false;
+					return RecorderAddress->ToString(bAppendPort);
+				}
+
+				return "127.0.0.1";
+			};
+
+			const TUniquePtr<UE::Trace::FStoreClient> StoreClient = UE::Trace::FStoreClient::Create(*GetTraceHostAddress());
+
+			if (StoreClient.IsValid())
+			{
+				if (const UE::Trace::FStoreClient::FStatus* StoreClientStatus = StoreClient->GetStatus())
+				{
+					const FString TraceStorePath = FString(StoreClientStatus->GetStoreDir());
+
+					FString FilePath = FPaths::Combine(TraceStorePath, FileName);
+					FPaths::NormalizeFilename(FilePath);
+
+					return FilePath;
+				}
+			}
+
+			UE_LOGF(LogAudioInsights, Warning, "Failed to connect to the Trace Store Client, returning an empty path.");
+
+			return FString();
+		}
 	}
 
 	void FRewindDebuggerAudioInsightsRuntime::RecordingStarted()
@@ -541,79 +583,31 @@ namespace UE::Audio::Insights
 		}
 	}
 	
-	void FTraceModule::SaveTraceSnapshot()
+	void FTraceModule::SaveCacheSnapshot()
 	{
-		auto GetSnapshotFilePath = []()
-		{
-			auto GetTraceHostAddress = []() -> FString
-			{
-				TSharedPtr<FInternetAddr> RecorderAddress;
-				
-				if (ISocketSubsystem* Sockets = ISocketSubsystem::Get())
-				{
-					bool bCanBindAll = false;
-					RecorderAddress = Sockets->GetLocalHostAddr(*GLog, bCanBindAll);
-				}
-				
-				if (RecorderAddress.IsValid())
-				{
-					constexpr bool bAppendPort = false;
-					return RecorderAddress->ToString(bAppendPort);
-				}
-				
-				return "127.0.0.1";
-			};
-			
-			
-			const TUniquePtr<UE::Trace::FStoreClient> StoreClient = UE::Trace::FStoreClient::Create(*GetTraceHostAddress());
-			if (StoreClient.IsValid())
-			{
-				if (const UE::Trace::FStoreClient::FStatus* StoreClientStatus = StoreClient->GetStatus())
-				{
-					const FString TraceStorePath = FString(StoreClientStatus->GetStoreDir());
-					FString SnapshotFileName = FDateTime::Now().ToString(TEXT("%Y%m%d_%H%M%S_AudioInsights.utrace"));
-					
-					FString SnapshotFilePath = FPaths::Combine(TraceStorePath, SnapshotFileName);
-					FPaths::NormalizeFilename(SnapshotFilePath);
-					
-					return SnapshotFilePath;
-				}
-			}
-
-			UE_LOGF(LogAudioInsights, Warning, "Failed to connect to the Trace Store Client when trying to save a snapshot, returning an empty snapshot path.");
+		using namespace FTraceModulePrivate;
 
-			return FString();
-		};
+		const FAudioInsightsCacheManager& CacheManager = FAudioInsightsModule::GetChecked().GetCacheManager();
 
-		const FString SnapshotFilePath = GetSnapshotFilePath();
-		if (SnapshotFilePath.IsEmpty())
+		if (CacheManager.GetNumUsedChunks() == 0)
 		{
-			UE_LOGF(LogAudioInsights, Warning, "Failed to save Audio Insights snapshot. Snapshot path is empty.");
+			UE_LOGF(LogAudioInsights, Warning, "Cache is empty, nothing to export.");
 			return;
 		}
 
-#if WITH_EDITOR
-		const bool bWriteSnapshotSuccess = FTraceAuxiliary::WriteSnapshot(*SnapshotFilePath);
-		if (!bWriteSnapshotSuccess)
-		{
-			UE_LOGF(LogAudioInsights, Warning, "Failed to save Audio Insights snapshot in %ls", *SnapshotFilePath);
-		}
-#else
-		ISessionServicesModule& SessionServicesModule = FModuleManager::LoadModuleChecked<ISessionServicesModule>("SessionServices");
-			
-		TSharedPtr<ITraceController> TraceController = SessionServicesModule.GetTraceController();
-		if (TraceController.IsValid())
+		const FString FileName = FDateTime::Now().ToString(TEXT("%Y%m%d_%H%M%S_AudioInsightsCache.utrace"));
+		const FString FilePath = GetTraceStoreFilePath(FileName);
+
+		if (FilePath.IsEmpty())
 		{
-			TraceController->WithInstance(InstanceID, [SnapshotFilePath](const FTraceStatus& Status, ITraceControllerCommands& Commands)
-			{
-				Commands.SnapshotFile(SnapshotFilePath);
-			});
+			UE_LOGF(LogAudioInsights, Warning, "Failed to determine trace store path for cache snapshot.");
+			return;
 		}
-		else
+
+		if (!FAudioInsightsCacheTraceWriter::WriteCacheToFile(CacheManager, FilePath))
 		{
-			UE_LOGF(LogAudioInsights, Warning, "Failed to save Audio Insights snapshot. TraceController is invalid.");
+			UE_LOGF(LogAudioInsights, Warning, "Failed to write cache snapshot to %ls", *FilePath);
 		}
-#endif // WITH_EDITOR
 	}
 
 	ETraceMode FTraceModule::GetTraceMode() const
diff --git a/Engine/Plugins/AudioInsights/Source/AudioInsights/Private/Cache/AudioInsightsCacheManager.cpp b/Engine/Plugins/AudioInsights/Source/AudioInsights/Private/Cache/AudioInsightsCacheManager.cpp
--- a/Engine/Plugins/AudioInsights/Source/AudioInsights/Private/Cache/AudioInsightsCacheManager.cpp
+++ b/Engine/Plugins/AudioInsights/Source/AudioInsights/Private/Cache/AudioInsightsCacheManager.cpp
@@ -46,6 +46,7 @@ namespace UE::Audio::Insights
 	void FAudioInsightsCacheManager::ClearCache()
 	{
 		NewMessageQueue.DequeueAll();
+		MessageTypeRegistry.Empty();
 
 		CreateCache();
 	}
@@ -170,6 +171,12 @@ namespace UE::Audio::Insights
 		const TArray<TSharedPtr<IAudioCachedMessage>> NewMessages = NewMessageQueue.DequeueAll();
 		for (const TSharedPtr<IAudioCachedMessage>& Message : NewMessages)
 		{
+			const FName MessageName = Message->GetMessageName();
+			if (!MessageTypeRegistry.Contains(MessageName))
+			{
+				MessageTypeRegistry.Emplace(MessageName, Message);
+			}
+
 			if (Chunk->IsChunkFull())
 			{
 				const uint32 NextWriteIndex = Cache.GetNextIndex(WriteIndex);
diff --git a/Engine/Plugins/AudioInsights/Source/AudioInsights/Private/Messages/AudioBusTraceMessages.cpp b/Engine/Plugins/AudioInsights/Source/AudioInsights/Private/Messages/AudioBusTraceMessages.cpp
--- a/Engine/Plugins/AudioInsights/Source/AudioInsights/Private/Messages/AudioBusTraceMessages.cpp
+++ b/Engine/Plugins/AudioInsights/Source/AudioInsights/Private/Messages/AudioBusTraceMessages.cpp
@@ -1,6 +1,9 @@
 // Copyright Epic Games, Inc. All Rights Reserved.
 #include "Messages/AudioBusTraceMessages.h"
 
+#include "Trace/TraceWriter.h"
+#include "TraceWriter/CacheWriteHandler.h"
+
 namespace UE::Audio::Insights
 {
 	namespace AudioBusMessageNames
@@ -45,6 +48,34 @@ namespace UE::Audio::Insights
 		return MemorySize;
 	}
 
+	FCacheWriteHandler FAudioBusStartMessage::GetCacheWriteHandler() const
+	{
+		return {
+			[](UE::Trace::FTraceWriter& Writer) -> uint32
+			{
+				return Writer.DeclareEvent(ANSITEXTVIEW("Audio"), ANSITEXTVIEW("AudioBusStart"))
+					.Field(ANSITEXTVIEW("DeviceId"),    UE::Trace::ETraceWriterFieldType::Uint32)
+					.Field(ANSITEXTVIEW("AudioBusId"),  UE::Trace::ETraceWriterFieldType::Uint32)
+					.Field(ANSITEXTVIEW("Timestamp"),   UE::Trace::ETraceWriterFieldType::Uint64)
+					.Field(ANSITEXTVIEW("Name"),        UE::Trace::ETraceWriterFieldType::WideString)
+					.Field(ANSITEXTVIEW("NumChannels"), UE::Trace::ETraceWriterFieldType::Int32)
+					.End();
+			},
+			[](UE::Trace::FTraceWriter& Writer, uint32 EventId, const IAudioCachedMessage& BaseMsg, uint64 TimestampCycles)
+			{
+				const FAudioBusStartMessage& Msg = static_cast<const FAudioBusStartMessage&>(BaseMsg);
+
+				Writer.WriteEvent(EventId)
+					.Field("DeviceId",    Msg.DeviceId)
+					.Field("AudioBusId",  Msg.AudioBusId)
+					.Field("Timestamp",   TimestampCycles)
+					.Field("Name",        FStringView(Msg.Name))
+					.Field("NumChannels", Msg.NumChannels)
+					.End();
+			}
+		};
+	}
+
 	// FAudioBusHasActivityMessage
 	FAudioBusHasActivityMessage::FAudioBusHasActivityMessage(const Trace::IAnalyzer::FOnEventContext& InContext)
 		: FAudioBusMessageBase(InContext)
@@ -54,6 +85,32 @@ namespace UE::Audio::Insights
 		bHasActivity = EventData.GetValue<bool>("HasActivity");
 	}
 
+	FCacheWriteHandler FAudioBusHasActivityMessage::GetCacheWriteHandler() const
+	{
+		return {
+			[](UE::Trace::FTraceWriter& Writer) -> uint32
+			{
+				return Writer.DeclareEvent(ANSITEXTVIEW("Audio"), ANSITEXTVIEW("AudioBusHasActivity"))
+					.Field(ANSITEXTVIEW("DeviceId"),    UE::Trace::ETraceWriterFieldType::Uint32)
+					.Field(ANSITEXTVIEW("AudioBusId"),  UE::Trace::ETraceWriterFieldType::Uint32)
+					.Field(ANSITEXTVIEW("Timestamp"),   UE::Trace::ETraceWriterFieldType::Uint64)
+					.Field(ANSITEXTVIEW("HasActivity"), UE::Trace::ETraceWriterFieldType::Bool)
+					.End();
+			},
+			[](UE::Trace::FTraceWriter& Writer, uint32 EventId, const IAudioCachedMessage& BaseMsg, uint64 TimestampCycles)
+			{
+				const FAudioBusHasActivityMessage& Msg = static_cast<const FAudioBusHasActivityMessage&>(BaseMsg);
+
+				Writer.WriteEvent(EventId)
+					.Field("DeviceId",    Msg.DeviceId)
+					.Field("AudioBusId",  Msg.AudioBusId)
+					.Field("Timestamp",   TimestampCycles)
+					.Field("HasActivity", Msg.bHasActivity)
+					.End();
+			}
+		};
+	}
+
 	// FAudioBusEnvelopeFollowerEnabledMessage
 	FAudioBusEnvelopeFollowerEnabledMessage::FAudioBusEnvelopeFollowerEnabledMessage(const Trace::IAnalyzer::FOnEventContext& InContext)
 		: FAudioBusMessageBase(InContext)
@@ -63,6 +120,32 @@ namespace UE::Audio::Insights
 		bEnvelopeFollowerEnabled = EventData.GetValue<bool>("EnvelopeFollowerEnabled");
 	}
 
+	FCacheWriteHandler FAudioBusEnvelopeFollowerEnabledMessage::GetCacheWriteHandler() const
+	{
+		return {
+			[](UE::Trace::FTraceWriter& Writer) -> uint32
+			{
+				return Writer.DeclareEvent(ANSITEXTVIEW("Audio"), ANSITEXTVIEW("AudioBusEnvelopeFollowerEnabled"))
+					.Field(ANSITEXTVIEW("DeviceId"),                UE::Trace::ETraceWriterFieldType::Uint32)
+					.Field(ANSITEXTVIEW("AudioBusId"),              UE::Trace::ETraceWriterFieldType::Uint32)
+					.Field(ANSITEXTVIEW("Timestamp"),               UE::Trace::ETraceWriterFieldType::Uint64)
+					.Field(ANSITEXTVIEW("EnvelopeFollowerEnabled"), UE::Trace::ETraceWriterFieldType::Bool)
+					.End();
+			},
+			[](UE::Trace::FTraceWriter& Writer, uint32 EventId, const IAudioCachedMessage& BaseMsg, uint64 TimestampCycles)
+			{
+				const FAudioBusEnvelopeFollowerEnabledMessage& Msg = static_cast<const FAudioBusEnvelopeFollowerEnabledMessage&>(BaseMsg);
+
+				Writer.WriteEvent(EventId)
+					.Field("DeviceId",                Msg.DeviceId)
+					.Field("AudioBusId",              Msg.AudioBusId)
+					.Field("Timestamp",               TimestampCycles)
+					.Field("EnvelopeFollowerEnabled", Msg.bEnvelopeFollowerEnabled)
+					.End();
+			}
+		};
+	}
+
 	// FAudioBusEnvelopeValuesMessage
 	FAudioBusEnvelopeValuesMessage::FAudioBusEnvelopeValuesMessage(const Trace::IAnalyzer::FOnEventContext& InContext)
 		: FAudioBusMessageBase(InContext)
@@ -91,12 +174,62 @@ namespace UE::Audio::Insights
 		return MemorySize;
 	}
 
+	FCacheWriteHandler FAudioBusEnvelopeValuesMessage::GetCacheWriteHandler() const
+	{
+		return {
+			[](UE::Trace::FTraceWriter& Writer) -> uint32
+			{
+				return Writer.DeclareEvent(ANSITEXTVIEW("Audio"), ANSITEXTVIEW("AudioBusEnvelopeValues"))
+					.Field(ANSITEXTVIEW("DeviceId"),       UE::Trace::ETraceWriterFieldType::Uint32)
+					.Field(ANSITEXTVIEW("AudioBusId"),     UE::Trace::ETraceWriterFieldType::Uint32)
+					.Field(ANSITEXTVIEW("Timestamp"),      UE::Trace::ETraceWriterFieldType::Uint64)
+					.Field(ANSITEXTVIEW("EnvelopeValues"), UE::Trace::ETraceWriterFieldType::Float32 | UE::Trace::ETraceWriterFieldType::ArrayFlag)
+					.End();
+			},
+			[](UE::Trace::FTraceWriter& Writer, uint32 EventId, const IAudioCachedMessage& BaseMsg, uint64 TimestampCycles)
+			{
+				const FAudioBusEnvelopeValuesMessage& Msg = static_cast<const FAudioBusEnvelopeValuesMessage&>(BaseMsg);
+
+				Writer.WriteEvent(EventId)
+					.Field("DeviceId",       Msg.DeviceId)
+					.Field("AudioBusId",     Msg.AudioBusId)
+					.Field("Timestamp",      TimestampCycles)
+					.Field("EnvelopeValues", TConstArrayView<float>(Msg.EnvelopeValues.GetData(), Msg.EnvelopeValues.Num()))
+					.End();
+			}
+		};
+	}
+
 	// FAudioBusStopMessage
 	FAudioBusStopMessage::FAudioBusStopMessage(const Trace::IAnalyzer::FOnEventContext& InContext)
 		: FAudioBusMessageBase(InContext)
 	{
 	}
 
+	FCacheWriteHandler FAudioBusStopMessage::GetCacheWriteHandler() const
+	{
+		return {
+			[](UE::Trace::FTraceWriter& Writer) -> uint32
+			{
+				return Writer.DeclareEvent(ANSITEXTVIEW("Audio"), ANSITEXTVIEW("AudioBusStop"))
+					.Field(ANSITEXTVIEW("DeviceId"),   UE::Trace::ETraceWriterFieldType::Uint32)
+					.Field(ANSITEXTVIEW("AudioBusId"), UE::Trace::ETraceWriterFieldType::Uint32)
+					.Field(ANSITEXTVIEW("Timestamp"),  UE::Trace::ETraceWriterFieldType::Uint64)
+					.End();
+			},
+			[](UE::Trace::FTraceWriter& Writer, uint32 EventId, const IAudioCachedMessage& BaseMsg, uint64 TimestampCycles)
+			{
+				const FAudioBusStopMessage& Msg = static_cast<const FAudioBusStopMessage&>(BaseMsg);
+
+				Writer.WriteEvent(EventId)
+					.Field("DeviceId",   Msg.DeviceId)
+					.Field("AudioBusId", Msg.AudioBusId)
+					.Field("Timestamp",  TimestampCycles)
+					.End();
+			}
+		};
+	}
+
 	// FAudioBusDashboardEntry
 	FAudioBusDashboardEntry::FAudioBusDashboardEntry()
 		: AudioMeterInfo(MakeShared<FAudioMeterInfo>())
diff --git a/Engine/Plugins/AudioInsights/Source/AudioInsights/Private/Messages/AudioEventLogTraceMessages.cpp b/Engine/Plugins/AudioInsights/Source/AudioInsights/Private/Messages/AudioEventLogTraceMessages.cpp
--- a/Engine/Plugins/AudioInsights/Source/AudioInsights/Private/Messages/AudioEventLogTraceMessages.cpp
+++ b/Engine/Plugins/AudioInsights/Source/AudioInsights/Private/Messages/AudioEventLogTraceMessages.cpp
@@ -1,6 +1,9 @@
 // Copyright Epic Games, Inc. All Rights Reserved.
 #include "Messages/AudioEventLogTraceMessages.h"
 
+#include "Trace/TraceWriter.h"
+#include "TraceWriter/CacheWriteHandler.h"
+
 namespace UE::Audio::Insights
 {
     FAudioEventLogMessage::FAudioEventLogMessage(const Trace::IAnalyzer::FOnEventContext& InContext)
@@ -66,4 +69,38 @@ namespace UE::Audio::Insights
 		static const FLazyName EventLogName = "EventLog";
 		return EventLogName;
 	}
+
+	FCacheWriteHandler FAudioEventLogMessage::GetCacheWriteHandler() const
+	{
+		return {
+			[](UE::Trace::FTraceWriter& Writer) -> uint32
+			{
+				return Writer.DeclareEvent(ANSITEXTVIEW("Audio"), ANSITEXTVIEW("EventLog"))
+					.Field(ANSITEXTVIEW("DeviceId"),       UE::Trace::ETraceWriterFieldType::Uint32)
+					.Field(ANSITEXTVIEW("Timestamp"),      UE::Trace::ETraceWriterFieldType::Uint64)
+					.Field(ANSITEXTVIEW("PlayOrder"),      UE::Trace::ETraceWriterFieldType::Uint32)
+					.Field(ANSITEXTVIEW("AssetPath"),      UE::Trace::ETraceWriterFieldType::WideString)
+					.Field(ANSITEXTVIEW("EventLogName"),   UE::Trace::ETraceWriterFieldType::WideString)
+					.Field(ANSITEXTVIEW("ActorLabel"),     UE::Trace::ETraceWriterFieldType::WideString)
+					.Field(ANSITEXTVIEW("ActorIconName"),  UE::Trace::ETraceWriterFieldType::WideString)
+					.Field(ANSITEXTVIEW("SoundClassName"), UE::Trace::ETraceWriterFieldType::WideString)
+					.End();
+			},
+			[](UE::Trace::FTraceWriter& Writer, uint32 EventId, const IAudioCachedMessage& BaseMsg, uint64 TimestampCycles)
+			{
+				const FAudioEventLogMessage& Msg = static_cast<const FAudioEventLogMessage&>(BaseMsg);
+
+				Writer.WriteEvent(EventId)
+					.Field("DeviceId",       Msg.DeviceId)
+					.Field("Timestamp",      TimestampCycles)
+					.Field("PlayOrder",      Msg.PlayOrder)
+					.Field("AssetPath",      FStringView(Msg.AssetPath))
+					.Field("EventLogName",   FStringView(Msg.EventName))
+					.Field("ActorLabel",     FStringView(Msg.ActorLabel))
+					.Field("ActorIconName",  FStringView(Msg.ActorIconName))
+					.Field("SoundClassName", FStringView(Msg.CategoryName))
+					.End();
+			}
+		};
+	}
 }
\ No newline at end of file
diff --git a/Engine/Plugins/AudioInsights/Source/AudioInsights/Private/Messages/OutputMeterTraceMessages.cpp b/Engine/Plugins/AudioInsights/Source/AudioInsights/Private/Messages/OutputMeterTraceMessages.cpp
--- a/Engine/Plugins/AudioInsights/Source/AudioInsights/Private/Messages/OutputMeterTraceMessages.cpp
+++ b/Engine/Plugins/AudioInsights/Source/AudioInsights/Private/Messages/OutputMeterTraceMessages.cpp
@@ -1,6 +1,9 @@
 // Copyright Epic Games, Inc. All Rights Reserved.
 #include "Messages/OutputMeterTraceMessages.h"
 
+#include "Trace/TraceWriter.h"
+#include "TraceWriter/CacheWriteHandler.h"
+
 namespace UE::Audio::Insights
 {
 	namespace OutputMeterMessageNames
@@ -32,6 +35,30 @@ namespace UE::Audio::Insights
 		return sizeof(FOutputMeterMainSubmixLoadedMessage);
 	}
 
+	FCacheWriteHandler FOutputMeterMainSubmixLoadedMessage::GetCacheWriteHandler() const
+	{
+		return {
+			[](UE::Trace::FTraceWriter& Writer) -> uint32
+			{
+				return Writer.DeclareEvent(ANSITEXTVIEW("Audio"), ANSITEXTVIEW("MainSubmixLoaded"))
+					.Field(ANSITEXTVIEW("DeviceId"),    UE::Trace::ETraceWriterFieldType::Uint32)
+					.Field(ANSITEXTVIEW("SubmixId"),    UE::Trace::ETraceWriterFieldType::Uint32)
+					.Field(ANSITEXTVIEW("Timestamp"),   UE::Trace::ETraceWriterFieldType::Uint64)
+					.End();
+			},
+			[](UE::Trace::FTraceWriter& Writer, uint32 EventId, const IAudioCachedMessage& BaseMsg, uint64 TimestampCycles)
+			{
+				const FOutputMeterMainSubmixLoadedMessage& Msg = static_cast<const FOutputMeterMainSubmixLoadedMessage&>(BaseMsg);
+
+				Writer.WriteEvent(EventId)
+					.Field("DeviceId",    Msg.DeviceId)
+					.Field("SubmixId",    Msg.SubmixId)
+					.Field("Timestamp",   TimestampCycles)
+					.End();
+			}
+		};
+	}
+
 	// FOutputMeterTruePeakMessage
 	FOutputMeterTruePeakMessage::FOutputMeterTruePeakMessage(const Trace::IAnalyzer::FOnEventContext& InContext)
 		: FOutputMeterMessageBase(InContext)
@@ -47,6 +74,32 @@ namespace UE::Audio::Insights
 		return sizeof(FOutputMeterTruePeakMessage);
 	}
 
+	FCacheWriteHandler FOutputMeterTruePeakMessage::GetCacheWriteHandler() const
+	{
+		return {
+			[](UE::Trace::FTraceWriter& Writer) -> uint32
+			{
+				return Writer.DeclareEvent(ANSITEXTVIEW("Audio"), ANSITEXTVIEW("MainSubmixTruePeakMaxValue"))
+					.Field(ANSITEXTVIEW("DeviceId"),          UE::Trace::ETraceWriterFieldType::Uint32)
+					.Field(ANSITEXTVIEW("SubmixId"),          UE::Trace::ETraceWriterFieldType::Uint32)
+					.Field(ANSITEXTVIEW("Timestamp"),         UE::Trace::ETraceWriterFieldType::Uint64)
+					.Field(ANSITEXTVIEW("TruePeakMaxValueDb"),UE::Trace::ETraceWriterFieldType::Float32)
+					.End();
+			},
+			[](UE::Trace::FTraceWriter& Writer, uint32 EventId, const IAudioCachedMessage& BaseMsg, uint64 TimestampCycles)
+			{
+				const FOutputMeterTruePeakMessage& Msg = static_cast<const FOutputMeterTruePeakMessage&>(BaseMsg);
+
+				Writer.WriteEvent(EventId)
+					.Field("DeviceId",          Msg.DeviceId)
+					.Field("SubmixId",          Msg.SubmixId)
+					.Field("Timestamp",         TimestampCycles)
+					.Field("TruePeakMaxValueDb",Msg.TruePeakMaxValueDb)
+					.End();
+			}
+		};
+	}
+
 	// FOutputMeterLKFSValuesMessage
 	FOutputMeterLKFSValuesMessage::FOutputMeterLKFSValuesMessage(const Trace::IAnalyzer::FOnEventContext& InContext)
 		: FOutputMeterMessageBase(InContext)
@@ -65,4 +118,38 @@ namespace UE::Audio::Insights
 		// Returning fixed size of the struct (no dynamically sized members in it)
 		return sizeof(FOutputMeterLKFSValuesMessage);
 	}
+
+	FCacheWriteHandler FOutputMeterLKFSValuesMessage::GetCacheWriteHandler() const
+	{
+		return {
+			[](UE::Trace::FTraceWriter& Writer) -> uint32
+			{
+				return Writer.DeclareEvent(ANSITEXTVIEW("Audio"), ANSITEXTVIEW("MainSubmixLKFSValues"))
+					.Field(ANSITEXTVIEW("DeviceId"),                UE::Trace::ETraceWriterFieldType::Uint32)
+					.Field(ANSITEXTVIEW("SubmixId"),                UE::Trace::ETraceWriterFieldType::Uint32)
+					.Field(ANSITEXTVIEW("Timestamp"),               UE::Trace::ETraceWriterFieldType::Uint64)
+					.Field(ANSITEXTVIEW("LongTermLoudness"),        UE::Trace::ETraceWriterFieldType::Float32)
+					.Field(ANSITEXTVIEW("ShortTermLoudness"),       UE::Trace::ETraceWriterFieldType::Float32)
+					.Field(ANSITEXTVIEW("MomentaryLoudness"),       UE::Trace::ETraceWriterFieldType::Float32)
+					.Field(ANSITEXTVIEW("LoudnessRangeLowerBound"), UE::Trace::ETraceWriterFieldType::Float32)
+					.Field(ANSITEXTVIEW("LoudnessRangeUpperBound"), UE::Trace::ETraceWriterFieldType::Float32)
+					.End();
+			},
+			[](UE::Trace::FTraceWriter& Writer, uint32 EventId, const IAudioCachedMessage& BaseMsg, uint64 TimestampCycles)
+			{
+				const FOutputMeterLKFSValuesMessage& Msg = static_cast<const FOutputMeterLKFSValuesMessage&>(BaseMsg);
+
+				Writer.WriteEvent(EventId)
+					.Field("DeviceId",                Msg.DeviceId)
+					.Field("SubmixId",                Msg.SubmixId)
+					.Field("Timestamp",               TimestampCycles)
+					.Field("LongTermLoudness",        Msg.LongTermLoudness)
+					.Field("ShortTermLoudness",       Msg.ShortTermLoudness)
+					.Field("MomentaryLoudness",       Msg.MomentaryLoudness)
+					.Field("LoudnessRangeLowerBound", Msg.LoudnessRangeLowerBound)
+					.Field("LoudnessRangeUpperBound", Msg.LoudnessRangeUpperBound)
+					.End();
+			}
+		};
+	}
 } // namespace UE::Audio::Insights
diff --git a/Engine/Plugins/AudioInsights/Source/AudioInsights/Private/Messages/SignalFlowTraceMessages.cpp b/Engine/Plugins/AudioInsights/Source/AudioInsights/Private/Messages/SignalFlowTraceMessages.cpp
--- a/Engine/Plugins/AudioInsights/Source/AudioInsights/Private/Messages/SignalFlowTraceMessages.cpp
+++ b/Engine/Plugins/AudioInsights/Source/AudioInsights/Private/Messages/SignalFlowTraceMessages.cpp
@@ -1,6 +1,9 @@
 // Copyright Epic Games, Inc. All Rights Reserved.
 #include "Messages/SignalFlowTraceMessages.h"
 
+#include "Trace/TraceWriter.h"
+#include "TraceWriter/CacheWriteHandler.h"
+
 namespace UE::Audio::Insights
 {
 	namespace SignalFlowMessageNames
@@ -15,6 +18,50 @@ namespace UE::Audio::Insights
 		constexpr uint32 DataPointsCapacity = 128u;
 	}
 
+	namespace FSignalFlowMessagePrivate
+	{
+		FCacheWriteHandler MakeSendWriteHandler(FAnsiStringView InEventName)
+		{
+			FAnsiString EventNameStr(InEventName);
+			return {
+				[EventNameStr](UE::Trace::FTraceWriter& Writer) -> uint32
+				{
+					return Writer.DeclareEvent(ANSITEXTVIEW("Audio"), FAnsiStringView(*EventNameStr))
+						.Field(ANSITEXTVIEW("DeviceId"),    UE::Trace::ETraceWriterFieldType::Uint32)
+						.Field(ANSITEXTVIEW("Timestamp"),   UE::Trace::ETraceWriterFieldType::Uint64)
+						.Field(ANSITEXTVIEW("SenderID"),    UE::Trace::ETraceWriterFieldType::Uint32)
+						.Field(ANSITEXTVIEW("ReceiverIDs"), UE::Trace::ETraceWriterFieldType::Uint32 | UE::Trace::ETraceWriterFieldType::ArrayFlag)
+						.Field(ANSITEXTVIEW("SendLevels"),  UE::Trace::ETraceWriterFieldType::Float32 | UE::Trace::ETraceWriterFieldType::ArrayFlag)
+						.End();
+				},
+				[](UE::Trace::FTraceWriter& Writer, uint32 EventId, const IAudioCachedMessage& BaseMsg, uint64 TimestampCycles)
+				{
+					const FSignalFlowSendMessageBase& Msg = static_cast<const FSignalFlowSendMessageBase&>(BaseMsg);
+
+					TArray<uint32> ReceiverIDs;
+					TArray<float> SendLevels;
+
+					ReceiverIDs.Reserve(Msg.SendInfos.Num());
+					SendLevels.Reserve(Msg.SendInfos.Num());
+
+					for (const FSignalFlowSendMessageBase::FSendInfoPair& Info : Msg.SendInfos)
+					{
+						ReceiverIDs.Add(Info.ReceiverID);
+						SendLevels.Add(Info.SendLevel);
+					}
+
+					Writer.WriteEvent(EventId)
+						.Field("DeviceId",    Msg.DeviceId)
+						.Field("Timestamp",   TimestampCycles)
+						.Field("SenderID",    Msg.SenderID)
+						.Field("ReceiverIDs", TConstArrayView<uint32>(ReceiverIDs.GetData(), ReceiverIDs.Num()))
+						.Field("SendLevels",  TConstArrayView<float>(SendLevels.GetData(), SendLevels.Num()))
+						.End();
+				}
+			};
+		}
+	}
+
 	FSignalFlowMessageBase::FSignalFlowMessageBase(const Trace::IAnalyzer::FOnEventContext& InContext)
 	{
 		const Trace::IAnalyzer::FEventData& EventData = InContext.EventData;
@@ -64,6 +111,46 @@ namespace UE::Audio::Insights
 		return sizeof(FSubmixToSubmixSendMessage);
 	}
 
+	FCacheWriteHandler FSoundToBusSendMessage::GetCacheWriteHandler() const
+	{
+		return FSignalFlowMessagePrivate::MakeSendWriteHandler("SoundToBusSend");
+	}
+
+	FCacheWriteHandler FSoundToSubmixSendMessage::GetCacheWriteHandler() const
+	{
+		return FSignalFlowMessagePrivate::MakeSendWriteHandler("SoundToSubmixSend");
+	}
+
+	FCacheWriteHandler FSubmixToSubmixSendMessage::GetCacheWriteHandler() const
+	{
+		return {
+			[](UE::Trace::FTraceWriter& Writer) -> uint32
+			{
+				return Writer.DeclareEvent(ANSITEXTVIEW("Audio"), ANSITEXTVIEW("SubmixToSubmixSend"))
+					.Field(ANSITEXTVIEW("DeviceId"),     UE::Trace::ETraceWriterFieldType::Uint32)
+					.Field(ANSITEXTVIEW("Timestamp"),    UE::Trace::ETraceWriterFieldType::Uint64)
+					.Field(ANSITEXTVIEW("SenderID"),     UE::Trace::ETraceWriterFieldType::Uint32)
+					.Field(ANSITEXTVIEW("ReceiverID"),   UE::Trace::ETraceWriterFieldType::Uint32)
+					.Field(ANSITEXTVIEW("SendLevel"),    UE::Trace::ETraceWriterFieldType::Float32)
+					.Field(ANSITEXTVIEW("IsFinalStage"), UE::Trace::ETraceWriterFieldType::Bool)
+					.End();
+			},
+			[](UE::Trace::FTraceWriter& Writer, uint32 EventId, const IAudioCachedMessage& BaseMsg, uint64 TimestampCycles)
+			{
+				const FSubmixToSubmixSendMessage& Msg = static_cast<const FSubmixToSubmixSendMessage&>(BaseMsg);
+
+				Writer.WriteEvent(EventId)
+					.Field("DeviceId",     Msg.DeviceId)
+					.Field("Timestamp",    TimestampCycles)
+					.Field("SenderID",     Msg.SenderID)
+					.Field("ReceiverID",   Msg.ReceiverID)
+					.Field("SendLevel",    Msg.SendLevel)
+					.Field("IsFinalStage", Msg.bIsFinalStage)
+					.End();
+			}
+		};
+	}
+
 	FSignalFlowDashboardEntry::FSignalFlowDashboardEntry()
 		: AmplitudeDataRange(FSignalFlowDashboardEntryPrivate::DataPointsCapacity, 0.0f)
 	{
diff --git a/Engine/Plugins/AudioInsights/Source/AudioInsights/Private/Messages/SoundTraceMessages.cpp b/Engine/Plugins/AudioInsights/Source/AudioInsights/Private/Messages/SoundTraceMessages.cpp
--- a/Engine/Plugins/AudioInsights/Source/AudioInsights/Private/Messages/SoundTraceMessages.cpp
+++ b/Engine/Plugins/AudioInsights/Source/AudioInsights/Private/Messages/SoundTraceMessages.cpp
@@ -1,6 +1,9 @@
 // Copyright Epic Games, Inc. All Rights Reserved.
 #include "Messages/SoundTraceMessages.h"
 
+#include "Trace/TraceWriter.h"
+#include "TraceWriter/CacheWriteHandler.h"
+
 namespace UE::Audio::Insights
 {
 	namespace SoundMessageNames
@@ -39,6 +42,53 @@ namespace UE::Audio::Insights
 		}
 	};
 
+	namespace FSoundMessagePrivate
+	{
+		FStringView GetSoundClassNameFromEntryType(ESoundDashboardEntryType InEntryType)
+		{
+			switch (InEntryType)
+			{
+				case ESoundDashboardEntryType::SoundCue:
+					return SoundClassNames::SoundCue;
+				case ESoundDashboardEntryType::SoundWave:
+					return SoundClassNames::SoundWave;
+				case ESoundDashboardEntryType::MetaSound:
+					return SoundClassNames::MetaSoundSource;
+				case ESoundDashboardEntryType::ProceduralSource:
+					return SoundClassNames::SoundWaveProcedural;
+				case ESoundDashboardEntryType::SoundCueTemplate:
+					return SoundClassNames::SoundCueTemplate;
+				default:
+					break;
+			}
+
+			return TEXTVIEW("");
+		}
+
+		uint32 DeclareParamEvent(UE::Trace::FTraceWriter& Writer, FAnsiStringView InEventName, FAnsiStringView InParamFieldName)
+		{
+			return Writer.DeclareEvent(ANSITEXTVIEW("Audio"), FAnsiStringView(InEventName))
+				.Field(ANSITEXTVIEW("DeviceId"),             UE::Trace::ETraceWriterFieldType::Uint32)
+				.Field(ANSITEXTVIEW("Timestamp"),            UE::Trace::ETraceWriterFieldType::Uint64)
+				.Field(ANSITEXTVIEW("PlayOrder"),            UE::Trace::ETraceWriterFieldType::Uint32)
+				.Field(ANSITEXTVIEW("ActiveSoundPlayOrder"), UE::Trace::ETraceWriterFieldType::Uint32)
+				.Field(FAnsiStringView(InParamFieldName),    UE::Trace::ETraceWriterFieldType::Float32)
+				.End();
+		}
+
+		uint32 DeclareSoundSourceFiltersEvent(UE::Trace::FTraceWriter& Writer)
+		{
+			return Writer.DeclareEvent(ANSITEXTVIEW("Audio"), ANSITEXTVIEW("SoundSourceFilters"))
+				.Field(ANSITEXTVIEW("DeviceId"),             UE::Trace::ETraceWriterFieldType::Uint32)
+				.Field(ANSITEXTVIEW("Timestamp"),            UE::Trace::ETraceWriterFieldType::Uint64)
+				.Field(ANSITEXTVIEW("PlayOrder"),            UE::Trace::ETraceWriterFieldType::Uint32)
+				.Field(ANSITEXTVIEW("ActiveSoundPlayOrder"), UE::Trace::ETraceWriterFieldType::Uint32)
+				.Field(ANSITEXTVIEW("LPFFrequency"),         UE::Trace::ETraceWriterFieldType::Float32)
+				.Field(ANSITEXTVIEW("HPFFrequency"),         UE::Trace::ETraceWriterFieldType::Float32)
+				.End();
+		}
+	};
+
 	FSoundMessageBase::FSoundMessageBase(const Trace::IAnalyzer::FOnEventContext& InContext)
 	{
 		const Trace::IAnalyzer::FEventData& EventData = InContext.EventData;
@@ -116,6 +166,277 @@ namespace UE::Audio::Insights
 		return sizeof(FSoundWaveStopMessage);
 	}
 
+	FCacheWriteHandler FSoundStartMessage::GetCacheWriteHandler() const
+	{
+		return {
+			[](UE::Trace::FTraceWriter& Writer) -> uint32
+			{
+				return Writer.DeclareEvent(ANSITEXTVIEW("Audio"), ANSITEXTVIEW("SoundStart"))
+					.Field(ANSITEXTVIEW("DeviceId"),       UE::Trace::ETraceWriterFieldType::Uint32)
+					.Field(ANSITEXTVIEW("Timestamp"),      UE::Trace::ETraceWriterFieldType::Uint64)
+					.Field(ANSITEXTVIEW("PlayOrder"),      UE::Trace::ETraceWriterFieldType::Uint32)
+					.Field(ANSITEXTVIEW("AssetPath"),      UE::Trace::ETraceWriterFieldType::WideString)
+					.Field(ANSITEXTVIEW("EventLogName"),   UE::Trace::ETraceWriterFieldType::WideString)
+					.Field(ANSITEXTVIEW("SoundClassName"), UE::Trace::ETraceWriterFieldType::WideString)
+					.Field(ANSITEXTVIEW("ActorLabel"),     UE::Trace::ETraceWriterFieldType::WideString)
+					.Field(ANSITEXTVIEW("ActorIconName"),  UE::Trace::ETraceWriterFieldType::WideString)
+					.End();
+			},
+			[](UE::Trace::FTraceWriter& Writer, uint32 EventId, const IAudioCachedMessage& BaseMsg, uint64 TimestampCycles)
+			{
+				const FSoundStartMessage& Msg = static_cast<const FSoundStartMessage&>(BaseMsg);
+
+				Writer.WriteEvent(EventId)
+					.Field("DeviceId",       Msg.DeviceId)
+					.Field("Timestamp",      TimestampCycles)
+					.Field("PlayOrder",      Msg.ActiveSoundPlayOrder)
+					.Field("AssetPath",      FStringView(Msg.Name))
+					.Field("EventLogName",   TEXTVIEW(""))
+					.Field("SoundClassName", FSoundMessagePrivate::GetSoundClassNameFromEntryType(Msg.EntryType))
+					.Field("ActorLabel",     FStringView(Msg.ActorLabel))
+					.Field("ActorIconName",  TEXTVIEW(""))
+					.End();
+			}
+		};
+	}
+
+	FCacheWriteHandler FSoundIsAlivePingMessage::GetCacheWriteHandler() const
+	{
+		return {
+			[](UE::Trace::FTraceWriter& Writer) -> uint32
+			{
+				return Writer.DeclareEvent(ANSITEXTVIEW("Audio"), ANSITEXTVIEW("SoundIsAlivePing"))
+					.Field(ANSITEXTVIEW("DeviceId"),       UE::Trace::ETraceWriterFieldType::Uint32)
+					.Field(ANSITEXTVIEW("Timestamp"),      UE::Trace::ETraceWriterFieldType::Uint64)
+					.Field(ANSITEXTVIEW("PlayOrder"),      UE::Trace::ETraceWriterFieldType::Uint32)
+					.Field(ANSITEXTVIEW("AssetPath"),      UE::Trace::ETraceWriterFieldType::WideString)
+					.Field(ANSITEXTVIEW("SoundClassName"), UE::Trace::ETraceWriterFieldType::WideString)
+					.Field(ANSITEXTVIEW("ActorLabel"),     UE::Trace::ETraceWriterFieldType::WideString)
+					.End();
+			},
+			[](UE::Trace::FTraceWriter& Writer, uint32 EventId, const IAudioCachedMessage& BaseMsg, uint64 TimestampCycles)
+			{
+				const FSoundIsAlivePingMessage& Msg = static_cast<const FSoundIsAlivePingMessage&>(BaseMsg);
+
+				Writer.WriteEvent(EventId)
+					.Field("DeviceId",       Msg.DeviceId)
+					.Field("Timestamp",      TimestampCycles)
+					.Field("PlayOrder",      Msg.ActiveSoundPlayOrder)
+					.Field("AssetPath",      FStringView(Msg.Name))
+					.Field("SoundClassName", FSoundMessagePrivate::GetSoundClassNameFromEntryType(Msg.EntryType))
+					.Field("ActorLabel",     FStringView(Msg.ActorLabel))
+					.End();
+			}
+		};
+	}
+
+	FCacheWriteHandler FSoundWaveStartMessage::GetCacheWriteHandler() const
+	{
+		return {
+			[](UE::Trace::FTraceWriter& Writer) -> uint32
+			{
+				return Writer.DeclareEvent(ANSITEXTVIEW("Audio"), ANSITEXTVIEW("SoundWaveStart"))
+					.Field(ANSITEXTVIEW("DeviceId"),             UE::Trace::ETraceWriterFieldType::Uint32)
+					.Field(ANSITEXTVIEW("Timestamp"),            UE::Trace::ETraceWriterFieldType::Uint64)
+					.Field(ANSITEXTVIEW("PlayOrder"),            UE::Trace::ETraceWriterFieldType::Uint32)
+					.Field(ANSITEXTVIEW("ActiveSoundPlayOrder"), UE::Trace::ETraceWriterFieldType::Uint32)
+					.Field(ANSITEXTVIEW("AssetPath"),            UE::Trace::ETraceWriterFieldType::WideString)
+					.Field(ANSITEXTVIEW("SoundClassName"),       UE::Trace::ETraceWriterFieldType::WideString)
+					.Field(ANSITEXTVIEW("ActorLabel"),           UE::Trace::ETraceWriterFieldType::WideString)
+					.Field(ANSITEXTVIEW("ActorIconName"),        UE::Trace::ETraceWriterFieldType::WideString)
+					.Field(ANSITEXTVIEW("SourceBusID"),          UE::Trace::ETraceWriterFieldType::Int32)
+					.End();
+			},
+			[](UE::Trace::FTraceWriter& Writer, uint32 EventId, const IAudioCachedMessage& BaseMsg, uint64 TimestampCycles)
+			{
+				const FSoundWaveStartMessage& Msg = static_cast<const FSoundWaveStartMessage&>(BaseMsg);
+
+				Writer.WriteEvent(EventId)
+					.Field("DeviceId",             Msg.DeviceId)
+					.Field("Timestamp",            TimestampCycles)
+					.Field("PlayOrder",            Msg.WaveInstancePlayOrder)
+					.Field("ActiveSoundPlayOrder", Msg.ActiveSoundPlayOrder)
+					.Field("AssetPath",            FStringView(Msg.Name))
+					.Field("SoundClassName",       FSoundMessagePrivate::GetSoundClassNameFromEntryType(Msg.EntryType))
+					.Field("ActorLabel",           FStringView(Msg.ActorLabel))
+					.Field("ActorIconName",        FStringView(Msg.ActorIconName))
+					.Field("SourceBusID",          Msg.SourceBusID)
+					.End();
+			}
+		};
+	}
+
+	FCacheWriteHandler FSoundWaveIsAlivePingMessage::GetCacheWriteHandler() const
+	{
+		return {
+			[](UE::Trace::FTraceWriter& Writer) -> uint32
+			{
+				return Writer.DeclareEvent(ANSITEXTVIEW("Audio"), ANSITEXTVIEW("SoundWaveIsAlivePing"))
+					.Field(ANSITEXTVIEW("DeviceId"),             UE::Trace::ETraceWriterFieldType::Uint32)
+					.Field(ANSITEXTVIEW("Timestamp"),            UE::Trace::ETraceWriterFieldType::Uint64)
+					.Field(ANSITEXTVIEW("PlayOrder"),            UE::Trace::ETraceWriterFieldType::Uint32)
+					.Field(ANSITEXTVIEW("ActiveSoundPlayOrder"), UE::Trace::ETraceWriterFieldType::Uint32)
+					.Field(ANSITEXTVIEW("AssetPath"),            UE::Trace::ETraceWriterFieldType::WideString)
+					.Field(ANSITEXTVIEW("SoundClassName"),       UE::Trace::ETraceWriterFieldType::WideString)
+					.Field(ANSITEXTVIEW("ActorLabel"),           UE::Trace::ETraceWriterFieldType::WideString)
+					.Field(ANSITEXTVIEW("ActorIconName"),        UE::Trace::ETraceWriterFieldType::WideString)
+					.Field(ANSITEXTVIEW("SourceBusID"),          UE::Trace::ETraceWriterFieldType::Int32)
+					.End();
+			},
+			[](UE::Trace::FTraceWriter& Writer, uint32 EventId, const IAudioCachedMessage& BaseMsg, uint64 TimestampCycles)
+			{
+				const FSoundWaveIsAlivePingMessage& Msg = static_cast<const FSoundWaveIsAlivePingMessage&>(BaseMsg);
+
+				Writer.WriteEvent(EventId)
+					.Field("DeviceId",             Msg.DeviceId)
+					.Field("Timestamp",            TimestampCycles)
+					.Field("PlayOrder",            Msg.WaveInstancePlayOrder)
+					.Field("ActiveSoundPlayOrder", Msg.ActiveSoundPlayOrder)
+					.Field("AssetPath",            FStringView(Msg.Name))
+					.Field("SoundClassName",       FSoundMessagePrivate::GetSoundClassNameFromEntryType(Msg.EntryType))
+					.Field("ActorLabel",           FStringView(Msg.ActorLabel))
+					.Field("ActorIconName",        FStringView(Msg.ActorIconName))
+					.Field("SourceBusID",          Msg.SourceBusID)
+					.End();
+			}
+		};
+	}
+
+	FCacheWriteHandler FSoundStopMessage::GetCacheWriteHandler() const
+	{
+		return {
+			[](UE::Trace::FTraceWriter& Writer) -> uint32
+			{
+				return Writer.DeclareEvent(ANSITEXTVIEW("Audio"), ANSITEXTVIEW("SoundStop"))
+					.Field(ANSITEXTVIEW("DeviceId"),       UE::Trace::ETraceWriterFieldType::Uint32)
+					.Field(ANSITEXTVIEW("Timestamp"),      UE::Trace::ETraceWriterFieldType::Uint64)
+					.Field(ANSITEXTVIEW("PlayOrder"),      UE::Trace::ETraceWriterFieldType::Uint32)
+					.Field(ANSITEXTVIEW("AssetPath"),      UE::Trace::ETraceWriterFieldType::WideString)
+					.Field(ANSITEXTVIEW("EventLogName"),   UE::Trace::ETraceWriterFieldType::WideString)
+					.Field(ANSITEXTVIEW("SoundClassName"), UE::Trace::ETraceWriterFieldType::WideString)
+					.Field(ANSITEXTVIEW("ActorLabel"),     UE::Trace::ETraceWriterFieldType::WideString)
+					.Field(ANSITEXTVIEW("ActorIconName"),  UE::Trace::ETraceWriterFieldType::WideString)
+					.End();
+			},
+			[](UE::Trace::FTraceWriter& Writer, uint32 EventId, const IAudioCachedMessage& BaseMsg, uint64 TimestampCycles)
+			{
+				const FSoundStopMessage& Msg = static_cast<const FSoundStopMessage&>(BaseMsg);
+
+				Writer.WriteEvent(EventId)
+					.Field("DeviceId",       Msg.DeviceId)
+					.Field("Timestamp",      TimestampCycles)
+					.Field("PlayOrder",      Msg.ActiveSoundPlayOrder)
+					.Field("AssetPath",      TEXTVIEW(""))
+					.Field("EventLogName",   TEXTVIEW(""))
+					.Field("SoundClassName", TEXTVIEW(""))
+					.Field("ActorLabel",     TEXTVIEW(""))
+					.Field("ActorIconName",  TEXTVIEW(""))
+					.End();
+			}
+		};
+	}
+
+	FCacheWriteHandler FSoundWaveStopMessage::GetCacheWriteHandler() const
+	{
+		return {
+			[](UE::Trace::FTraceWriter& Writer) -> uint32
+			{
+				return Writer.DeclareEvent(ANSITEXTVIEW("Audio"), ANSITEXTVIEW("SoundWaveStop"))
+					.Field(ANSITEXTVIEW("DeviceId"),             UE::Trace::ETraceWriterFieldType::Uint32)
+					.Field(ANSITEXTVIEW("Timestamp"),            UE::Trace::ETraceWriterFieldType::Uint64)
+					.Field(ANSITEXTVIEW("PlayOrder"),            UE::Trace::ETraceWriterFieldType::Uint32)
+					.Field(ANSITEXTVIEW("ActiveSoundPlayOrder"), UE::Trace::ETraceWriterFieldType::Uint32)
+					.End();
+			},
+			[](UE::Trace::FTraceWriter& Writer, uint32 EventId, const IAudioCachedMessage& BaseMsg, uint64 TimestampCycles)
+			{
+				const FSoundWaveStopMessage& Msg = static_cast<const FSoundWaveStopMessage&>(BaseMsg);
+
+				Writer.WriteEvent(EventId)
+					.Field("DeviceId",             Msg.DeviceId)
+					.Field("Timestamp",            TimestampCycles)
+					.Field("PlayOrder",            Msg.WaveInstancePlayOrder)
+					.Field("ActiveSoundPlayOrder", Msg.ActiveSoundPlayOrder)
+					.End();
+			}
+		};
+	}
+
+	FCacheWriteHandler FSoundLPFFreqMessage::GetCacheWriteHandler() const
+	{
+		return {
+			[](UE::Trace::FTraceWriter& Writer) -> uint32 { return FSoundMessagePrivate::DeclareSoundSourceFiltersEvent(Writer); },
+			[](UE::Trace::FTraceWriter& Writer, uint32 EventId, const IAudioCachedMessage& BaseMsg, uint64 TimestampCycles)
+			{
+				const FSoundLPFFreqMessage& Msg = static_cast<const FSoundLPFFreqMessage&>(BaseMsg);
+
+				Writer.WriteEvent(EventId)
+					.Field("DeviceId",             Msg.DeviceId)
+					.Field("Timestamp",            TimestampCycles)
+					.Field("PlayOrder",            Msg.WaveInstancePlayOrder)
+					.Field("ActiveSoundPlayOrder", Msg.ActiveSoundPlayOrder)
+					.Field("LPFFrequency",         Msg.LPFFrequency)
+					.Field("HPFFrequency",         MIN_FILTER_FREQUENCY)
+					.End();
+			}
+		};
+	}
+
+	FCacheWriteHandler FSoundHPFFreqMessage::GetCacheWriteHandler() const
+	{
+		return {
+			[](UE::Trace::FTraceWriter& Writer) -> uint32 { return FSoundMessagePrivate::DeclareSoundSourceFiltersEvent(Writer); },
+			[](UE::Trace::FTraceWriter& Writer, uint32 EventId, const IAudioCachedMessage& BaseMsg, uint64 TimestampCycles)
+			{
+				const FSoundHPFFreqMessage& Msg = static_cast<const FSoundHPFFreqMessage&>(BaseMsg);
+
+				Writer.WriteEvent(EventId)
+					.Field("DeviceId",             Msg.DeviceId)
+					.Field("Timestamp",            TimestampCycles)
+					.Field("PlayOrder",            Msg.WaveInstancePlayOrder)
+					.Field("ActiveSoundPlayOrder", Msg.ActiveSoundPlayOrder)
+					.Field("LPFFrequency",         MAX_FILTER_FREQUENCY)
+					.Field("HPFFrequency",         Msg.HPFFrequency)
+					.End();
+			}
+		};
+	}
+
+	// Macro-generated param message implementations use a shared helper
+	// DEFINE_SOUND_PARAM_MESSAGE generates the struct but GetCacheWriteHandler must be defined per-type
+#define IMPLEMENT_PARAM_GET_CACHE_WRITE_HANDLER(ClassName, TraceEventName, ParamFieldName, ParamMember)						\
+	FCacheWriteHandler ClassName::GetCacheWriteHandler() const																\
+	{																														\
+		return {																											\
+			[](UE::Trace::FTraceWriter& Writer) -> uint32																	\
+			{																												\
+				FAnsiString EventNameStr(TraceEventName);																	\
+				FAnsiString FieldNameStr(ParamFieldName);																	\
+				return FSoundMessagePrivate::DeclareParamEvent(Writer, EventNameStr, FieldNameStr);							\
+			},																												\
+			[](UE::Trace::FTraceWriter& Writer, uint32 EventId, const IAudioCachedMessage& BaseMsg, uint64 TimestampCycles)							\
+			{																												\
+				const FSoundParameterMessage& Msg = static_cast<const FSoundParameterMessage&>(BaseMsg);					\
+				Writer.WriteEvent(EventId)																					\
+					.Field("DeviceId",             Msg.DeviceId)															\
+					.Field("Timestamp",            TimestampCycles)			\
+					.Field("PlayOrder",            Msg.WaveInstancePlayOrder)												\
+					.Field("ActiveSoundPlayOrder", Msg.ActiveSoundPlayOrder)												\
+					.Field(4, static_cast<const ClassName&>(BaseMsg).ParamMember)											\
+					.End();																									\
+			}																												\
+		};																													\
+	}
+
+	IMPLEMENT_PARAM_GET_CACHE_WRITE_HANDLER(FSoundPriorityMessage, "SoundPriority", "Priority", Priority)
+	IMPLEMENT_PARAM_GET_CACHE_WRITE_HANDLER(FSoundDistanceMessage, "SoundDistance", "Distance", Distance)
+	IMPLEMENT_PARAM_GET_CACHE_WRITE_HANDLER(FSoundDistanceAttenuationMessage, "SoundDistanceAttenuation", "DistanceAttenuation", DistanceAttenuation)
+	IMPLEMENT_PARAM_GET_CACHE_WRITE_HANDLER(FSoundVolumeMessage, "SoundVolume", "Volume", Volume)
+	IMPLEMENT_PARAM_GET_CACHE_WRITE_HANDLER(FSoundPitchMessage, "SoundPitch", "Pitch", Pitch)
+	IMPLEMENT_PARAM_GET_CACHE_WRITE_HANDLER(FSoundEnvelopeMessage, "SoundSourceEnvelope", "Envelope", Envelope)
+	IMPLEMENT_PARAM_GET_CACHE_WRITE_HANDLER(FSoundRelativeRenderCostMessage, "SoundRelativeRenderCost", "RelativeRenderCost", RelativeRenderCost)
+#undef IMPLEMENT_PARAM_GET_CACHE_WRITE_HANDLER
+
 	FSoundDashboardEntry::FSoundDashboardEntry()
 	{
 		constexpr uint32 DataPointsCapacity = 256u;
diff --git a/Engine/Plugins/AudioInsights/Source/AudioInsights/Private/Messages/SubmixTraceMessages.cpp b/Engine/Plugins/AudioInsights/Source/AudioInsights/Private/Messages/SubmixTraceMessages.cpp
--- a/Engine/Plugins/AudioInsights/Source/AudioInsights/Private/Messages/SubmixTraceMessages.cpp
+++ b/Engine/Plugins/AudioInsights/Source/AudioInsights/Private/Messages/SubmixTraceMessages.cpp
@@ -1,6 +1,9 @@
 // Copyright Epic Games, Inc. All Rights Reserved.
 #include "Messages/SubmixTraceMessages.h"
 
+#include "Trace/TraceWriter.h"
+#include "TraceWriter/CacheWriteHandler.h"
+
 namespace UE::Audio::Insights
 {
 	namespace SubmixMessageNames
@@ -45,6 +48,66 @@ namespace UE::Audio::Insights
 		return MemorySize;
 	};
 
+	FCacheWriteHandler FSubmixLoadedMessage::GetCacheWriteHandler() const
+	{
+		return {
+			[](UE::Trace::FTraceWriter& Writer) -> uint32
+			{
+				return Writer.DeclareEvent(ANSITEXTVIEW("Audio"), ANSITEXTVIEW("SubmixLoaded"))
+					.Field(ANSITEXTVIEW("DeviceId"),     UE::Trace::ETraceWriterFieldType::Uint32)
+					.Field(ANSITEXTVIEW("SubmixId"),     UE::Trace::ETraceWriterFieldType::Uint32)
+					.Field(ANSITEXTVIEW("Timestamp"),    UE::Trace::ETraceWriterFieldType::Uint64)
+					.Field(ANSITEXTVIEW("Name"),         UE::Trace::ETraceWriterFieldType::WideString)
+					.Field(ANSITEXTVIEW("NumChannels"),  UE::Trace::ETraceWriterFieldType::Int32)
+					.Field(ANSITEXTVIEW("IsMainSubmix"), UE::Trace::ETraceWriterFieldType::Bool)
+					.End();
+			},
+			[](UE::Trace::FTraceWriter& Writer, uint32 EventId, const IAudioCachedMessage& BaseMsg, uint64 TimestampCycles)
+			{
+				const FSubmixLoadedMessage& Msg = static_cast<const FSubmixLoadedMessage&>(BaseMsg);
+
+				Writer.WriteEvent(EventId)
+					.Field("DeviceId",     Msg.DeviceId)
+					.Field("SubmixId",     Msg.SubmixId)
+					.Field("Timestamp",    TimestampCycles)
+					.Field("Name",         FStringView(Msg.Name))
+					.Field("NumChannels",  Msg.NumChannels)
+					.Field("IsMainSubmix", Msg.bIsMainSubmix)
+					.End();
+			}
+		};
+	}
+
+	FCacheWriteHandler FSubmixAlivePingMessage::GetCacheWriteHandler() const
+	{
+		return {
+			[](UE::Trace::FTraceWriter& Writer) -> uint32
+			{
+				return Writer.DeclareEvent(ANSITEXTVIEW("Audio"), ANSITEXTVIEW("SubmixAlivePing"))
+					.Field(ANSITEXTVIEW("DeviceId"),     UE::Trace::ETraceWriterFieldType::Uint32)
+					.Field(ANSITEXTVIEW("SubmixId"),     UE::Trace::ETraceWriterFieldType::Uint32)
+					.Field(ANSITEXTVIEW("Timestamp"),    UE::Trace::ETraceWriterFieldType::Uint64)
+					.Field(ANSITEXTVIEW("Name"),         UE::Trace::ETraceWriterFieldType::WideString)
+					.Field(ANSITEXTVIEW("NumChannels"),  UE::Trace::ETraceWriterFieldType::Int32)
+					.Field(ANSITEXTVIEW("IsMainSubmix"), UE::Trace::ETraceWriterFieldType::Bool)
+					.End();
+			},
+			[](UE::Trace::FTraceWriter& Writer, uint32 EventId, const IAudioCachedMessage& BaseMsg, uint64 TimestampCycles)
+			{
+				const FSubmixAlivePingMessage& Msg = static_cast<const FSubmixAlivePingMessage&>(BaseMsg);
+
+				Writer.WriteEvent(EventId)
+					.Field("DeviceId",     Msg.DeviceId)
+					.Field("SubmixId",     Msg.SubmixId)
+					.Field("Timestamp",    TimestampCycles)
+					.Field("Name",         FStringView(Msg.Name))
+					.Field("NumChannels",  Msg.NumChannels)
+					.Field("IsMainSubmix", Msg.bIsMainSubmix)
+					.End();
+			}
+		};
+	}
+
 	// FSubmixHasActivityMessage
 	FSubmixHasActivityMessage::FSubmixHasActivityMessage(const Trace::IAnalyzer::FOnEventContext& InContext)
 		: FSubmixMessageBase(InContext)
@@ -54,6 +117,32 @@ namespace UE::Audio::Insights
 		bHasActivity = EventData.GetValue<bool>("HasActivity");
 	}
 
+	FCacheWriteHandler FSubmixHasActivityMessage::GetCacheWriteHandler() const
+	{
+		return {
+			[](UE::Trace::FTraceWriter& Writer) -> uint32
+			{
+				return Writer.DeclareEvent(ANSITEXTVIEW("Audio"), ANSITEXTVIEW("SubmixHasActivity"))
+					.Field(ANSITEXTVIEW("DeviceId"),    UE::Trace::ETraceWriterFieldType::Uint32)
+					.Field(ANSITEXTVIEW("SubmixId"),    UE::Trace::ETraceWriterFieldType::Uint32)
+					.Field(ANSITEXTVIEW("Timestamp"),   UE::Trace::ETraceWriterFieldType::Uint64)
+					.Field(ANSITEXTVIEW("HasActivity"), UE::Trace::ETraceWriterFieldType::Bool)
+					.End();
+			},
+			[](UE::Trace::FTraceWriter& Writer, uint32 EventId, const IAudioCachedMessage& BaseMsg, uint64 TimestampCycles)
+			{
+				const FSubmixHasActivityMessage& Msg = static_cast<const FSubmixHasActivityMessage&>(BaseMsg);
+
+				Writer.WriteEvent(EventId)
+					.Field("DeviceId",    Msg.DeviceId)
+					.Field("SubmixId",    Msg.SubmixId)
+					.Field("Timestamp",   TimestampCycles)
+					.Field("HasActivity", Msg.bHasActivity)
+					.End();
+			}
+		};
+	}
+
 	// FSubmixEnvelopeFollowerEnabledMessage
 	FSubmixEnvelopeFollowerEnabledMessage::FSubmixEnvelopeFollowerEnabledMessage(const Trace::IAnalyzer::FOnEventContext& InContext)
 		: FSubmixMessageBase(InContext)
@@ -63,6 +152,32 @@ namespace UE::Audio::Insights
 		bEnvelopeFollowerEnabled = EventData.GetValue<bool>("EnvelopeFollowerEnabled");
 	}
 
+	FCacheWriteHandler FSubmixEnvelopeFollowerEnabledMessage::GetCacheWriteHandler() const
+	{
+		return {
+			[](UE::Trace::FTraceWriter& Writer) -> uint32
+			{
+				return Writer.DeclareEvent(ANSITEXTVIEW("Audio"), ANSITEXTVIEW("SubmixEnvelopeFollowerEnabled"))
+					.Field(ANSITEXTVIEW("DeviceId"),                UE::Trace::ETraceWriterFieldType::Uint32)
+					.Field(ANSITEXTVIEW("SubmixId"),                UE::Trace::ETraceWriterFieldType::Uint32)
+					.Field(ANSITEXTVIEW("Timestamp"),               UE::Trace::ETraceWriterFieldType::Uint64)
+					.Field(ANSITEXTVIEW("EnvelopeFollowerEnabled"), UE::Trace::ETraceWriterFieldType::Bool)
+					.End();
+			},
+			[](UE::Trace::FTraceWriter& Writer, uint32 EventId, const IAudioCachedMessage& BaseMsg, uint64 TimestampCycles)
+			{
+				const FSubmixEnvelopeFollowerEnabledMessage& Msg = static_cast<const FSubmixEnvelopeFollowerEnabledMessage&>(BaseMsg);
+
+				Writer.WriteEvent(EventId)
+					.Field("DeviceId",                Msg.DeviceId)
+					.Field("SubmixId",                Msg.SubmixId)
+					.Field("Timestamp",               TimestampCycles)
+					.Field("EnvelopeFollowerEnabled", Msg.bEnvelopeFollowerEnabled)
+					.End();
+			}
+		};
+	}
+
 	// FSubmixEnvelopeValuesMessage
 	FSubmixEnvelopeValuesMessage::FSubmixEnvelopeValuesMessage(const Trace::IAnalyzer::FOnEventContext& InContext)
 		: FSubmixMessageBase(InContext)
@@ -91,6 +206,32 @@ namespace UE::Audio::Insights
 		return MemorySize;
 	};
 
+	FCacheWriteHandler FSubmixEnvelopeValuesMessage::GetCacheWriteHandler() const
+	{
+		return {
+			[](UE::Trace::FTraceWriter& Writer) -> uint32
+			{
+				return Writer.DeclareEvent(ANSITEXTVIEW("Audio"), ANSITEXTVIEW("SubmixEnvelopeValues"))
+					.Field(ANSITEXTVIEW("DeviceId"),       UE::Trace::ETraceWriterFieldType::Uint32)
+					.Field(ANSITEXTVIEW("SubmixId"),       UE::Trace::ETraceWriterFieldType::Uint32)
+					.Field(ANSITEXTVIEW("Timestamp"),      UE::Trace::ETraceWriterFieldType::Uint64)
+					.Field(ANSITEXTVIEW("EnvelopeValues"), UE::Trace::ETraceWriterFieldType::Float32 | UE::Trace::ETraceWriterFieldType::ArrayFlag)
+					.End();
+			},
+			[](UE::Trace::FTraceWriter& Writer, uint32 EventId, const IAudioCachedMessage& BaseMsg, uint64 TimestampCycles)
+			{
+				const FSubmixEnvelopeValuesMessage& Msg = static_cast<const FSubmixEnvelopeValuesMessage&>(BaseMsg);
+
+				Writer.WriteEvent(EventId)
+					.Field("DeviceId",       Msg.DeviceId)
+					.Field("SubmixId",       Msg.SubmixId)
+					.Field("Timestamp",      TimestampCycles)
+					.Field("EnvelopeValues", TConstArrayView<float>(Msg.EnvelopeValues.GetData(), Msg.EnvelopeValues.Num()))
+					.End();
+			}
+		};
+	}
+
 	// FSubmixUnloadedMessage
 	FSubmixUnloadedMessage::FSubmixUnloadedMessage(const Trace::IAnalyzer::FOnEventContext& InContext)
 		: FSubmixMessageBase(InContext)
@@ -98,6 +239,30 @@ namespace UE::Audio::Insights
 
 	};
 
+	FCacheWriteHandler FSubmixUnloadedMessage::GetCacheWriteHandler() const
+	{
+		return {
+			[](UE::Trace::FTraceWriter& Writer) -> uint32
+			{
+				return Writer.DeclareEvent(ANSITEXTVIEW("Audio"), ANSITEXTVIEW("SubmixUnloaded"))
+					.Field(ANSITEXTVIEW("DeviceId"),  UE::Trace::ETraceWriterFieldType::Uint32)
+					.Field(ANSITEXTVIEW("SubmixId"),  UE::Trace::ETraceWriterFieldType::Uint32)
+					.Field(ANSITEXTVIEW("Timestamp"), UE::Trace::ETraceWriterFieldType::Uint64)
+					.End();
+			},
+			[](UE::Trace::FTraceWriter& Writer, uint32 EventId, const IAudioCachedMessage& BaseMsg, uint64 TimestampCycles)
+			{
+				const FSubmixUnloadedMessage& Msg = static_cast<const FSubmixUnloadedMessage&>(BaseMsg);
+
+				Writer.WriteEvent(EventId)
+					.Field("DeviceId",  Msg.DeviceId)
+					.Field("SubmixId",  Msg.SubmixId)
+					.Field("Timestamp", TimestampCycles)
+					.End();
+			}
+		};
+	}
+
 	// FSubmixDashboardEntry
 	FSubmixDashboardEntry::FSubmixDashboardEntry()
 		: AudioMeterInfo(MakeShared<FAudioMeterInfo>())
diff --git a/Engine/Plugins/AudioInsights/Source/AudioInsights/Private/Messages/VirtualLoopTraceMessages.cpp b/Engine/Plugins/AudioInsights/Source/AudioInsights/Private/Messages/VirtualLoopTraceMessages.cpp
--- a/Engine/Plugins/AudioInsights/Source/AudioInsights/Private/Messages/VirtualLoopTraceMessages.cpp
+++ b/Engine/Plugins/AudioInsights/Source/AudioInsights/Private/Messages/VirtualLoopTraceMessages.cpp
@@ -1,6 +1,9 @@
 // Copyright Epic Games, Inc. All Rights Reserved.
 #include "Messages/VirtualLoopTraceMessages.h"
 
+#include "Trace/TraceWriter.h"
+#include "TraceWriter/CacheWriteHandler.h"
+
 namespace UE::Audio::Insights
 {
 	namespace VirtualLoopMessageNames
@@ -40,7 +43,35 @@ namespace UE::Audio::Insights
 
 		return MemorySize;
 	}
-	
+
+	FCacheWriteHandler FVirtualLoopVirtualizeMessage::GetCacheWriteHandler() const
+	{
+		return {
+			[](UE::Trace::FTraceWriter& Writer) -> uint32
+			{
+				return Writer.DeclareEvent(ANSITEXTVIEW("Audio"), ANSITEXTVIEW("VirtualLoopVirtualize"))
+					.Field(ANSITEXTVIEW("DeviceId"),    UE::Trace::ETraceWriterFieldType::Uint32)
+					.Field(ANSITEXTVIEW("Timestamp"),   UE::Trace::ETraceWriterFieldType::Uint64)
+					.Field(ANSITEXTVIEW("PlayOrder"),   UE::Trace::ETraceWriterFieldType::Uint32)
+					.Field(ANSITEXTVIEW("ComponentId"), UE::Trace::ETraceWriterFieldType::Uint64)
+					.Field(ANSITEXTVIEW("Name"),        UE::Trace::ETraceWriterFieldType::WideString)
+					.End();
+			},
+			[](UE::Trace::FTraceWriter& Writer, uint32 EventId, const IAudioCachedMessage& BaseMsg, uint64 TimestampCycles)
+			{
+				const FVirtualLoopVirtualizeMessage& Msg = static_cast<const FVirtualLoopVirtualizeMessage&>(BaseMsg);
+
+				Writer.WriteEvent(EventId)
+					.Field("DeviceId",    Msg.DeviceId)
+					.Field("Timestamp",   TimestampCycles)
+					.Field("PlayOrder",   Msg.PlayOrder)
+					.Field("ComponentId", Msg.ComponentId)
+					.Field("Name",        FStringView(Msg.Name))
+					.End();
+			}
+		};
+	}
+
 	// FVirtualLoopStopOrRealizeMessage
 	FVirtualLoopStopOrRealizeMessage::FVirtualLoopStopOrRealizeMessage(const Trace::IAnalyzer::FOnEventContext& InContext)
 		: FVirtualLoopMessageBase(InContext)
@@ -53,6 +84,30 @@ namespace UE::Audio::Insights
 		return sizeof(FVirtualLoopStopOrRealizeMessage);
 	}
 
+	FCacheWriteHandler FVirtualLoopStopOrRealizeMessage::GetCacheWriteHandler() const
+	{
+		return {
+			[](UE::Trace::FTraceWriter& Writer) -> uint32
+			{
+				return Writer.DeclareEvent(ANSITEXTVIEW("Audio"), ANSITEXTVIEW("VirtualLoopStopOrRealize"))
+					.Field(ANSITEXTVIEW("DeviceId"),  UE::Trace::ETraceWriterFieldType::Uint32)
+					.Field(ANSITEXTVIEW("Timestamp"), UE::Trace::ETraceWriterFieldType::Uint64)
+					.Field(ANSITEXTVIEW("PlayOrder"), UE::Trace::ETraceWriterFieldType::Uint32)
+					.End();
+			},
+			[](UE::Trace::FTraceWriter& Writer, uint32 EventId, const IAudioCachedMessage& BaseMsg, uint64 TimestampCycles)
+			{
+				const FVirtualLoopStopOrRealizeMessage& Msg = static_cast<const FVirtualLoopStopOrRealizeMessage&>(BaseMsg);
+
+				Writer.WriteEvent(EventId)
+					.Field("DeviceId",  Msg.DeviceId)
+					.Field("Timestamp", TimestampCycles)
+					.Field("PlayOrder", Msg.PlayOrder)
+					.End();
+			}
+		};
+	}
+
 	// FVirtualLoopUpdateMessage
 	FVirtualLoopUpdateMessage::FVirtualLoopUpdateMessage(const Trace::IAnalyzer::FOnEventContext& InContext)
 		: FVirtualLoopMessageBase(InContext)
@@ -77,4 +132,46 @@ namespace UE::Audio::Insights
 		// Returning fixed size of the struct (no dynamically sized members in it)
 		return sizeof(FVirtualLoopUpdateMessage);
 	}
+
+	FCacheWriteHandler FVirtualLoopUpdateMessage::GetCacheWriteHandler() const
+	{
+		return {
+			[](UE::Trace::FTraceWriter& Writer) -> uint32
+			{
+				return Writer.DeclareEvent(ANSITEXTVIEW("Audio"), ANSITEXTVIEW("VirtualLoopUpdate"))
+					.Field(ANSITEXTVIEW("DeviceId"),        UE::Trace::ETraceWriterFieldType::Uint32)
+					.Field(ANSITEXTVIEW("Timestamp"),       UE::Trace::ETraceWriterFieldType::Uint64)
+					.Field(ANSITEXTVIEW("PlayOrder"),       UE::Trace::ETraceWriterFieldType::Uint32)
+					.Field(ANSITEXTVIEW("TimeVirtualized"), UE::Trace::ETraceWriterFieldType::Float32)
+					.Field(ANSITEXTVIEW("PlaybackTime"),    UE::Trace::ETraceWriterFieldType::Float32)
+					.Field(ANSITEXTVIEW("UpdateInterval"),  UE::Trace::ETraceWriterFieldType::Float32)
+					.Field(ANSITEXTVIEW("LocationX"),       UE::Trace::ETraceWriterFieldType::Float64)
+					.Field(ANSITEXTVIEW("LocationY"),       UE::Trace::ETraceWriterFieldType::Float64)
+					.Field(ANSITEXTVIEW("LocationZ"),       UE::Trace::ETraceWriterFieldType::Float64)
+					.Field(ANSITEXTVIEW("RotatorPitch"),    UE::Trace::ETraceWriterFieldType::Float64)
+					.Field(ANSITEXTVIEW("RotatorYaw"),      UE::Trace::ETraceWriterFieldType::Float64)
+					.Field(ANSITEXTVIEW("RotatorRoll"),     UE::Trace::ETraceWriterFieldType::Float64)
+					.End();
+			},
+			[](UE::Trace::FTraceWriter& Writer, uint32 EventId, const IAudioCachedMessage& BaseMsg, uint64 TimestampCycles)
+			{
+				const FVirtualLoopUpdateMessage& Msg = static_cast<const FVirtualLoopUpdateMessage&>(BaseMsg);
+
+				Writer.WriteEvent(EventId)
+					.Field("DeviceId",        Msg.DeviceId)
+					.Field("Timestamp",       TimestampCycles)
+					.Field("PlayOrder",       Msg.PlayOrder)
+					.Field("TimeVirtualized", Msg.TimeVirtualized)
+					.Field("PlaybackTime",    Msg.PlaybackTime)
+					.Field("UpdateInterval",  Msg.UpdateInterval)
+					.Field("LocationX",       Msg.LocationX)
+					.Field("LocationY",       Msg.LocationY)
+					.Field("LocationZ",       Msg.LocationZ)
+					.Field("RotatorPitch",    Msg.RotatorPitch)
+					.Field("RotatorYaw",      Msg.RotatorYaw)
+					.Field("RotatorRoll",     Msg.RotatorRoll)
+					.End();
+			}
+		};
+	}
 } // namespace UE::Audio::Insights
diff --git a/Engine/Plugins/AudioInsights/Source/AudioInsights/Private/TraceWriter/AudioInsightsCacheTraceWriter.cpp b/Engine/Plugins/AudioInsights/Source/AudioInsights/Private/TraceWriter/AudioInsightsCacheTraceWriter.cpp
new file mode 100644
--- /dev/null
+++ b/Engine/Plugins/AudioInsights/Source/AudioInsights/Private/TraceWriter/AudioInsightsCacheTraceWriter.cpp
@@ -0,0 +1,239 @@
+// Copyright Epic Games, Inc. All Rights Reserved.
+#include "TraceWriter/AudioInsightsCacheTraceWriter.h"
+
+#include "AudioInsightsLog.h"
+#include "BuildSettings.h"
+#include "Cache/AudioCachedMessageChunk.h"
+#include "Cache/AudioInsightsCacheManager.h"
+#include "Misc/App.h"
+#include "ProfilingDebugging/MiscTrace.h"
+#include "Trace/OutDataStream.h"
+#include "Trace/TraceWriter.h"
+#include "TraceWriter/CacheWriteHandler.h"
+
+namespace UE::Audio::Insights
+{
+	namespace FAudioInsightsCacheTraceWriterPrivate
+	{
+		struct FResolvedHandler
+		{
+			FCacheWriteHandler Handler;
+			uint32 EventId = 0;
+		};
+
+		struct FFrameEventContext
+		{
+			uint32 GameThreadId      = 0;
+			uint32 BeginFrameEventId = 0;
+			uint32 EndFrameEventId   = 0;
+		};
+
+		struct FTimestampContext
+		{
+			double CacheStartTimestamp = 0.0;
+			uint64 BaseTimestamp       = 0;
+			uint64 TimeFrequency       = 0;
+		};
+
+		uint64 CacheTimestampToTraceCycles(double InTimestampSeconds, const FTimestampContext& InContext)
+		{
+			const double SecondsFromCacheStart = FMath::Max(0.0, InTimestampSeconds - InContext.CacheStartTimestamp);
+			return InContext.BaseTimestamp + static_cast<uint64>(SecondsFromCacheStart * InContext.TimeFrequency);
+		}
+
+		TMap<FName, FResolvedHandler> DeclareAudioEvents(UE::Trace::FTraceWriter& InWriter, const FAudioInsightsCacheManager& InCacheManager)
+		{
+			TMap<FName, FResolvedHandler> Handlers;
+
+			for (const auto& [MessageName, Message] : InCacheManager.GetMessageTypeRegistry())
+			{
+				FCacheWriteHandler Handler = Message->GetCacheWriteHandler();
+
+				if (Handler.IsValid())
+				{
+					const uint32 EventId = Handler.DeclareEvent(InWriter);
+					Handlers.Emplace(MessageName, FResolvedHandler{ MoveTemp(Handler), EventId });
+				}
+			}
+
+			return Handlers;
+		}
+
+		FFrameEventContext DeclareFrameEvents(UE::Trace::FTraceWriter& InWriter)
+		{
+			FFrameEventContext Context;
+
+			Context.GameThreadId = InWriter.RegisterThread(ANSITEXTVIEW("GameThread"));
+
+			Context.BeginFrameEventId = InWriter.DeclareEvent(ANSITEXTVIEW("Misc"), ANSITEXTVIEW("BeginFrame"))
+				.Field(ANSITEXTVIEW("Cycle"),     UE::Trace::ETraceWriterFieldType::Uint64)
+				.Field(ANSITEXTVIEW("FrameType"), UE::Trace::ETraceWriterFieldType::Uint8)
+				.End();
+
+			Context.EndFrameEventId = InWriter.DeclareEvent(ANSITEXTVIEW("Misc"), ANSITEXTVIEW("EndFrame"))
+				.Field(ANSITEXTVIEW("Cycle"),     UE::Trace::ETraceWriterFieldType::Uint64)
+				.Field(ANSITEXTVIEW("FrameType"), UE::Trace::ETraceWriterFieldType::Uint8)
+				.End();
+
+			return Context;
+		}
+
+		void WriteDiagnosticsSessionEvent(UE::Trace::FTraceWriter& InWriter)
+		{
+			constexpr int32 InstanceIdSize = 4;
+
+			FGuid InstanceGuid = FApp::GetInstanceId();
+			uint32 InstanceId[InstanceIdSize];
+
+			for (int32 Index = 0; Index < InstanceIdSize; ++Index)
+			{
+				InstanceId[Index] = InstanceGuid[Index];
+			}
+
+			const uint32 Session2EventId = InWriter.DeclareDiagnosticsSession2Event();
+
+			InWriter.SetCurrentThreadImportants();
+
+			InWriter.WriteEvent(Session2EventId)
+				.Field(ANSITEXTVIEW("ConfigurationType"), static_cast<uint8>(FApp::GetBuildConfiguration()))
+				.Field(ANSITEXTVIEW("TargetType"), static_cast<uint8>(FApp::GetBuildTargetType()))
+				.Field(ANSITEXTVIEW("Changelist"), BuildSettings::GetCurrentChangelist())
+				.Field(ANSITEXTVIEW("InstanceId"), TConstArrayView<uint32>(InstanceId, InstanceIdSize))
+				.Field(ANSITEXTVIEW("Platform"), FAnsiStringView(UE_STRINGIZE(UBT_COMPILED_PLATFORM)))
+				.Field(ANSITEXTVIEW("AppName"), UE_APP_NAME)
+				.Field(ANSITEXTVIEW("ProjectName"), FApp::GetProjectName())
+				.Field(ANSITEXTVIEW("Branch"), FStringView(*FApp::GetBranchName()))
+				.Field(ANSITEXTVIEW("BuildVersion"), FApp::GetBuildVersion())
+				.Field(ANSITEXTVIEW("EngineVersion"), BuildSettings::GetEngineVersionString())
+				.End();
+		}
+
+		void WriteFrameEvents(UE::Trace::FTraceWriter& InWriter, const FAudioInsightsCacheManager& InCacheManager, const FFrameEventContext& InFrameContext, const FTimestampContext& InTimestampContext)
+		{
+			if (InTimestampContext.TimeFrequency == 0)
+			{
+				return;
+			}
+
+			const uint64 CacheDurationCycles = static_cast<uint64>(InCacheManager.GetCacheDuration() * InTimestampContext.TimeFrequency);
+			constexpr uint64 FrameIntervalMs = 16; // ~60fps synthetic frame interval
+			const uint64 FrameDuration = (FrameIntervalMs * InTimestampContext.TimeFrequency) / 1000;
+
+			if (FrameDuration == 0)
+			{
+				return;
+			}
+			const uint8 GameFrameType  = static_cast<uint8>(TraceFrameType_Game);
+
+			InWriter.SetCurrentThread(InFrameContext.GameThreadId);
+
+			for (uint64 Offset = 0; Offset < CacheDurationCycles; Offset += FrameDuration)
+			{
+				const uint64 Duration = FMath::Min(FrameDuration, CacheDurationCycles - Offset);
+
+				InWriter.WriteEvent(InFrameContext.BeginFrameEventId)
+					.Field(0, InTimestampContext.BaseTimestamp + Offset)
+					.Field(1, GameFrameType)
+					.End();
+
+				InWriter.WriteEvent(InFrameContext.EndFrameEventId)
+					.Field(0, InTimestampContext.BaseTimestamp + Offset + Duration)
+					.Field(1, GameFrameType)
+					.End();
+			}
+		}
+
+		void WriteAudioEvents(UE::Trace::FTraceWriter& InWriter, const FAudioInsightsCacheManager& InCacheManager, const TMap<FName, FResolvedHandler>& InHandlers, const FTimestampContext& InTimestampContext)
+		{
+			uint32 TotalMessagesWritten = 0;
+			TSet<FName> UnhandledMessageTypes;
+
+			for (uint32 ChunkIndex = 0; ChunkIndex < InCacheManager.GetNumUsedChunks(); ++ChunkIndex)
+			{
+				const FAudioCachedMessageChunk* const Chunk = InCacheManager.GetChunk(ChunkIndex);
+
+				if (Chunk == nullptr || !Chunk->HasAnyData())
+				{
+					continue;
+				}
+
+				for (const auto& [MessageName, MessageArray] : Chunk->GetAllChunkMessages())
+				{
+					const FResolvedHandler* ResolvedHandler = InHandlers.Find(MessageName);
+
+					if (ResolvedHandler == nullptr)
+					{
+						UnhandledMessageTypes.Emplace(MessageName);
+						continue;
+					}
+
+					for (uint64 MessageIndex = 0; MessageIndex < MessageArray->Num(); ++MessageIndex)
+					{
+						const TSharedPtr<IAudioCachedMessage>& AudioCachedMessage = (*MessageArray)[MessageIndex];
+
+						if (AudioCachedMessage.IsValid())
+						{
+							const uint64 TimestampCycles = CacheTimestampToTraceCycles(AudioCachedMessage->Timestamp, InTimestampContext);
+
+							ResolvedHandler->Handler.WriteEvent(InWriter, ResolvedHandler->EventId, *AudioCachedMessage, TimestampCycles);
+							++TotalMessagesWritten;
+						}
+					}
+				}
+			}
+
+			UE_LOG(LogAudioInsights, Log, TEXT("Audio events written: %u, unhandled message types: %d"), TotalMessagesWritten, UnhandledMessageTypes.Num());
+		}
+	} // namespace FAudioInsightsCacheTraceWriterPrivate
+
+	bool FAudioInsightsCacheTraceWriter::WriteCacheToFile(const FAudioInsightsCacheManager& InCacheManager, const FString& InFilePath)
+	{
+		using namespace FAudioInsightsCacheTraceWriterPrivate;
+
+		// Setup data stream
+		UE::Trace::FFileOutDataStream DataStream;
+
+		if (!DataStream.Open(*InFilePath))
+		{
+			UE_LOG(LogAudioInsights, Warning, TEXT("Failed to open file for cache trace writing: %s"), *InFilePath);
+			return false;
+		}
+
+		// Setup writer
+		UE::Trace::FTraceWriter Writer(DataStream);
+		Writer.Begin();
+
+		// Setup timestamp conversion context
+		const FTimestampContext TimestampContext
+		{
+			.CacheStartTimestamp = InCacheManager.GetCacheStartTimeStamp(),
+			.BaseTimestamp       = Writer.GetStartTime(),
+			.TimeFrequency       = Writer.GetTimeFrequency()
+		};
+
+		// Declare events
+		FFrameEventContext FrameEventContext = DeclareFrameEvents(Writer);
+
+		TMap<FName, FResolvedHandler> Handlers = DeclareAudioEvents(Writer, InCacheManager);
+
+		if (Handlers.IsEmpty())
+		{
+			UE_LOG(LogAudioInsights, Warning, TEXT("No messages registered cache write handlers. Nothing to write."));
+			Writer.End();
+
+			return false;
+		}
+
+		// Write events
+		WriteDiagnosticsSessionEvent(Writer);
+		WriteFrameEvents(Writer, InCacheManager, FrameEventContext, TimestampContext);
+		WriteAudioEvents(Writer, InCacheManager, Handlers, TimestampContext);
+
+		Writer.End();
+
+		UE_LOG(LogAudioInsights, Log, TEXT("Cache trace written to %s"), *InFilePath);
+
+		return true;
+	}
+
+} // namespace UE::Audio::Insights
diff --git a/Engine/Plugins/AudioInsights/Source/AudioInsights/Public/AudioInsightsDashboardFactory.h b/Engine/Plugins/AudioInsights/Source/AudioInsights/Public/AudioInsightsDashboardFactory.h
--- a/Engine/Plugins/AudioInsights/Source/AudioInsights/Public/AudioInsightsDashboardFactory.h
+++ b/Engine/Plugins/AudioInsights/Source/AudioInsights/Public/AudioInsightsDashboardFactory.h
@@ -27,7 +27,7 @@ namespace UE::Audio::Insights
 	private:
 		TSharedRef<SWidget> MakeMenuBarWidget();
 		TSharedRef<SWidget> MakeMainToolbarWidget();
-		TSharedRef<SWidget> MakeSnapshotButtonWidget();
+		TSharedRef<SWidget> MakeCacheSnapshotButtonWidget();
 
 		TSharedPtr<FTabManager::FLayout> GetDefaultTabLayout();
 
diff --git a/Engine/Plugins/AudioInsights/Source/AudioInsights/Public/AudioInsightsStyle.h b/Engine/Plugins/AudioInsights/Source/AudioInsights/Public/AudioInsightsStyle.h
--- a/Engine/Plugins/AudioInsights/Source/AudioInsights/Public/AudioInsightsStyle.h
+++ b/Engine/Plugins/AudioInsights/Source/AudioInsights/Public/AudioInsightsStyle.h
@@ -225,7 +225,7 @@ namespace UE::Audio::Insights
 			Set("AudioInsights.Icon.Stop.Inactive", new IMAGE_BRUSH_SVG(TEXT("Icons/stop_inactive"), Icon20));
 			Set("AudioInsights.Icon.Submix", new IMAGE_BRUSH_SVG(TEXT("Icons/audio_submix"), Icon16));
 			Set("AudioInsights.Icon.Trace", new IMAGE_BRUSH_SVG(TEXT("Icons/trace"), Icon16));
-			Set("AudioInsights.Icon.Trace.Snapshot", new IMAGE_BRUSH_SVG(TEXT("Icons/trace_snapshot"), Icon16));
+			Set("AudioInsights.Icon.Trace.CacheSnapshot", new IMAGE_BRUSH_SVG(TEXT("Icons/cache_snapshot"), Icon16));
 			Set("AudioInsights.Icon.Trace.Bookmark", new IMAGE_BRUSH_SVG(TEXT("Icons/bookmark"), Icon16));
 			Set("AudioInsights.Icon.VirtualLoop", new IMAGE_BRUSH_SVG(TEXT("Icons/audio_virtualloop"), Icon16));
 			Set("AudioInsights.Icon.Viewport", new IMAGE_BRUSH_SVG(TEXT("Icons/viewport"), Icon16));
diff --git a/Engine/Plugins/AudioInsights/Source/AudioInsights/Public/AudioInsightsTraceModule.h b/Engine/Plugins/AudioInsights/Source/AudioInsights/Public/AudioInsightsTraceModule.h
--- a/Engine/Plugins/AudioInsights/Source/AudioInsights/Public/AudioInsightsTraceModule.h
+++ b/Engine/Plugins/AudioInsights/Source/AudioInsights/Public/AudioInsightsTraceModule.h
@@ -60,7 +60,7 @@ namespace UE::Audio::Insights
 
 		UE_API virtual void ExecuteConsoleCommand(const FString& InCommandStr) const override;
 
-		UE_API virtual void SaveTraceSnapshot();
+		UE_API virtual void SaveCacheSnapshot() override;
 
 		UE_API virtual ETraceMode GetTraceMode() const override;
 
diff --git a/Engine/Plugins/AudioInsights/Source/AudioInsights/Public/Cache/AudioInsightsCacheManager.h b/Engine/Plugins/AudioInsights/Source/AudioInsights/Public/Cache/AudioInsightsCacheManager.h
--- a/Engine/Plugins/AudioInsights/Source/AudioInsights/Public/Cache/AudioInsightsCacheManager.h
+++ b/Engine/Plugins/AudioInsights/Source/AudioInsights/Public/Cache/AudioInsightsCacheManager.h
@@ -83,6 +83,10 @@ namespace UE::Audio::Insights
 		// e.g. NumChunksFromStart = 0u will return the chunk with the earliest messages in the cache
 		UE_API const FAudioCachedMessageChunk* const GetChunk(const uint32 NumChunksFromStart) const;
 
+		// Returns one representative message per FName, collected as messages flow into the cache.
+		// Used by the cache trace writer to discover and declare events without iterating all chunks.
+		const TMap<FName, TSharedPtr<const IAudioCachedMessage>>& GetMessageTypeRegistry() const { return MessageTypeRegistry; }
+
 		// Find the closest message to a timestamp in the cache for a specific message type and where predicate returns true.
 		// Will look backwards from the chosen timestamp until a message is found
 		// @param Predicate: A callable that returns a BooleanTestable value.
@@ -266,6 +270,7 @@ namespace UE::Audio::Insights
 
 		FAudioInsightsCache Cache;
 		TAnalyzerMessageQueue<TSharedPtr<IAudioCachedMessage>> NewMessageQueue;
+		TMap<FName, TSharedPtr<const IAudioCachedMessage>> MessageTypeRegistry;
 
 		uint32 WriteIndex = 0u;
 		uint32 StartIndex = 0u;
diff --git a/Engine/Plugins/AudioInsights/Source/AudioInsights/Public/Cache/IAudioCachedMessage.h b/Engine/Plugins/AudioInsights/Source/AudioInsights/Public/Cache/IAudioCachedMessage.h
--- a/Engine/Plugins/AudioInsights/Source/AudioInsights/Public/Cache/IAudioCachedMessage.h
+++ b/Engine/Plugins/AudioInsights/Source/AudioInsights/Public/Cache/IAudioCachedMessage.h
@@ -3,6 +3,7 @@
 
 #include "Templates/SharedPointer.h"
 #include "Trace/Analyzer.h"
+#include "TraceWriter/CacheWriteHandler.h"
 
 namespace UE::Audio::Insights
 {
@@ -15,6 +16,11 @@ namespace UE::Audio::Insights
 		virtual uint64 GetID() const = 0;
 		virtual const FName GetMessageName() const = 0;
 
+		// Override to return a declare/write handler pair so this message type is included
+		// when exporting the Audio Insights cache to a custom .utrace file.
+		// See existing implementations in AudioInsights/Private/Messages/ for examples.
+		virtual FCacheWriteHandler GetCacheWriteHandler() const { return {}; }
+
 		double Timestamp = 0.0;
 	};
 } // namespace UE::Audio::Insights
diff --git a/Engine/Plugins/AudioInsights/Source/AudioInsights/Public/IAudioInsightsTraceModule.h b/Engine/Plugins/AudioInsights/Source/AudioInsights/Public/IAudioInsightsTraceModule.h
--- a/Engine/Plugins/AudioInsights/Source/AudioInsights/Public/IAudioInsightsTraceModule.h
+++ b/Engine/Plugins/AudioInsights/Source/AudioInsights/Public/IAudioInsightsTraceModule.h
@@ -30,7 +30,7 @@ public:
 
 	virtual void ExecuteConsoleCommand(const FString& InCommandStr) const = 0;
 
-	virtual void SaveTraceSnapshot() = 0;
+	virtual void SaveCacheSnapshot() = 0;
 
 	virtual UE::Audio::Insights::ETraceMode GetTraceMode() const = 0;
 
diff --git a/Engine/Plugins/AudioInsights/Source/AudioInsights/Public/Messages/AudioBusTraceMessages.h b/Engine/Plugins/AudioInsights/Source/AudioInsights/Public/Messages/AudioBusTraceMessages.h
--- a/Engine/Plugins/AudioInsights/Source/AudioInsights/Public/Messages/AudioBusTraceMessages.h
+++ b/Engine/Plugins/AudioInsights/Source/AudioInsights/Public/Messages/AudioBusTraceMessages.h
@@ -48,6 +48,7 @@ namespace UE::Audio::Insights
 
 		virtual uint32 GetSizeOf() const override;
 		virtual const FName GetMessageName() const override { return AudioBusMessageNames::Start; }
+		virtual FCacheWriteHandler GetCacheWriteHandler() const override;
 
 		FString Name;
 		int32 NumChannels = 0;
@@ -61,6 +62,7 @@ namespace UE::Audio::Insights
 
 		virtual uint32 GetSizeOf() const override { return sizeof(FAudioBusHasActivityMessage); };
 		virtual const FName GetMessageName() const override { return AudioBusMessageNames::HasActivity; }
+		virtual FCacheWriteHandler GetCacheWriteHandler() const override;
 
 		bool bHasActivity = false;
 	};
@@ -72,6 +74,7 @@ namespace UE::Audio::Insights
 
 		virtual uint32 GetSizeOf() const override { return sizeof(FAudioBusEnvelopeFollowerEnabledMessage); };
 		virtual const FName GetMessageName() const override { return AudioBusMessageNames::EnvelopeFollowerEnabled; }
+		virtual FCacheWriteHandler GetCacheWriteHandler() const override;
 
 		bool bEnvelopeFollowerEnabled = false;
 	};
@@ -83,6 +86,7 @@ namespace UE::Audio::Insights
 
 		virtual uint32 GetSizeOf() const override;
 		virtual const FName GetMessageName() const override { return AudioBusMessageNames::EnvelopeValues; }
+		virtual FCacheWriteHandler GetCacheWriteHandler() const override;
 
 		TArray<float> EnvelopeValues;
 	};
@@ -94,6 +98,7 @@ namespace UE::Audio::Insights
 
 		virtual uint32 GetSizeOf() const override { return sizeof(FAudioBusStopMessage); };
 		virtual const FName GetMessageName() const override { return AudioBusMessageNames::Stop; }
+		virtual FCacheWriteHandler GetCacheWriteHandler() const override;
 	};
 
 	class FAudioBusMessages
diff --git a/Engine/Plugins/AudioInsights/Source/AudioInsights/Public/Messages/AudioEventLogTraceMessages.h b/Engine/Plugins/AudioInsights/Source/AudioInsights/Public/Messages/AudioEventLogTraceMessages.h
--- a/Engine/Plugins/AudioInsights/Source/AudioInsights/Public/Messages/AudioEventLogTraceMessages.h
+++ b/Engine/Plugins/AudioInsights/Source/AudioInsights/Public/Messages/AudioEventLogTraceMessages.h
@@ -34,6 +34,7 @@ namespace UE::Audio::Insights
 		virtual uint32 GetSizeOf() const override;
 		virtual uint64 GetID() const override { return PlayOrder; }
 		virtual const FName GetMessageName() const override;
+		virtual FCacheWriteHandler GetCacheWriteHandler() const override;
 
 		::Audio::FDeviceId DeviceId = INDEX_NONE;
 		uint32 MessageID = INDEX_NONE;
diff --git a/Engine/Plugins/AudioInsights/Source/AudioInsights/Public/Messages/OutputMeterTraceMessages.h b/Engine/Plugins/AudioInsights/Source/AudioInsights/Public/Messages/OutputMeterTraceMessages.h
--- a/Engine/Plugins/AudioInsights/Source/AudioInsights/Public/Messages/OutputMeterTraceMessages.h
+++ b/Engine/Plugins/AudioInsights/Source/AudioInsights/Public/Messages/OutputMeterTraceMessages.h
@@ -33,6 +33,7 @@ namespace UE::Audio::Insights
 
 		virtual const FName GetMessageName() const override { return OutputMeterMessageNames::MainSubmixLoaded; }
 		virtual uint32 GetSizeOf() const override;
+		virtual FCacheWriteHandler GetCacheWriteHandler() const override;
 	};
 
 	struct FOutputMeterTruePeakMessage : public FOutputMeterMessageBase
@@ -42,6 +43,7 @@ namespace UE::Audio::Insights
 
 		virtual const FName GetMessageName() const override { return OutputMeterMessageNames::TruePeak; }
 		virtual uint32 GetSizeOf() const override;
+		virtual FCacheWriteHandler GetCacheWriteHandler() const override;
 
 		float TruePeakMaxValueDb = MIN_VOLUME_DECIBELS;
 	};
@@ -53,6 +55,7 @@ namespace UE::Audio::Insights
 
 		virtual const FName GetMessageName() const override { return OutputMeterMessageNames::LKFSValues; }
 		virtual uint32 GetSizeOf() const override;
+		virtual FCacheWriteHandler GetCacheWriteHandler() const override;
 
 		float LongTermLoudness  = MIN_VOLUME_DECIBELS;
 		float ShortTermLoudness = MIN_VOLUME_DECIBELS;
diff --git a/Engine/Plugins/AudioInsights/Source/AudioInsights/Public/Messages/SignalFlowTraceMessages.h b/Engine/Plugins/AudioInsights/Source/AudioInsights/Public/Messages/SignalFlowTraceMessages.h
--- a/Engine/Plugins/AudioInsights/Source/AudioInsights/Public/Messages/SignalFlowTraceMessages.h
+++ b/Engine/Plugins/AudioInsights/Source/AudioInsights/Public/Messages/SignalFlowTraceMessages.h
@@ -55,6 +55,7 @@ namespace UE::Audio::Insights
 
 		virtual uint64 GetID() const override { return GetTypeHash(FSignalFlowEntryKey{ DeviceId, ESignalFlowEntryType::SoundSource, SenderID }); }
 		virtual const FName GetMessageName() const override { return SignalFlowMessageNames::SoundWaveToBusSend; }
+		virtual FCacheWriteHandler GetCacheWriteHandler() const override;
 	};
 
 	struct FSoundToSubmixSendMessage : public FSignalFlowSendMessageBase
@@ -67,6 +68,7 @@ namespace UE::Audio::Insights
 
 		virtual uint64 GetID() const override { return GetTypeHash(FSignalFlowEntryKey{ DeviceId, ESignalFlowEntryType::SoundSource, SenderID }); }
 		virtual const FName GetMessageName() const override { return SignalFlowMessageNames::SoundWaveToSubmixSend; }
+		virtual FCacheWriteHandler GetCacheWriteHandler() const override;
 	};
 
 	struct FSubmixToSubmixSendMessage : public FSignalFlowMessageBase
@@ -77,6 +79,7 @@ namespace UE::Audio::Insights
 		virtual uint64 GetID() const override { return GetTypeHash(FSignalFlowEntryKey{ DeviceId, ESignalFlowEntryType::Submix, SenderID }); }
 		virtual const FName GetMessageName() const override { return SignalFlowMessageNames::SubmixToSubmixSend; }
 		virtual uint32 GetSizeOf() const override;
+		virtual FCacheWriteHandler GetCacheWriteHandler() const override;
 
 		uint32 SenderID = INDEX_NONE;
 		uint32 ReceiverID = INDEX_NONE;
diff --git a/Engine/Plugins/AudioInsights/Source/AudioInsights/Public/Messages/SoundTraceMessages.h b/Engine/Plugins/AudioInsights/Source/AudioInsights/Public/Messages/SoundTraceMessages.h
--- a/Engine/Plugins/AudioInsights/Source/AudioInsights/Public/Messages/SoundTraceMessages.h
+++ b/Engine/Plugins/AudioInsights/Source/AudioInsights/Public/Messages/SoundTraceMessages.h
@@ -98,6 +98,7 @@ namespace UE::Audio::Insights
 		virtual uint32 GetSizeOf() const override;
 		virtual const FName GetMessageName() const override { return SoundMessageNames::SoundStart; }
 		virtual uint64 GetID() const override { return ActiveSoundPlayOrder; }
+		virtual FCacheWriteHandler GetCacheWriteHandler() const override;
 
 		FString Name;
 		ESoundDashboardEntryType EntryType;
@@ -114,6 +115,7 @@ namespace UE::Audio::Insights
 		}
 
 		virtual const FName GetMessageName() const override { return SoundMessageNames::SoundIsAlivePing; }
+		virtual FCacheWriteHandler GetCacheWriteHandler() const override;
 	};
 
 	struct FSoundWaveStartMessage : public FSoundStartMessage
@@ -133,6 +135,7 @@ namespace UE::Audio::Insights
 		virtual uint32 GetSizeOf() const override;
 		virtual const FName GetMessageName() const override { return SoundMessageNames::SoundWaveStart; }
 		virtual uint64 GetID() const override { return SoundMessageUtils::GeneratePlayOrderUniqueID(ActiveSoundPlayOrder, WaveInstancePlayOrder); }
+		virtual FCacheWriteHandler GetCacheWriteHandler() const override;
 
 		uint32 WaveInstancePlayOrder = INDEX_NONE;
 		FString ActorIconName;
@@ -148,6 +151,7 @@ namespace UE::Audio::Insights
 		}
 
 		virtual const FName GetMessageName() const override { return SoundMessageNames::SoundWaveIsAlivePing; }
+		virtual FCacheWriteHandler GetCacheWriteHandler() const override;
 	};
 
 	struct FSoundStopMessage : public FSoundMessageBase
@@ -162,6 +166,7 @@ namespace UE::Audio::Insights
 		virtual uint32 GetSizeOf() const override;
 		virtual const FName GetMessageName() const override { return SoundMessageNames::SoundStop; }
 		virtual uint64 GetID() const override { return ActiveSoundPlayOrder; }
+		virtual FCacheWriteHandler GetCacheWriteHandler() const override;
 
 		uint32 ActiveSoundPlayOrder = INDEX_NONE;
 	};
@@ -180,6 +185,7 @@ namespace UE::Audio::Insights
 		virtual uint32 GetSizeOf() const override;
 		virtual const FName GetMessageName() const override { return SoundMessageNames::SoundWaveStop; }
 		virtual uint64 GetID() const override { return SoundMessageUtils::GeneratePlayOrderUniqueID(ActiveSoundPlayOrder, WaveInstancePlayOrder); }
+		virtual FCacheWriteHandler GetCacheWriteHandler() const override;
 
 		uint32 WaveInstancePlayOrder = INDEX_NONE;
 		uint32 ActiveSoundPlayOrder = INDEX_NONE;
@@ -198,6 +204,7 @@ namespace UE::Audio::Insights
 																								\
 		virtual uint32 GetSizeOf() const override { return sizeof(ClassName); }					\
 		virtual const FName GetMessageName() const override { return MessageName; }				\
+		virtual FCacheWriteHandler GetCacheWriteHandler() const override;						\
 																								\
 		Type ParamName = Default;																\
 	};
diff --git a/Engine/Plugins/AudioInsights/Source/AudioInsights/Public/Messages/SubmixTraceMessages.h b/Engine/Plugins/AudioInsights/Source/AudioInsights/Public/Messages/SubmixTraceMessages.h
--- a/Engine/Plugins/AudioInsights/Source/AudioInsights/Public/Messages/SubmixTraceMessages.h
+++ b/Engine/Plugins/AudioInsights/Source/AudioInsights/Public/Messages/SubmixTraceMessages.h
@@ -43,6 +43,7 @@ namespace UE::Audio::Insights
 
 		virtual uint32 GetSizeOf() const override;
 		virtual const FName GetMessageName() const override { return SubmixMessageNames::Loaded; }
+		virtual FCacheWriteHandler GetCacheWriteHandler() const override;
 
 		FString Name;
 		int32 NumChannels = 0;
@@ -58,6 +59,7 @@ namespace UE::Audio::Insights
 		}
 
 		virtual const FName GetMessageName() const override { return SubmixMessageNames::IsAlivePing; }
+		virtual FCacheWriteHandler GetCacheWriteHandler() const override;
 	};
 
 	struct FSubmixHasActivityMessage : public FSubmixMessageBase
@@ -67,6 +69,7 @@ namespace UE::Audio::Insights
 
 		virtual uint32 GetSizeOf() const override { return sizeof(FSubmixHasActivityMessage); };
 		virtual const FName GetMessageName() const override { return SubmixMessageNames::HasActivity; }
+		virtual FCacheWriteHandler GetCacheWriteHandler() const override;
 
 		bool bHasActivity = false;
 	};
@@ -78,6 +81,7 @@ namespace UE::Audio::Insights
 
 		virtual uint32 GetSizeOf() const override { return sizeof(FSubmixEnvelopeFollowerEnabledMessage); };
 		virtual const FName GetMessageName() const override { return SubmixMessageNames::EnvelopeFollowerEnabled; }
+		virtual FCacheWriteHandler GetCacheWriteHandler() const override;
 
 		bool bEnvelopeFollowerEnabled = false;
 	};
@@ -89,6 +93,7 @@ namespace UE::Audio::Insights
 
 		virtual uint32 GetSizeOf() const override;
 		virtual const FName GetMessageName() const override { return SubmixMessageNames::EnvelopeValues; }
+		virtual FCacheWriteHandler GetCacheWriteHandler() const override;
 
 		TArray<float> EnvelopeValues;
 	};
@@ -100,7 +105,7 @@ namespace UE::Audio::Insights
 
 		virtual uint32 GetSizeOf() const override { return sizeof(FSubmixUnloadedMessage); };
 		virtual const FName GetMessageName() const override { return SubmixMessageNames::Unloaded; }
-
+		virtual FCacheWriteHandler GetCacheWriteHandler() const override;
 	};
 
 	class FSubmixMessages
diff --git a/Engine/Plugins/AudioInsights/Source/AudioInsights/Public/Messages/VirtualLoopTraceMessages.h b/Engine/Plugins/AudioInsights/Source/AudioInsights/Public/Messages/VirtualLoopTraceMessages.h
--- a/Engine/Plugins/AudioInsights/Source/AudioInsights/Public/Messages/VirtualLoopTraceMessages.h
+++ b/Engine/Plugins/AudioInsights/Source/AudioInsights/Public/Messages/VirtualLoopTraceMessages.h
@@ -40,6 +40,7 @@ namespace UE::Audio::Insights
 
 		virtual const FName GetMessageName() const override { return VirtualLoopMessageNames::Virtualize; }
 		virtual uint32 GetSizeOf() const override;
+		virtual FCacheWriteHandler GetCacheWriteHandler() const override;
 
 		FString Name;
 		uint64 ComponentId = TNumericLimits<uint64>::Max();
@@ -52,6 +53,7 @@ namespace UE::Audio::Insights
 
 		virtual const FName GetMessageName() const override { return VirtualLoopMessageNames::StopOrRealize; }
 		virtual uint32 GetSizeOf() const override;
+		virtual FCacheWriteHandler GetCacheWriteHandler() const override;
 	};
 
 	struct FVirtualLoopUpdateMessage : public FVirtualLoopMessageBase
@@ -61,6 +63,7 @@ namespace UE::Audio::Insights
 
 		virtual const FName GetMessageName() const override { return VirtualLoopMessageNames::Update; }
 		virtual uint32 GetSizeOf() const override;
+		virtual FCacheWriteHandler GetCacheWriteHandler() const override;
 
 		float TimeVirtualized = 0.0f;
 		float PlaybackTime = 0.0f;
diff --git a/Engine/Plugins/AudioInsights/Source/AudioInsights/Public/TraceWriter/AudioInsightsCacheTraceWriter.h b/Engine/Plugins/AudioInsights/Source/AudioInsights/Public/TraceWriter/AudioInsightsCacheTraceWriter.h
new file mode 100644
--- /dev/null
+++ b/Engine/Plugins/AudioInsights/Source/AudioInsights/Public/TraceWriter/AudioInsightsCacheTraceWriter.h
@@ -0,0 +1,25 @@
+// Copyright Epic Games, Inc. All Rights Reserved.
+#pragma once
+
+#include "Containers/UnrealString.h"
+#include "HAL/Platform.h"
+
+namespace UE::Trace
+{
+	class FTraceWriter;
+} // namespace UE::Trace
+
+#define UE_API AUDIOINSIGHTS_API
+
+namespace UE::Audio::Insights
+{
+	class FAudioInsightsCacheManager;
+
+	class FAudioInsightsCacheTraceWriter
+	{
+	public:
+		UE_API static bool WriteCacheToFile(const FAudioInsightsCacheManager& InCacheManager, const FString& InFilePath);
+	};
+} // namespace UE::Audio::Insights
+
+#undef UE_API
diff --git a/Engine/Plugins/AudioInsights/Source/AudioInsights/Public/TraceWriter/CacheWriteHandler.h b/Engine/Plugins/AudioInsights/Source/AudioInsights/Public/TraceWriter/CacheWriteHandler.h
new file mode 100644
--- /dev/null
+++ b/Engine/Plugins/AudioInsights/Source/AudioInsights/Public/TraceWriter/CacheWriteHandler.h
@@ -0,0 +1,25 @@
+// Copyright Epic Games, Inc. All Rights Reserved.
+#pragma once
+
+#include "Templates/Function.h"
+
+namespace UE::Trace
+{
+	class FTraceWriter;
+} // namespace UE::Trace
+
+namespace UE::Audio::Insights
+{
+	class IAudioCachedMessage;
+
+	struct FCacheWriteHandler
+	{
+		using FDeclareEventFunc = TFunction<uint32(UE::Trace::FTraceWriter& Writer)>;
+		using FWriteEventFunc   = TFunction<void(UE::Trace::FTraceWriter& Writer, uint32 EventId, const IAudioCachedMessage& Message, uint64 TimestampCycles)>;
+
+		FDeclareEventFunc DeclareEvent;
+		FWriteEventFunc WriteEvent;
+
+		bool IsValid() const { return DeclareEvent && WriteEvent; }
+	};
+} // namespace UE::Audio::Insights
diff --git a/Engine/Plugins/AudioInsights/Source/AudioInsightsEditor/Private/AudioInsightsEditorDashboardFactory.cpp b/Engine/Plugins/AudioInsights/Source/AudioInsightsEditor/Private/AudioInsightsEditorDashboardFactory.cpp
--- a/Engine/Plugins/AudioInsights/Source/AudioInsightsEditor/Private/AudioInsightsEditorDashboardFactory.cpp
+++ b/Engine/Plugins/AudioInsights/Source/AudioInsightsEditor/Private/AudioInsightsEditorDashboardFactory.cpp
@@ -58,10 +58,9 @@ namespace UE::Audio::Insights
 		static const FText RecordingNotificationText = LOCTEXT("AudioDashboardEditor_RecordingButton_NotificationText", "Trace Started");
 		static const FText RecordingNotificationSubText = LOCTEXT("AudioDashboardEditor_RecordingButton_NotificationSubText", "Trace is now active and saving to the following location (file or tracestore):\n");
 
-		static const FText SaveSnapshotTooltipText = LOCTEXT("AudioDashboardEditor_SnapshotButton_TooltipText", "Saves a snapshot of the most recent trace data to a .utrace file. By default snapshots are limited to 32MB and can be accessed in Unreal Insights. To specify a custom snapshot size (in MB) launch the editor executable with the -tracetailmb=<size> option.");
-
-		static const FText SaveSnapshotNotificationText = LOCTEXT("AudioDashboardEditor_SnapshotButton_NotificationText", "Audio Insights Snapshot saved.");
-		static const FText SaveSnapshotNotificationSubText = LOCTEXT("AudioDashboardEditor_SnapshotButton_NotificationSubText", "A snapshot .utrace with the most recent events has been saved to your trace server.");
+		static const FText SaveCacheSnapshotTooltipText = LOCTEXT("AudioDashboardEditor_SaveCacheSnapshotButton_TooltipText", "Saves a snapshot of the Audio Insights cached data to a .utrace file containing only audio trace events.");
+		static const FText SaveCacheSnapshotNotificationText = LOCTEXT("AudioDashboardEditor_SaveCacheSnapshotButton_NotificationText", "Audio Insights Cache saved.");
+		static const FText SaveCacheSnapshotNotificationSubText = LOCTEXT("AudioDashboardEditor_SaveCacheSnapshotButton_NotificationSubText", "A snapshot .utrace with Audio Insights cache data with audio trace events has been saved to your trace server.");
 
 		static const FText SaveBookmarkTooltipText = LOCTEXT("AudioDashboardEditor_BookmarkButton_TooltipText", "Ctrl+M - Creates an 'Audio Insights' bookmark in the .utrace file generated by Save Snapshot or Trace at the timestamp the button was pressed.");
 
@@ -628,7 +627,7 @@ namespace UE::Audio::Insights
 						.Thickness(2.0f)
 					]
 				]
-				// Save snapshot button
+				// Save cache snapshot button
 				+ SHorizontalBox::Slot()
 				.HAlign(HAlign_Left)
 				.VAlign(VAlign_Center)
@@ -637,7 +636,7 @@ namespace UE::Audio::Insights
 				[
 					SNew(SButton)
 					.ButtonStyle(&FAppStyle::Get().GetWidgetStyle<FButtonStyle>("SimpleButton"))
-					.ToolTipText(SaveSnapshotTooltipText)
+					.ToolTipText(SaveCacheSnapshotTooltipText)
 					.IsEnabled_Lambda([this]()
 					{
 						IAudioInsightsTraceModule& TraceModule = FAudioInsightsEditorModule::GetChecked().GetTraceModule();
@@ -645,11 +644,10 @@ namespace UE::Audio::Insights
 					})
 					.OnClicked_Lambda([]()
 					{
-						// Save snapshot
 						IAudioInsightsTraceModule& TraceModule = FAudioInsightsEditorModule::GetChecked().GetTraceModule();
-						TraceModule.SaveTraceSnapshot();
+						TraceModule.SaveCacheSnapshot();
 
-						ShowNotification(SaveSnapshotNotificationText, SaveSnapshotNotificationSubText);
+						ShowNotification(SaveCacheSnapshotNotificationText, SaveCacheSnapshotNotificationSubText);
 
 						return FReply::Handled();
 					})
@@ -658,13 +656,13 @@ namespace UE::Audio::Insights
 						.HeightOverride(16.0f)
 						[
 							SNew(SHorizontalBox)
-							// Trace snapshot icon
+							// Save cache snapshot icon
 							+ SHorizontalBox::Slot()
 							.VAlign(VAlign_Center)
 							.AutoWidth()
 							[
 								SNew(SImage)
-								.Image(FSlateStyle::Get().GetBrush("AudioInsights.Icon.Trace.Snapshot"))
+								.Image(FSlateStyle::Get().GetBrush("AudioInsights.Icon.Trace.CacheSnapshot"))
 								.ColorAndOpacity(FSlateColor::UseForeground())
 							]
 						]
diff --git a/Engine/Plugins/Developer/AudioModulationInsights/Source/AudioModulationInsights/Private/Messages/AudioModulationTraceMessages.cpp b/Engine/Plugins/Developer/AudioModulationInsights/Source/AudioModulationInsights/Private/Messages/AudioModulationTraceMessages.cpp
new file mode 100644
--- /dev/null
+++ b/Engine/Plugins/Developer/AudioModulationInsights/Source/AudioModulationInsights/Private/Messages/AudioModulationTraceMessages.cpp
@@ -0,0 +1,94 @@
+// Copyright Epic Games, Inc. All Rights Reserved.
+#include "AudioModulationTraceMessages.h"
+
+#include "Trace/TraceWriter.h"
+#include "TraceWriter/CacheWriteHandler.h"
+
+namespace AudioModulationInsights
+{
+	using UE::Audio::Insights::IAudioCachedMessage;
+
+	UE::Audio::Insights::FCacheWriteHandler FActivateModulatorTraceMessage::GetCacheWriteHandler() const
+	{
+		return {
+			[](UE::Trace::FTraceWriter& Writer) -> uint32
+			{
+				return Writer.DeclareEvent(ANSITEXTVIEW("Audio"), ANSITEXTVIEW("ActivateModulatorTraceMessage"))
+					.Field(ANSITEXTVIEW("DeviceId"),      UE::Trace::ETraceWriterFieldType::Uint32)
+					.Field(ANSITEXTVIEW("Timestamp"),     UE::Trace::ETraceWriterFieldType::Uint64)
+					.Field(ANSITEXTVIEW("ModulatorId"),   UE::Trace::ETraceWriterFieldType::Uint32)
+					.Field(ANSITEXTVIEW("ModulatorType"), UE::Trace::ETraceWriterFieldType::Uint8)
+					.Field(ANSITEXTVIEW("Name"),          UE::Trace::ETraceWriterFieldType::WideString)
+					.End();
+			},
+			[](UE::Trace::FTraceWriter& Writer, uint32 EventId, const IAudioCachedMessage& BaseMsg, uint64 TimestampCycles)
+			{
+				const FActivateModulatorTraceMessage& Msg = static_cast<const FActivateModulatorTraceMessage&>(BaseMsg);
+
+				Writer.WriteEvent(EventId)
+					.Field("DeviceId",      Msg.DeviceId)
+					.Field("Timestamp",     TimestampCycles)
+					.Field("ModulatorId",   static_cast<uint32>(Msg.ModulatorId))
+					.Field("ModulatorType", static_cast<uint8>(Msg.ModulatorType))
+					.Field("Name",          FStringView(Msg.ModulatorName))
+					.End();
+			}
+		};
+	}
+
+	UE::Audio::Insights::FCacheWriteHandler FUpdateModulatorTraceMessage::GetCacheWriteHandler() const
+	{
+		return {
+			[](UE::Trace::FTraceWriter& Writer) -> uint32
+			{
+				return Writer.DeclareEvent(ANSITEXTVIEW("Audio"), ANSITEXTVIEW("UpdateModulatorTraceMessage"))
+					.Field(ANSITEXTVIEW("DeviceId"),                  UE::Trace::ETraceWriterFieldType::Uint32)
+					.Field(ANSITEXTVIEW("Timestamp"),                 UE::Trace::ETraceWriterFieldType::Uint64)
+					.Field(ANSITEXTVIEW("ModulatorId"),               UE::Trace::ETraceWriterFieldType::Uint32)
+					.Field(ANSITEXTVIEW("ModulatorValue"),            UE::Trace::ETraceWriterFieldType::Float32)
+					.Field(ANSITEXTVIEW("IsBypassed"),                UE::Trace::ETraceWriterFieldType::Bool)
+					.Field(ANSITEXTVIEW("ContributingModulatorIds"),  UE::Trace::ETraceWriterFieldType::Uint32 | UE::Trace::ETraceWriterFieldType::ArrayFlag)
+					.Field(ANSITEXTVIEW("ContributingModulatorValues"), UE::Trace::ETraceWriterFieldType::Float32 | UE::Trace::ETraceWriterFieldType::ArrayFlag)
+					.End();
+			},
+			[](UE::Trace::FTraceWriter& Writer, uint32 EventId, const IAudioCachedMessage& BaseMsg, uint64 TimestampCycles)
+			{
+				const FUpdateModulatorTraceMessage& Msg = static_cast<const FUpdateModulatorTraceMessage&>(BaseMsg);
+
+				Writer.WriteEvent(EventId)
+					.Field("DeviceId",       Msg.DeviceId)
+					.Field("Timestamp",      TimestampCycles)
+					.Field("ModulatorId",    static_cast<uint32>(Msg.ModulatorId))
+					.Field("ModulatorValue", Msg.ModulatorValue)
+					.Field("IsBypassed",     Msg.bIsBypassed)
+					.Field("ContributingModulatorIds",    TConstArrayView<uint32>(Msg.ContributingModulatorIds.GetData(), Msg.ContributingModulatorIds.Num()))
+					.Field("ContributingModulatorValues", TConstArrayView<float>(Msg.ContributingModulatorValues.GetData(), Msg.ContributingModulatorValues.Num()))
+					.End();
+			}
+		};
+	}
+
+	UE::Audio::Insights::FCacheWriteHandler FDeactivateModulatorTraceMessage::GetCacheWriteHandler() const
+	{
+		return {
+			[](UE::Trace::FTraceWriter& Writer) -> uint32
+			{
+				return Writer.DeclareEvent(ANSITEXTVIEW("Audio"), ANSITEXTVIEW("DeactivateModulatorTraceMessage"))
+					.Field(ANSITEXTVIEW("DeviceId"),    UE::Trace::ETraceWriterFieldType::Uint32)
+					.Field(ANSITEXTVIEW("Timestamp"),   UE::Trace::ETraceWriterFieldType::Uint64)
+					.Field(ANSITEXTVIEW("ModulatorId"), UE::Trace::ETraceWriterFieldType::Uint32)
+					.End();
+			},
+			[](UE::Trace::FTraceWriter& Writer, uint32 EventId, const IAudioCachedMessage& BaseMsg, uint64 TimestampCycles)
+			{
+				const FDeactivateModulatorTraceMessage& Msg = static_cast<const FDeactivateModulatorTraceMessage&>(BaseMsg);
+
+				Writer.WriteEvent(EventId)
+					.Field("DeviceId",    Msg.DeviceId)
+					.Field("Timestamp",   TimestampCycles)
+					.Field("ModulatorId", static_cast<uint32>(Msg.ModulatorId))
+					.End();
+			}
+		};
+	}
+} // namespace AudioModulationInsights
diff --git a/Engine/Plugins/Developer/AudioModulationInsights/Source/AudioModulationInsights/Private/Messages/AudioModulationTraceMessages.h b/Engine/Plugins/Developer/AudioModulationInsights/Source/AudioModulationInsights/Private/Messages/AudioModulationTraceMessages.h
--- a/Engine/Plugins/Developer/AudioModulationInsights/Source/AudioModulationInsights/Private/Messages/AudioModulationTraceMessages.h
+++ b/Engine/Plugins/Developer/AudioModulationInsights/Source/AudioModulationInsights/Private/Messages/AudioModulationTraceMessages.h
@@ -49,6 +49,7 @@ namespace AudioModulationInsights
 		// from UE::Audio::Insights::IAudioCachedMessage
 		virtual uint32 GetSizeOf() const override { return (sizeof(FActivateModulatorTraceMessage) + ModulatorName.GetAllocatedSize()); }
 		virtual const FName GetMessageName() const override { return ModulatorTraceMessageNames::ActivateModulatorTraceMessage; }
+		virtual UE::Audio::Insights::FCacheWriteHandler GetCacheWriteHandler() const override;
 
 		EModulatorTraceType ModulatorType { EModulatorTraceType::COUNT };
 		FString ModulatorName;
@@ -81,6 +82,7 @@ namespace AudioModulationInsights
 		// from UE::Audio::Insights::IAudioCachedMessage
 		virtual uint32 GetSizeOf() const override { return (sizeof(FUpdateModulatorTraceMessage) + ContributingModulatorIds.GetAllocatedSize() + ContributingModulatorValues.GetAllocatedSize()); }
 		virtual const FName GetMessageName() const override { return ModulatorTraceMessageNames::UpdateModulatorTraceMessage; }
+		virtual UE::Audio::Insights::FCacheWriteHandler GetCacheWriteHandler() const override;
 
 		float ModulatorValue { 1.0f };
 		bool bIsBypassed { false };
@@ -96,6 +98,7 @@ namespace AudioModulationInsights
 
 		// from UE::Audio::Insights::IAudioCachedMessage
 		virtual const FName GetMessageName() const override { return ModulatorTraceMessageNames::DeactivateModulatorTraceMessage; }
+		virtual UE::Audio::Insights::FCacheWriteHandler GetCacheWriteHandler() const override;
 	};
 
 } // namespace AudioModulationInsights
