[Audio Insights] Added bookmarks support to Audio Insights standalone + added logic so bookmarks are part of the cache system and therefore be saved into the cache snapshots.

#Synced-CL 52673159

diff --git a/Engine/Plugins/AudioInsights/Source/AudioInsights/Private/AudioInsightsDashboardFactory.cpp b/Engine/Plugins/AudioInsights/Source/AudioInsights/Private/AudioInsightsDashboardFactory.cpp
--- a/Engine/Plugins/AudioInsights/Source/AudioInsights/Private/AudioInsightsDashboardFactory.cpp
+++ b/Engine/Plugins/AudioInsights/Source/AudioInsights/Private/AudioInsightsDashboardFactory.cpp
@@ -4,9 +4,13 @@
 #include "AudioInsightsComponent.h"
 #include "AudioInsightsModule.h"
 #include "AudioInsightsStyle.h"
+#include "Cache/AudioInsightsCacheManager.h"
 #include "Framework/Docking/TabManager.h"
 #include "Framework/MultiBox/MultiBoxBuilder.h"
+#include "IAudioInsightsModuleInterface.h"
 #include "Internationalization/Text.h"
+#include "Messages/BookmarkTraceMessages.h"
+#include "ProfilingDebugging/MiscTrace.h"
 #include "Templates/SharedPointer.h"
 #include "Widgets/Images/SImage.h"
 #include "Widgets/Input/SButton.h"
@@ -31,6 +35,9 @@ namespace UE::Audio::Insights
 		static const FText TryEnablingMessagingText = LOCTEXT("AudioDashboard_TryEnablingMessagingText", "Make sure you have launched this package with the -Messaging command line argument.");
 		static const FText SaveCacheSnapshotText = LOCTEXT("AudioDashboard_SaveCacheSnapshotButton_Text", "Save Cache Snapshot");
 		static const FText SaveCacheSnapshotTooltipText = LOCTEXT("AudioDashboard_SaveCacheSnapshotButton_TooltipText", "Saves a snapshot of the Audio Insights cached data to a .utrace file containing only audio trace events.");
+		static const FText SaveBookmarkText = LOCTEXT("AudioDashboard_BookmarkButton_Text", "Save Bookmark");
+		static const FText SaveBookmarkTooltipText = LOCTEXT("AudioDashboard_BookmarkButton_TooltipText", "Creates an 'Audio Insights' bookmark at the current cache timestamp.");
+		static const FText BookmarkLabelText = LOCTEXT("AudioDashboard_BookmarkLabelTitle_Text", "Audio Insights Bookmark");
 	}
 
 	::Audio::FDeviceId FDashboardFactory::GetDeviceId() const
@@ -214,6 +221,15 @@ namespace UE::Audio::Insights
 			[
 				MakeCacheSnapshotButtonWidget()
 			]
+			// Bookmark button
+			+ SHorizontalBox::Slot()
+			.HAlign(HAlign_Left)
+			.VAlign(VAlign_Center)
+			.AutoWidth()
+			.Padding(5.0f, 0.0f, 0.0f, 0.0f)
+			[
+				MakeBookmarkButtonWidget()
+			]
 			// Cache toolbar widgets (only visible during live monitoring)
 			+ SHorizontalBox::Slot()
 			.AutoWidth()
@@ -315,7 +331,54 @@ namespace UE::Audio::Insights
 				]
 			];
 	}
-	
+
+	TSharedRef<SWidget> FDashboardFactory::MakeBookmarkButtonWidget()
+	{
+		using namespace DashboardFactoryPrivate;
+
+		return SNew(SButton)
+			.ButtonStyle(&FAppStyle::Get().GetWidgetStyle<FButtonStyle>("SimpleButton"))
+			.ToolTipText(SaveBookmarkTooltipText)
+#if !WITH_EDITOR
+			.IsEnabled_Lambda([]()
+			{
+				const TSharedPtr<FAudioInsightsComponent> Component = FAudioInsightsModule::GetChecked().GetAudioInsightsComponent();
+				return Component.IsValid() && Component->GetIsLiveSession();
+			})
+#endif // !WITH_EDITOR
+			.OnClicked_Lambda([this]()
+			{
+				CreateTraceBookmark();
+
+				return FReply::Handled();
+			})
+			[
+				SNew(SBox)
+				.HeightOverride(16.0f)
+				[
+					SNew(SHorizontalBox)
+					+ SHorizontalBox::Slot()
+					.VAlign(VAlign_Center)
+					.AutoWidth()
+					[
+						SNew(SImage)
+						.Image(FSlateStyle::Get().GetBrush("AudioInsights.Icon.Trace.Bookmark"))
+						.ColorAndOpacity(FSlateColor::UseForeground())
+					]
+					// Save bookmark text
+					+ SHorizontalBox::Slot()
+					.VAlign(VAlign_Center)
+					.AutoWidth()
+					.Padding(6.0f, 1.0f, 0.0f, 0.0f)
+					[
+						SNew(STextBlock)
+						.Justification(ETextJustify::Center)
+						.Text(SaveBookmarkText)
+					]
+				]
+			];
+	}
+
 	TSharedPtr<FTabManager::FLayout> FDashboardFactory::GetDefaultTabLayout()
 	{
 		using namespace DashboardFactoryPrivate;
@@ -617,5 +680,20 @@ namespace UE::Audio::Insights
 	}
 #endif // !WITH_EDITOR
 
+	void FDashboardFactory::CreateTraceBookmark()
+	{
+		using namespace DashboardFactoryPrivate;
+
+		const FString BookmarkLabel = FString::Printf(TEXT("%s %d"), *BookmarkLabelText.ToString(), BookmarkIndex);
+
+		TRACE_BOOKMARK(TEXT("%s"), *BookmarkLabel);
+
+		FAudioInsightsCacheManager& CacheManager = IAudioInsightsModule::GetChecked().GetCacheManager();
+
+		const double CurrentTimestamp = CacheManager.GetCacheEndTimeStamp();
+		CacheManager.AddMessageToCache(MakeShared<FBookmarkTraceMessage>(CurrentTimestamp, FString(BookmarkLabel)));
+
+		BookmarkIndex++;
+	}
 } // namespace UE::Audio::Insights
 #undef LOCTEXT_NAMESPACE
diff --git a/Engine/Plugins/AudioInsights/Source/AudioInsights/Private/Messages/BookmarkTraceMessages.cpp b/Engine/Plugins/AudioInsights/Source/AudioInsights/Private/Messages/BookmarkTraceMessages.cpp
new file mode 100644
--- /dev/null
+++ b/Engine/Plugins/AudioInsights/Source/AudioInsights/Private/Messages/BookmarkTraceMessages.cpp
@@ -0,0 +1,78 @@
+// Copyright Epic Games, Inc. All Rights Reserved.
+#include "Messages/BookmarkTraceMessages.h"
+
+#include "ProfilingDebugging/FormatArgsTrace.h"
+#include "Trace/TraceWriter.h"
+#include "TraceWriter/CacheWriteHandler.h"
+
+namespace UE::Audio::Insights
+{
+	namespace BookmarkMessageNames
+	{
+		const FName Bookmark = "AudioInsightsBookmark";
+	};
+
+	FBookmarkTraceMessage::FBookmarkTraceMessage(double InTimestamp, const FString& InLabel)
+		: Label(InLabel)
+	{
+		Timestamp = InTimestamp;
+	}
+
+	uint32 FBookmarkTraceMessage::GetSizeOf() const
+	{
+		uint32 MemorySize = sizeof(FBookmarkTraceMessage);
+
+		MemorySize += Label.GetAllocatedSize();
+
+		return MemorySize;
+	}
+
+	FCacheWriteHandler FBookmarkTraceMessage::GetCacheWriteHandler() const
+	{
+		return {
+			[](UE::Trace::FTraceWriter& Writer) -> uint32
+			{
+				// Declare and write BookmarkSpec (once for all bookmark messages)
+				const uint32 SpecEventId = Writer.DeclareEvent(ANSITEXTVIEW("Misc"), ANSITEXTVIEW("BookmarkSpec"), UE::Trace::ETraceWriterEventFlags::ImportantNoSync)
+					.Field(ANSITEXTVIEW("BookmarkPoint"), UE::Trace::ETraceWriterFieldType::Pointer)
+					.Field(ANSITEXTVIEW("Line"),          UE::Trace::ETraceWriterFieldType::Int32)
+					.Field(ANSITEXTVIEW("FormatString"),  UE::Trace::ETraceWriterFieldType::WideString)
+					.Field(ANSITEXTVIEW("FileName"),      UE::Trace::ETraceWriterFieldType::AnsiString)
+					.End();
+
+				constexpr uint64 SyntheticBookmarkPoint = 1;
+
+				Writer.WriteEvent(SpecEventId)
+					.Field(ANSITEXTVIEW("BookmarkPoint"), SyntheticBookmarkPoint)
+					.Field(ANSITEXTVIEW("Line"),          0)
+					.Field(ANSITEXTVIEW("FormatString"),  TEXTVIEW("%s"))
+					.Field(ANSITEXTVIEW("FileName"),      ANSITEXTVIEW("AudioInsightsBookmark"))
+					.End();
+
+				// Declare Bookmark event (written per cached bookmark message)
+				return Writer.DeclareEvent(ANSITEXTVIEW("Misc"), ANSITEXTVIEW("Bookmark"))
+					.Field(ANSITEXTVIEW("Cycle"),         UE::Trace::ETraceWriterFieldType::Uint64)
+					.Field(ANSITEXTVIEW("BookmarkPoint"), UE::Trace::ETraceWriterFieldType::Pointer)
+					.Field(ANSITEXTVIEW("FormatArgs"),    UE::Trace::ETraceWriterFieldType::Uint8 | UE::Trace::ETraceWriterFieldType::ArrayFlag)
+					.Field(ANSITEXTVIEW("CallstackId"),   UE::Trace::ETraceWriterFieldType::Uint32)
+					.End();
+			},
+			[](UE::Trace::FTraceWriter& Writer, uint32 EventId, const IAudioCachedMessage& BaseMsg, uint64 TimestampCycles)
+			{
+				const FBookmarkTraceMessage& Msg = static_cast<const FBookmarkTraceMessage&>(BaseMsg);
+
+				uint8 FormatArgsBuffer[4096];
+				const uint16 FormatArgsSize = FFormatArgsTrace::EncodeArguments(FormatArgsBuffer, *Msg.Label);
+
+				constexpr uint64 SyntheticBookmarkPoint = 1;
+
+				Writer.WriteEvent(EventId)
+					.Field(ANSITEXTVIEW("Cycle"),         TimestampCycles)
+					.Field(ANSITEXTVIEW("BookmarkPoint"), SyntheticBookmarkPoint)
+					.Field(ANSITEXTVIEW("FormatArgs"),    TConstArrayView<uint8>(FormatArgsBuffer, FormatArgsSize))
+					.Field(ANSITEXTVIEW("CallstackId"),   0u)
+					.End();
+			}
+		};
+	}
+} // namespace UE::Audio::Insights
diff --git a/Engine/Plugins/AudioInsights/Source/AudioInsights/Public/AudioInsightsDashboardFactory.h b/Engine/Plugins/AudioInsights/Source/AudioInsights/Public/AudioInsightsDashboardFactory.h
--- a/Engine/Plugins/AudioInsights/Source/AudioInsights/Public/AudioInsightsDashboardFactory.h
+++ b/Engine/Plugins/AudioInsights/Source/AudioInsights/Public/AudioInsightsDashboardFactory.h
@@ -28,6 +28,7 @@ namespace UE::Audio::Insights
 		TSharedRef<SWidget> MakeMenuBarWidget();
 		TSharedRef<SWidget> MakeMainToolbarWidget();
 		TSharedRef<SWidget> MakeCacheSnapshotButtonWidget();
+		TSharedRef<SWidget> MakeBookmarkButtonWidget();
 
 		TSharedPtr<FTabManager::FLayout> GetDefaultTabLayout();
 
@@ -39,6 +40,7 @@ namespace UE::Audio::Insights
 		TSharedRef<SWidget> MakeEnableTracesButton();
 		FReply ToggleAutoEnableAudioTraces();
 #endif
+		void CreateTraceBookmark();
 
 		TSharedPtr<FToolbarWidgets> ToolbarWidgets;
 
@@ -48,6 +50,8 @@ namespace UE::Audio::Insights
 
 		static constexpr ::Audio::FDeviceId ActiveDeviceId = 1; // The default audio device id in standalone
 
+		int32 BookmarkIndex = 0;
+
 		TMap<FName, TSharedPtr<IDashboardViewFactory>> DashboardViewFactories;
 	};
 } // namespace UE::Audio::Insights
diff --git a/Engine/Plugins/AudioInsights/Source/AudioInsights/Public/Messages/BookmarkTraceMessages.h b/Engine/Plugins/AudioInsights/Source/AudioInsights/Public/Messages/BookmarkTraceMessages.h
new file mode 100644
--- /dev/null
+++ b/Engine/Plugins/AudioInsights/Source/AudioInsights/Public/Messages/BookmarkTraceMessages.h
@@ -0,0 +1,28 @@
+// Copyright Epic Games, Inc. All Rights Reserved.
+#pragma once
+
+#include "Cache/IAudioCachedMessage.h"
+
+#define UE_API AUDIOINSIGHTS_API
+
+namespace UE::Audio::Insights
+{
+	namespace BookmarkMessageNames
+	{
+		extern UE_API const FName Bookmark;
+	};
+
+	struct FBookmarkTraceMessage : public IAudioCachedMessage
+	{
+		UE_API FBookmarkTraceMessage(double InTimestamp, const FString& InLabel);
+
+		virtual uint64 GetID() const override { return 0; }
+		virtual const FName GetMessageName() const override { return BookmarkMessageNames::Bookmark; }
+		virtual UE_API uint32 GetSizeOf() const override;
+		virtual UE_API FCacheWriteHandler GetCacheWriteHandler() const override;
+
+		FString Label;
+	};
+} // namespace UE::Audio::Insights
+
+#undef UE_API
diff --git a/Engine/Plugins/AudioInsights/Source/AudioInsightsEditor/Private/AudioInsightsEditorDashboardFactory.cpp b/Engine/Plugins/AudioInsights/Source/AudioInsightsEditor/Private/AudioInsightsEditorDashboardFactory.cpp
--- a/Engine/Plugins/AudioInsights/Source/AudioInsightsEditor/Private/AudioInsightsEditorDashboardFactory.cpp
+++ b/Engine/Plugins/AudioInsights/Source/AudioInsightsEditor/Private/AudioInsightsEditorDashboardFactory.cpp
@@ -9,14 +9,17 @@
 #include "AudioInsightsEditorSettings.h"
 #include "AudioInsightsStyle.h"
 #include "AudioInsightsTraceModule.h"
+#include "Cache/AudioInsightsCacheManager.h"
 #include "Engine/World.h"
 #include "Framework/Commands/UICommandList.h"
 #include "Framework/Docking/LayoutService.h"
 #include "Framework/Docking/TabManager.h"
 #include "Framework/Notifications/NotificationManager.h"
+#include "IAudioInsightsModuleInterface.h"
 #include "Internationalization/Text.h"
 #include "IPropertyTypeCustomization.h"
 #include "ISettingsModule.h"
+#include "Messages/BookmarkTraceMessages.h"
 #include "Templates/SharedPointer.h"
 #include "Widgets/Input/SButton.h"
 #include "Widgets/Input/SCheckBox.h"
@@ -196,7 +199,15 @@ namespace UE::Audio::Insights
 
 		if (TraceModule.IsTraceAnalysisActive())
 		{
-			TRACE_BOOKMARK(TEXT("%s %d"), *BookmarkLabelText.ToString(), BookmarkIndex);
+			const FString BookmarkLabel = FString::Printf(TEXT("%s %d"), *BookmarkLabelText.ToString(), BookmarkIndex);
+
+			TRACE_BOOKMARK(TEXT("%s"), *BookmarkLabel);
+
+			FAudioInsightsCacheManager& CacheManager = IAudioInsightsModule::GetChecked().GetCacheManager();
+
+			const double CurrentTimestamp = CacheManager.GetCacheEndTimeStamp();
+			CacheManager.AddMessageToCache(MakeShared<FBookmarkTraceMessage>(CurrentTimestamp, BookmarkLabel));
+
 			BookmarkIndex++;
 
 			// Show floating notification
