[Audio Insights] Increase the spacing between nodes in the signal flow graph

#Synced-CL 52724280

diff --git a/Engine/Plugins/AudioInsights/Source/AudioInsights/Private/Views/SSignalFlowGraphStyle.h b/Engine/Plugins/AudioInsights/Source/AudioInsights/Private/Views/SSignalFlowGraphStyle.h
--- a/Engine/Plugins/AudioInsights/Source/AudioInsights/Private/Views/SSignalFlowGraphStyle.h
+++ b/Engine/Plugins/AudioInsights/Source/AudioInsights/Private/Views/SSignalFlowGraphStyle.h
@@ -48,7 +48,7 @@ namespace UE::Audio::Insights
 		constexpr float NodeBorderPadding = 12.0f;
 		constexpr float NodeDetailVerticalPadding = 8.0f;
 		constexpr float NodeDetailIndentationPadding = 20.0f;
-		constexpr float LargeBetweenNodesPadding = 48.0f;
+		constexpr float LargeBetweenNodesPadding = 128.0f;
 		constexpr float SmallBetweenNodePadding = 8.0f;
 
 		constexpr float SendLabelVerticalWidth = 20.0f;
