[Audio Insights] Process BusPatch messages in FSignalFlowTraceProvider

#Synced-CL 52724280

diff --git a/Engine/Plugins/AudioInsights/Source/AudioInsights/Private/Messages/SignalFlowTraceMessages.cpp b/Engine/Plugins/AudioInsights/Source/AudioInsights/Private/Messages/SignalFlowTraceMessages.cpp
--- a/Engine/Plugins/AudioInsights/Source/AudioInsights/Private/Messages/SignalFlowTraceMessages.cpp
+++ b/Engine/Plugins/AudioInsights/Source/AudioInsights/Private/Messages/SignalFlowTraceMessages.cpp
@@ -11,6 +11,8 @@ namespace UE::Audio::Insights
 		const FName SoundWaveToBusSend = "SoundToBusSend";
 		const FName SoundWaveToSubmixSend = "SoundToSubmixSend";
 		const FName SubmixToSubmixSend = "SubmixToSubmixSend";
+		const FName BusPatchWriterConnected = "BusPatchWriterConnected";
+		const FName BusPatchReaderConnected = "BusPatchReaderConnected";
 	};
 
 	namespace FSignalFlowDashboardEntryPrivate
diff --git a/Engine/Plugins/AudioInsights/Source/AudioInsights/Private/Providers/SignalFlowTraceProvider.cpp b/Engine/Plugins/AudioInsights/Source/AudioInsights/Private/Providers/SignalFlowTraceProvider.cpp
--- a/Engine/Plugins/AudioInsights/Source/AudioInsights/Private/Providers/SignalFlowTraceProvider.cpp
+++ b/Engine/Plugins/AudioInsights/Source/AudioInsights/Private/Providers/SignalFlowTraceProvider.cpp
@@ -85,6 +85,8 @@ namespace UE::Audio::Insights
 				Builder.RouteEvent(RouteId_SubmixUnloaded, "Audio", "SubmixUnloaded");
 
 				Builder.RouteEvent(RouteId_SoundToBusSend, "Audio", "SoundToBusSend");
+				Builder.RouteEvent(RouteId_BusPatchWriterConnected, "Audio", "BusPatchWriterConnected");
+				Builder.RouteEvent(RouteId_BusPatchReaderConnected, "Audio", "BusPatchReaderConnected");
 				Builder.RouteEvent(RouteId_SoundToSubmixSend, "Audio", "SoundToSubmixSend");
 				Builder.RouteEvent(RouteId_SubmixToSubmixSend, "Audio", "SubmixToSubmixSend");
 
@@ -241,6 +243,18 @@ namespace UE::Audio::Insights
 						break;
 					}
 
+					case RouteId_BusPatchWriterConnected:
+					{
+						CacheMessage<FBusPatchWriterConnectedMessage>(Context, Messages.BusPatchWriterConnectedMessages);
+						break;
+					}
+
+					case RouteId_BusPatchReaderConnected:
+					{
+						CacheMessage<FBusPatchReaderConnectedMessage>(Context, Messages.BusPatchReaderConnectedMessages);
+						break;
+					}
+
 					case RouteId_SoundToSubmixSend:
 					{
 						CacheMessage<FSoundToSubmixSendMessage>(Context, Messages.SoundToSubmixSendMessages);
@@ -284,6 +298,8 @@ namespace UE::Audio::Insights
 				RouteId_SubmixUnloaded,
 
 				RouteId_SoundToBusSend,
+				RouteId_BusPatchWriterConnected,
+				RouteId_BusPatchReaderConnected,
 				RouteId_SoundToSubmixSend,
 				RouteId_SubmixToSubmixSend,
 
@@ -349,6 +365,11 @@ namespace UE::Audio::Insights
 			return FindDeviceEntry(Msg.DeviceId, { Msg.DeviceId, ESignalFlowEntryType::AudioBus, Msg.AudioBusId });
 		};
 
+		auto GetBusPatchReaderEntry = [this](const FBusPatchReaderConnectedMessage& Msg)
+		{
+			return FindDeviceEntry(Msg.DeviceId, { Msg.DeviceId, ESignalFlowEntryType::AudioBus, Msg.SenderID });
+		};
+
 		auto GetSubmixEntry = [this](const FSubmixMessageBase& Msg)
 		{
 			return FindDeviceEntry(Msg.DeviceId, { Msg.DeviceId, ESignalFlowEntryType::Submix, Msg.SubmixId });
@@ -454,6 +475,18 @@ namespace UE::Audio::Insights
 			}
 		});
 
+		ProcessMessageQueue<FBusPatchWriterConnectedMessage>(TraceMessages.BusPatchWriterConnectedMessages, GetEntryFromSendMessage,
+		[this](const FBusPatchWriterConnectedMessage& Msg, TSharedPtr<FSignalFlowDashboardEntry>* OutEntry)
+		{
+			ProcessBusPatchWriterMessage(Msg, OutEntry);
+		});
+
+		ProcessMessageQueue<FBusPatchReaderConnectedMessage>(TraceMessages.BusPatchReaderConnectedMessages, GetBusPatchReaderEntry,
+		[this](const FBusPatchReaderConnectedMessage& Msg, TSharedPtr<FSignalFlowDashboardEntry>* OutEntry)
+		{
+			ProcessBusPatchReaderMessage(Msg, OutEntry);
+		});
+
 		ProcessMessageQueue<FSubmixToSubmixSendMessage>(TraceMessages.SubmixToSubmixSendMessages, GetSubmixToSubmixEntry,
 		[this](const FSubmixToSubmixSendMessage& Msg, TSharedPtr<FSignalFlowDashboardEntry>* OutEntry)
 		{
@@ -733,6 +766,23 @@ namespace UE::Audio::Insights
 			}
 		});
 
+		/*
+		 * Bus patch connections
+		*/
+		CacheManager.IterateOverRange<FBusPatchWriterConnectedMessage>(SignalFlowMessageNames::BusPatchWriterConnected, RangeCheckStart, TimeMarker,
+		[this](const FBusPatchWriterConnectedMessage& Msg)
+		{
+			TSharedPtr<FSignalFlowDashboardEntry>* Entry = FindDeviceEntry(Msg.DeviceId, { Msg.DeviceId, ESignalFlowEntryType::SoundSource, Msg.SenderID });
+			ProcessBusPatchWriterMessage(Msg, Entry);
+		});
+
+		CacheManager.IterateOverRange<FBusPatchReaderConnectedMessage>(SignalFlowMessageNames::BusPatchReaderConnected, RangeCheckStart, TimeMarker,
+		[this](const FBusPatchReaderConnectedMessage& Msg)
+		{
+			TSharedPtr<FSignalFlowDashboardEntry>* Entry = FindDeviceEntry(Msg.DeviceId, { Msg.DeviceId, ESignalFlowEntryType::AudioBus, Msg.SenderID });
+			ProcessBusPatchReaderMessage(Msg, Entry);
+		});
+
 		/*
 		 * Send levels/remove unreferenced entries
 		*/
@@ -1195,6 +1245,64 @@ namespace UE::Audio::Insights
 		}
 	}
 
+	// Writer: a source writes to a bus. OutEntry is the source.
+	// Populates the source's LinkedBusPatchOutputs (bus keys) and the bus's LinkedBusPatchInputs (source key).
+	void FSignalFlowTraceProvider::ProcessBusPatchWriterMessage(const FBusPatchWriterConnectedMessage& Msg, TSharedPtr<FSignalFlowDashboardEntry>* OutEntry)
+	{
+		if (OutEntry && OutEntry->IsValid())
+		{
+			const FSignalFlowEntryKey SourceEntryKey = (*OutEntry)->GetSignalFlowEntryKey();
+
+			for (const FSignalFlowSendMessageBase::FSendInfoPair& BusPatchInput : Msg.SendInfos)
+			{
+				const FSignalFlowEntryKey BusEntryKey{ Msg.DeviceId, ESignalFlowEntryType::AudioBus, BusPatchInput.ReceiverID };
+
+				bool bIsAlreadySet = false;
+				(*OutEntry)->LinkedBusPatchOutputs.FindOrAdd(BusEntryKey, &bIsAlreadySet);
+
+				if (!bIsAlreadySet)
+				{
+					bForceGraphRefresh = true;
+				}
+
+				TSharedPtr<FSignalFlowDashboardEntry>* BusEntry = FindDeviceEntry(Msg.DeviceId, BusEntryKey);
+				if (BusEntry && BusEntry->IsValid())
+				{
+					(*BusEntry)->LinkedBusPatchInputs.Add(SourceEntryKey);
+				}
+			}
+		}
+	}
+
+	// Reader: a bus outputs to a source. OutEntry is the bus.
+	// Populates the bus's LinkedBusPatchOutputs (source keys) and the source's LinkedBusPatchInputs (bus key).
+	void FSignalFlowTraceProvider::ProcessBusPatchReaderMessage(const FBusPatchReaderConnectedMessage& Msg, TSharedPtr<FSignalFlowDashboardEntry>* OutEntry)
+	{
+		if (OutEntry && OutEntry->IsValid())
+		{
+			const FSignalFlowEntryKey BusEntryKey = (*OutEntry)->GetSignalFlowEntryKey();
+
+			for (const FSignalFlowSendMessageBase::FSendInfoPair& BusPatchOutput : Msg.SendInfos)
+			{
+				const FSignalFlowEntryKey SourceEntryKey{ Msg.DeviceId, ESignalFlowEntryType::SoundSource, BusPatchOutput.ReceiverID };
+
+				bool bIsAlreadySet = false;
+				(*OutEntry)->LinkedBusPatchOutputs.FindOrAdd(SourceEntryKey, &bIsAlreadySet);
+
+				if (!bIsAlreadySet)
+				{
+					bForceGraphRefresh = true;
+				}
+
+				TSharedPtr<FSignalFlowDashboardEntry>* SourceEntry = FindDeviceEntry(Msg.DeviceId, SourceEntryKey);
+				if (SourceEntry && SourceEntry->IsValid())
+				{
+					(*SourceEntry)->LinkedBusPatchInputs.Add(BusEntryKey);
+				}
+			}
+		}
+	}
+
 	void FSignalFlowTraceProvider::ProcessMultiChannelEnvelopeMessage(TSharedPtr<FSignalFlowDashboardEntry>* OutEntry, const TArray<float>& EnvelopeValues)
 	{
 		if (OutEntry == nullptr || !OutEntry->IsValid() || EnvelopeValues.IsEmpty())
@@ -1262,6 +1370,24 @@ namespace UE::Audio::Insights
 						Itr.RemoveCurrent();
 					}
 				}
+
+				// Clean stale bus patch input connections
+				for (auto Itr = Entry->LinkedBusPatchInputs.CreateIterator(); Itr; ++Itr)
+				{
+					if (RemovedKeys.Contains(*Itr))
+					{
+						Itr.RemoveCurrent();
+					}
+				}
+
+				// Clean stale bus patch output connections
+				for (auto Itr = Entry->LinkedBusPatchOutputs.CreateIterator(); Itr; ++Itr)
+				{
+					if (RemovedKeys.Contains(*Itr))
+					{
+						Itr.RemoveCurrent();
+					}
+				}
 			}
 		}
 	}
diff --git a/Engine/Plugins/AudioInsights/Source/AudioInsights/Private/Providers/SignalFlowTraceProvider.h b/Engine/Plugins/AudioInsights/Source/AudioInsights/Private/Providers/SignalFlowTraceProvider.h
--- a/Engine/Plugins/AudioInsights/Source/AudioInsights/Private/Providers/SignalFlowTraceProvider.h
+++ b/Engine/Plugins/AudioInsights/Source/AudioInsights/Private/Providers/SignalFlowTraceProvider.h
@@ -48,6 +48,8 @@ namespace UE::Audio::Insights
 		void ProcessStopMessage(const ::Audio::FDeviceId DeviceID, const double MsgTimestamp, TSharedPtr<FSignalFlowDashboardEntry>* OutEntry, TSet<FSignalFlowEntryKey>& OutRemovedKeys);
 
 		void ProcessOutputSendMessage(TSharedPtr<FSignalFlowDashboardEntry>* OutEntry, const FSignalFlowEntryKey& OutputEntryKey, const float SendLevel);
+		void ProcessBusPatchWriterMessage(const FBusPatchWriterConnectedMessage& Msg, TSharedPtr<FSignalFlowDashboardEntry>* OutEntry);
+		void ProcessBusPatchReaderMessage(const FBusPatchReaderConnectedMessage& Msg, TSharedPtr<FSignalFlowDashboardEntry>* OutEntry);
 		void ProcessMultiChannelEnvelopeMessage(TSharedPtr<FSignalFlowDashboardEntry>* OutEntry, const TArray<float>& EnvelopeValues);
 
 		void LinkInput(const FSignalFlowDashboardEntry& Entry, const FSignalFlowEntryKey& OutputKey, const ::Audio::FDeviceId DeviceID);
diff --git a/Engine/Plugins/AudioInsights/Source/AudioInsights/Public/Messages/SignalFlowTraceMessages.h b/Engine/Plugins/AudioInsights/Source/AudioInsights/Public/Messages/SignalFlowTraceMessages.h
--- a/Engine/Plugins/AudioInsights/Source/AudioInsights/Public/Messages/SignalFlowTraceMessages.h
+++ b/Engine/Plugins/AudioInsights/Source/AudioInsights/Public/Messages/SignalFlowTraceMessages.h
@@ -18,6 +18,8 @@ namespace UE::Audio::Insights
 		extern const FName SoundWaveToBusSend;
 		extern const FName SoundWaveToSubmixSend;
 		extern const FName SubmixToSubmixSend;
+		extern const FName BusPatchWriterConnected;
+		extern const FName BusPatchReaderConnected;
 	};
 
 	struct FSignalFlowMessageBase : public IAudioCachedMessage
@@ -58,6 +60,30 @@ namespace UE::Audio::Insights
 		virtual FCacheWriteHandler GetCacheWriteHandler() const override;
 	};
 
+	struct FBusPatchWriterConnectedMessage : public FSignalFlowSendMessageBase
+	{
+		FBusPatchWriterConnectedMessage() = default;
+		FBusPatchWriterConnectedMessage(const Trace::IAnalyzer::FOnEventContext& InContext)
+			: FSignalFlowSendMessageBase(InContext)
+		{
+		}
+
+		virtual uint64 GetID() const override { return GetTypeHash(FSignalFlowEntryKey{ DeviceId, ESignalFlowEntryType::SoundSource, SenderID }); }
+		virtual const FName GetMessageName() const override { return SignalFlowMessageNames::BusPatchWriterConnected; }
+	};
+
+	struct FBusPatchReaderConnectedMessage : public FSignalFlowSendMessageBase
+	{
+		FBusPatchReaderConnectedMessage() = default;
+		FBusPatchReaderConnectedMessage(const Trace::IAnalyzer::FOnEventContext& InContext)
+			: FSignalFlowSendMessageBase(InContext)
+		{
+		}
+
+		virtual uint64 GetID() const override { return GetTypeHash(FSignalFlowEntryKey{ DeviceId, ESignalFlowEntryType::AudioBus, SenderID }); }
+		virtual const FName GetMessageName() const override { return SignalFlowMessageNames::BusPatchReaderConnected; }
+	};
+
 	struct FSoundToSubmixSendMessage : public FSignalFlowSendMessageBase
 	{
 		FSoundToSubmixSendMessage() = default;
@@ -101,6 +127,8 @@ namespace UE::Audio::Insights
 		TAnalyzerMessageQueue<FSubmixUnloadedMessage> SubmixUnloadedMessages;
 
 		TAnalyzerMessageQueue<FSoundToBusSendMessage> SoundToBusSendMessages;
+		TAnalyzerMessageQueue<FBusPatchWriterConnectedMessage> BusPatchWriterConnectedMessages;
+		TAnalyzerMessageQueue<FBusPatchReaderConnectedMessage> BusPatchReaderConnectedMessages;
 		TAnalyzerMessageQueue<FSoundToSubmixSendMessage> SoundToSubmixSendMessages;
 		TAnalyzerMessageQueue<FSubmixToSubmixSendMessage> SubmixToSubmixSendMessages;
 
@@ -160,6 +188,9 @@ namespace UE::Audio::Insights
 		TOptional<FSignalFlowEntryKey> LinkedSourceBus;
 		TSet<FSignalFlowEntryKey> LinkedSoundSources;
 
+		TSet<FSignalFlowEntryKey> LinkedBusPatchInputs;   // Entries that send bus patch connections to this entry
+		TSet<FSignalFlowEntryKey> LinkedBusPatchOutputs;   // Entries that this entry sends bus patch connections to
+
 		TCircularBuffer<float> AmplitudeDataRange;
 		int32 AmplitudeWriteIndex = 0;
 
