[Audio Insights] Add BusPatch connection rendering to Signal Flow dashboard and graph widget

#Synced-CL 52724280

diff --git a/Engine/Plugins/AudioInsights/Source/AudioInsights/Private/Views/SSignalFlowGraph.cpp b/Engine/Plugins/AudioInsights/Source/AudioInsights/Private/Views/SSignalFlowGraph.cpp
--- a/Engine/Plugins/AudioInsights/Source/AudioInsights/Private/Views/SSignalFlowGraph.cpp
+++ b/Engine/Plugins/AudioInsights/Source/AudioInsights/Private/Views/SSignalFlowGraph.cpp
@@ -374,6 +374,32 @@ namespace UE::Audio::Insights
 				}
 			}
 
+			// Draw any bus patch output connections
+			for (int32 BusPatchOutputIndex = 0; BusPatchOutputIndex < Node->FilteredLinkedBusPatchOutputs.Num(); ++BusPatchOutputIndex)
+			{
+				const FSignalFlowEntryKey& BusPatchOutputKey = Node->FilteredLinkedBusPatchOutputs[BusPatchOutputIndex];
+				const TSharedPtr<SSignalFlowGraphNode>* BusPatchOutputNode = NodeWidgets.Find(BusPatchOutputKey);
+
+				if (BusPatchOutputNode && BusPatchOutputNode->IsValid())
+				{
+					constexpr float SendLevel = 1.0f;
+					const int32 SendLabelIndex = BusPatchOutputIndex + Node->FilteredOutputs.Num() + Node->FilteredLinkedSoundSources.Num();
+					const bool bIsBusConnection = Entry.IsValid() && Entry->EntryType == ESignalFlowEntryType::AudioBus;
+
+					LayerId = PaintConnectionToNode(AllottedGeometry,
+													OutDrawElements,
+													LayerId,
+													Entry,
+													**InputNode,
+													**BusPatchOutputNode,
+													LocalBoundsBox,
+													ViewportPosInGraph,
+													SendLevel,
+													SendLabelIndex,
+													bIsBusConnection);
+				}
+			}
+
 			// Now draw any output connections (flowing forwards through the graph)
 			for (int32 OutputIndex = 0; OutputIndex < Node->FilteredOutputs.Num(); ++OutputIndex)
 			{
@@ -732,6 +758,11 @@ namespace UE::Audio::Insights
 				NodeWidget->AddSendLabelWidget(LinkedSoundSourceKey);
 			}
 
+			for (const FSignalFlowEntryKey& BusPatchOutputKey : Node->FilteredLinkedBusPatchOutputs)
+			{
+				NodeWidget->AddSendLabelWidget(BusPatchOutputKey);
+			}
+
 			MaxNumNodeOutputsInGraph = FMath::Max(MaxNumNodeOutputsInGraph, Node->FilteredOutputs.Num());
 		}
 	}
@@ -779,6 +810,26 @@ namespace UE::Audio::Insights
 						(*SourceBusNode)->RemoveSendLabelWidget(Key);
 					}
 				}
+
+				// Clean up bus patch input connections
+				for (const FSignalFlowEntryKey& BusPatchInputKey : RemovedEntry->LinkedBusPatchInputs)
+				{
+					TSharedPtr<SSignalFlowGraphNode>* BusNode = NodeWidgets.Find(BusPatchInputKey);
+					if (BusNode && BusNode->IsValid())
+					{
+						(*BusNode)->RemoveSendLabelWidget(Key);
+					}
+				}
+
+				// Clean up bus patch output connections
+				for (const FSignalFlowEntryKey& BusPatchOutputKey : RemovedEntry->LinkedBusPatchOutputs)
+				{
+					TSharedPtr<SSignalFlowGraphNode>* SourceNode = NodeWidgets.Find(BusPatchOutputKey);
+					if (SourceNode && SourceNode->IsValid())
+					{
+						(*SourceNode)->RemoveSendLabelWidget(Key);
+					}
+				}
 			}
 
 			NodeWidgets.Remove(Key);
@@ -952,6 +1003,11 @@ namespace UE::Audio::Insights
 					ActiveOutputKeys.AddUnique(LinkedSoundSourceKey);
 				}
 
+				for (const FSignalFlowEntryKey& BusPatchOutputKey : Node->FilteredLinkedBusPatchOutputs)
+				{
+					ActiveOutputKeys.AddUnique(BusPatchOutputKey);
+				}
+
 				(*NodeWidget)->SetSendLabelWidgets(ActiveOutputKeys);
 
 				return;
diff --git a/Engine/Plugins/AudioInsights/Source/AudioInsights/Private/Views/SSignalFlowGraphNode.cpp b/Engine/Plugins/AudioInsights/Source/AudioInsights/Private/Views/SSignalFlowGraphNode.cpp
--- a/Engine/Plugins/AudioInsights/Source/AudioInsights/Private/Views/SSignalFlowGraphNode.cpp
+++ b/Engine/Plugins/AudioInsights/Source/AudioInsights/Private/Views/SSignalFlowGraphNode.cpp
@@ -402,9 +402,11 @@ namespace UE::Audio::Insights
 					{
 						return SendInfo->SendLevel;
 					}
-					else if (Entry->LinkedSoundSources.Contains(ReceiverEntryKey))
+					else if (Entry->LinkedSoundSources.Contains(ReceiverEntryKey)
+						  || Entry->LinkedBusPatchOutputs.Contains(ReceiverEntryKey)
+						  || Entry->LinkedBusPatchInputs.Contains(ReceiverEntryKey))
 					{
-						// Source buses always have a send value of 1.0
+						// Source buses, bus patch outputs and bus patch inputs always have a send value of 1.0
 						return 1.0f;
 					}
 				}
diff --git a/Engine/Plugins/AudioInsights/Source/AudioInsights/Private/Views/SignalFlowDashboardViewFactory.cpp b/Engine/Plugins/AudioInsights/Source/AudioInsights/Private/Views/SignalFlowDashboardViewFactory.cpp
--- a/Engine/Plugins/AudioInsights/Source/AudioInsights/Private/Views/SignalFlowDashboardViewFactory.cpp
+++ b/Engine/Plugins/AudioInsights/Source/AudioInsights/Private/Views/SignalFlowDashboardViewFactory.cpp
@@ -1669,6 +1669,7 @@ namespace UE::Audio::Insights
 		}
 
 		AddFilteredLinkedSourceBus(DeviceData, Entry, *EntryNode, OutExistingEntryKeys);
+		AddFilteredLinkedBusPatchInputs(DeviceData, Entry, *EntryNode, OutExistingEntryKeys);
 	}
 
 	void FSignalFlowDashboardViewFactory::AddFilteredOutputs(const TSortedMap<FSignalFlowEntryKey, TSharedPtr<FSignalFlowDashboardEntry>>& DeviceData, TSharedPtr<FSignalFlowDashboardEntry> Entry, int32 TreeDepth, TSet<FSignalFlowEntryKey>& OutExistingEntryKeys)
@@ -1693,6 +1694,7 @@ namespace UE::Audio::Insights
 		if (EntryNode && EntryNode->IsValid())
 		{
 			AddFilteredLinkedSources(DeviceData, Entry, **EntryNode, OutExistingEntryKeys);
+			AddFilteredLinkedBusPatchOutputs(DeviceData, Entry, **EntryNode, OutExistingEntryKeys);
 		}
 	}
 
@@ -1762,6 +1764,7 @@ namespace UE::Audio::Insights
 		}
 
 		AddFilteredLinkedSources(DeviceData, Entry, *Node, OutExistingEntryKeys);
+		AddFilteredLinkedBusPatchOutputs(DeviceData, Entry, *Node, OutExistingEntryKeys);
 	}
 
 	void FSignalFlowDashboardViewFactory::AddFilteredLinkedSourceBus(const TSortedMap<FSignalFlowEntryKey, TSharedPtr<FSignalFlowDashboardEntry>>& DeviceData, TSharedPtr<FSignalFlowDashboardEntry> Entry, ISignalFlowNode& OutEntryNode, TSet<FSignalFlowEntryKey>& OutExistingEntryKeys)
@@ -1806,23 +1809,93 @@ namespace UE::Audio::Insights
 
 	void FSignalFlowDashboardViewFactory::AddFilteredLinkedSources(const TSortedMap<FSignalFlowEntryKey, TSharedPtr<FSignalFlowDashboardEntry>>& DeviceData, TSharedPtr<FSignalFlowDashboardEntry> Entry, ISignalFlowNode& OutEntryNode, TSet<FSignalFlowEntryKey>& OutExistingEntryKeys)
 	{
-		if (!Entry.IsValid() || Entry->LinkedSoundSources.IsEmpty())
+		const FSignalFlowEntryKey EntryKey = Entry->GetSignalFlowEntryKey();
+
+		AddFilteredLinkedEntries(DeviceData,
+								 Entry,
+								 Entry->LinkedSoundSources,
+								 OutEntryNode.FilteredLinkedSoundSources,
+								 [&EntryKey](ISignalFlowNode& Node) { Node.FilteredLinkedSourceBus = EntryKey; },
+								 OutExistingEntryKeys);
+	}
+
+	// Does not use AddFilteredLinkedEntries because the traversal direction is inverted:
+	// we walk inputs-then-outputs (upstream-first) to discover the bus's full graph context,
+	// whereas the shared helper walks outputs-then-inputs (downstream-first).
+	void FSignalFlowDashboardViewFactory::AddFilteredLinkedBusPatchInputs(const TSortedMap<FSignalFlowEntryKey, TSharedPtr<FSignalFlowDashboardEntry>>& DeviceData, TSharedPtr<FSignalFlowDashboardEntry> Entry, ISignalFlowNode& OutEntryNode, TSet<FSignalFlowEntryKey>& OutExistingEntryKeys)
+	{
+		if (!Entry.IsValid() || Entry->LinkedBusPatchInputs.IsEmpty())
+		{
+			return;
+		}
+
+		for (const FSignalFlowEntryKey& BusPatchInputKey : Entry->LinkedBusPatchInputs)
+		{
+			const TSharedPtr<FSignalFlowDashboardEntry>* BusEntry = DeviceData.Find(BusPatchInputKey);
+			if (!BusEntry || !BusEntry->IsValid())
+			{
+				continue;
+			}
+
+			OutEntryNode.FilteredLinkedBusPatchInputs.AddUnique(BusPatchInputKey);
+
+			if (!EntryIsInGraph(BusPatchInputKey, OutExistingEntryKeys))
+			{
+				AddFilteredInputsRecursive(DeviceData, *BusEntry, nullptr, 0, OutExistingEntryKeys);
+
+				for (const auto& [OutputKey, SendInfo] : (*BusEntry)->Outputs)
+				{
+					const TSharedPtr<FSignalFlowDashboardEntry>* OutputEntry = DeviceData.Find(OutputKey);
+					if (OutputEntry && OutputEntry->IsValid())
+					{
+						AddFilteredOutputsRecursive(DeviceData, *OutputEntry, *BusEntry, 0, OutExistingEntryKeys);
+					}
+				}
+			}
+
+			const TSharedPtr<ISignalFlowNode>* LinkedBusNode = FilteredEntries.Find(BusPatchInputKey);
+			if (LinkedBusNode && LinkedBusNode->IsValid())
+			{
+				(*LinkedBusNode)->FilteredLinkedBusPatchOutputs.AddUnique(Entry->GetSignalFlowEntryKey());
+			}
+		}
+	}
+
+	void FSignalFlowDashboardViewFactory::AddFilteredLinkedBusPatchOutputs(const TSortedMap<FSignalFlowEntryKey, TSharedPtr<FSignalFlowDashboardEntry>>& DeviceData, TSharedPtr<FSignalFlowDashboardEntry> Entry, ISignalFlowNode& OutEntryNode, TSet<FSignalFlowEntryKey>& OutExistingEntryKeys)
+	{
+		if (!Entry.IsValid())
 		{
 			return;
 		}
 
-		// Cycle through this source bus's linked sound sources, detect any new paths that need adding to the filtered entries
-		for (const FSignalFlowEntryKey& LinkedSourceEntryKey : Entry->LinkedSoundSources)
+		const FSignalFlowEntryKey EntryKey = Entry->GetSignalFlowEntryKey();
+
+		AddFilteredLinkedEntries(DeviceData,
+								 Entry,
+								 Entry->LinkedBusPatchOutputs,
+								 OutEntryNode.FilteredLinkedBusPatchOutputs,
+								 [&EntryKey](ISignalFlowNode& Node) { Node.FilteredLinkedBusPatchInputs.AddUnique(EntryKey); },
+								 OutExistingEntryKeys);
+	}
+
+	void FSignalFlowDashboardViewFactory::AddFilteredLinkedEntries(const TSortedMap<FSignalFlowEntryKey, TSharedPtr<FSignalFlowDashboardEntry>>& DeviceData, TSharedPtr<FSignalFlowDashboardEntry> Entry, const TSet<FSignalFlowEntryKey>& LinkedEntryKeys, TArray<FSignalFlowEntryKey>& OutFilteredLinkedEntries, TFunctionRef<void(ISignalFlowNode&)> LinkBackFunction, TSet<FSignalFlowEntryKey>& OutExistingEntryKeys)
+	{
+		if (!Entry.IsValid() || LinkedEntryKeys.IsEmpty())
+		{
+			return;
+		}
+
+		for (const FSignalFlowEntryKey& LinkedSourceEntryKey : LinkedEntryKeys)
 		{
 			const TSharedPtr<FSignalFlowDashboardEntry>* LinkedSourceEntry = DeviceData.Find(LinkedSourceEntryKey);
 			if (LinkedSourceEntry && LinkedSourceEntry->IsValid())
 			{
-				OutEntryNode.FilteredLinkedSoundSources.AddUnique(LinkedSourceEntryKey);
+				OutFilteredLinkedEntries.AddUnique(LinkedSourceEntryKey);
 
 				if (!EntryIsInGraph(LinkedSourceEntryKey, OutExistingEntryKeys))
 				{
-					// Detected a new sound source connected to this source bus entry
-					// Filter the graph from the found sound source to make sure it's full path is present in the graph
+					// Detected a new source connected to this entry
+					// Filter the graph from the found source to make sure it's full path is present in the graph
 					AddFilteredOutputsRecursive(DeviceData, *LinkedSourceEntry, nullptr, 0, OutExistingEntryKeys);
 
 					for (const FSignalFlowEntryKey& InputKey : (*LinkedSourceEntry)->Inputs)
@@ -1835,11 +1908,11 @@ namespace UE::Audio::Insights
 					}
 				}
 
-				// Make sure the linked sound source is connected to this source bus entry
+				// Make sure the linked source is connected to this entry
 				const TSharedPtr<ISignalFlowNode>* LinkedSourceNode = FilteredEntries.Find(LinkedSourceEntryKey);
 				if (LinkedSourceNode && LinkedSourceNode->IsValid())
 				{
-					(*LinkedSourceNode)->FilteredLinkedSourceBus = Entry->GetSignalFlowEntryKey();
+					LinkBackFunction(**LinkedSourceNode);
 				}
 			}
 		}
@@ -1866,7 +1939,10 @@ namespace UE::Audio::Insights
 			{
 				Entry->FilteredLinkedSourceBus.Reset();
 			}
+
 			Entry->FilteredLinkedSoundSources.RemoveAll([&EntriesToRemove](const FSignalFlowEntryKey& LinkedSource) { return EntriesToRemove.Contains(LinkedSource); });
+			Entry->FilteredLinkedBusPatchInputs.RemoveAll([&EntriesToRemove](const FSignalFlowEntryKey& Key) { return EntriesToRemove.Contains(Key); });
+			Entry->FilteredLinkedBusPatchOutputs.RemoveAll([&EntriesToRemove](const FSignalFlowEntryKey& Key) { return EntriesToRemove.Contains(Key); });
 		}
 
 		for (const FSignalFlowEntryKey& EntryKey : EntriesToRemove)
diff --git a/Engine/Plugins/AudioInsights/Source/AudioInsights/Public/Views/SignalFlowDashboardViewFactory.h b/Engine/Plugins/AudioInsights/Source/AudioInsights/Public/Views/SignalFlowDashboardViewFactory.h
--- a/Engine/Plugins/AudioInsights/Source/AudioInsights/Public/Views/SignalFlowDashboardViewFactory.h
+++ b/Engine/Plugins/AudioInsights/Source/AudioInsights/Public/Views/SignalFlowDashboardViewFactory.h
@@ -118,9 +118,13 @@ namespace UE::Audio::Insights
 		void AddFilteredInputsRecursive(const TSortedMap<FSignalFlowEntryKey, TSharedPtr<FSignalFlowDashboardEntry>>& DeviceData, TSharedPtr<FSignalFlowDashboardEntry> Entry, TSharedPtr<FSignalFlowDashboardEntry> ChildEntry, int32 TreeDepth, TSet<FSignalFlowEntryKey>& OutExistingEntryKeys);
 		void AddFilteredOutputs(const TSortedMap<FSignalFlowEntryKey, TSharedPtr<FSignalFlowDashboardEntry>>& DeviceData, TSharedPtr<FSignalFlowDashboardEntry> Entry, int32 TreeDepth, TSet<FSignalFlowEntryKey>& OutExistingEntryKeys);
 		void AddFilteredOutputsRecursive(const TSortedMap<FSignalFlowEntryKey, TSharedPtr<FSignalFlowDashboardEntry>>& DeviceData, TSharedPtr<FSignalFlowDashboardEntry> Entry, TSharedPtr<FSignalFlowDashboardEntry> ParentEntry, int32 TreeDepth, TSet<FSignalFlowEntryKey>& OutExistingEntryKeys);
+		
 		void AddFilteredLinkedSourceBus(const TSortedMap<FSignalFlowEntryKey, TSharedPtr<FSignalFlowDashboardEntry>>& DeviceData, TSharedPtr<FSignalFlowDashboardEntry> Entry, ISignalFlowNode& OutEntryNode, TSet<FSignalFlowEntryKey>& OutExistingEntryKeys);
 		void AddFilteredLinkedSources(const TSortedMap<FSignalFlowEntryKey, TSharedPtr<FSignalFlowDashboardEntry>>& DeviceData, TSharedPtr<FSignalFlowDashboardEntry> Entry, ISignalFlowNode& OutEntryNode, TSet<FSignalFlowEntryKey>& OutExistingEntryKeys);
-		
+		void AddFilteredLinkedBusPatchInputs(const TSortedMap<FSignalFlowEntryKey, TSharedPtr<FSignalFlowDashboardEntry>>& DeviceData, TSharedPtr<FSignalFlowDashboardEntry> Entry, ISignalFlowNode& OutEntryNode, TSet<FSignalFlowEntryKey>& OutExistingEntryKeys);
+		void AddFilteredLinkedBusPatchOutputs(const TSortedMap<FSignalFlowEntryKey, TSharedPtr<FSignalFlowDashboardEntry>>& DeviceData, TSharedPtr<FSignalFlowDashboardEntry> Entry, ISignalFlowNode& OutEntryNode, TSet<FSignalFlowEntryKey>& OutExistingEntryKeys);
+		void AddFilteredLinkedEntries(const TSortedMap<FSignalFlowEntryKey, TSharedPtr<FSignalFlowDashboardEntry>>& DeviceData, TSharedPtr<FSignalFlowDashboardEntry> Entry, const TSet<FSignalFlowEntryKey>& LinkedEntryKeys, TArray<FSignalFlowEntryKey>& OutFilteredLinkedEntries, TFunctionRef<void(ISignalFlowNode&)> LinkBackFunction, TSet<FSignalFlowEntryKey>& OutExistingEntryKeys);
+
 		void CleanUpFilteredEntries(const TSet<FSignalFlowEntryKey>& EntriesToRemove);
 
 		void AddMissingNodesToGraph(const TSortedMap<FSignalFlowEntryKey, TSharedPtr<FSignalFlowDashboardEntry>>& DeviceData, TSet<FSignalFlowEntryKey>& OutEntriesMarkedToDelete);
diff --git a/Engine/Plugins/AudioInsights/Source/AudioInsights/Public/Views/SignalFlowNodes.h b/Engine/Plugins/AudioInsights/Source/AudioInsights/Public/Views/SignalFlowNodes.h
--- a/Engine/Plugins/AudioInsights/Source/AudioInsights/Public/Views/SignalFlowNodes.h
+++ b/Engine/Plugins/AudioInsights/Source/AudioInsights/Public/Views/SignalFlowNodes.h
@@ -96,6 +96,9 @@ namespace UE::Audio::Insights
 		TArray<FSignalFlowEntryKey> FilteredLinkedSoundSources;
 		TOptional<FSignalFlowEntryKey> FilteredLinkedSourceBus;
 
+		TArray<FSignalFlowEntryKey> FilteredLinkedBusPatchInputs;   // Entries that send bus patch connections to this node
+		TArray<FSignalFlowEntryKey> FilteredLinkedBusPatchOutputs;  // Entries that this node sends bus patch connections to
+
 		float XPos = 0.0f;
 
 	private:
