[Audio Insights] Add Clear All Mutes/Solos toolbar button to Signal Flow Dashboard

Adds a new toolbar action button to the Signal Flow Dashboard that clears all assigned mute/solo states via AudioDeviceManager::GetDebugger(). The button is positioned after the existing button actions section and is guarded by MUTE_SOLO_ENABLED (WITH_EDITOR && ENABLE_AUDIO_DEBUG).

#Synced-CL 52724280

diff --git a/Engine/Plugins/AudioInsights/Source/AudioInsights/Private/Views/SignalFlowDashboardViewFactory.cpp b/Engine/Plugins/AudioInsights/Source/AudioInsights/Private/Views/SignalFlowDashboardViewFactory.cpp
--- a/Engine/Plugins/AudioInsights/Source/AudioInsights/Private/Views/SignalFlowDashboardViewFactory.cpp
+++ b/Engine/Plugins/AudioInsights/Source/AudioInsights/Private/Views/SignalFlowDashboardViewFactory.cpp
@@ -19,6 +19,10 @@
 #include "Widgets/Input/SSegmentedControl.h"
 #include "Widgets/Layout/SSeparator.h"
 
+#if MUTE_SOLO_ENABLED
+#include "Audio/AudioDebug.h"
+#endif // MUTE_SOLO_ENABLED
+
 #if WITH_EDITOR
 #include "AudioInsightsDetailsSelectionManager.h"
 #include "Editor.h"
@@ -372,9 +376,27 @@ namespace UE::Audio::Insights
 			.Orientation(EOrientation::Orient_Vertical)
 		]
 
+#if MUTE_SOLO_ENABLED
+		+ SHorizontalBox::Slot()
+		.HAlign(HAlign_Left)
+		.AutoWidth()
+		[
+			CreateToolbarClearMuteSoloButton()
+		]
+
+		+ SHorizontalBox::Slot()
+		.HAlign(HAlign_Left)
+		.AutoWidth()
+		.Padding(ToolbarButtonPaddingX, 0.0f)
+		[
+			SNew(SSeparator)
+			.Orientation(EOrientation::Orient_Vertical)
+		]
+#endif // MUTE_SOLO_ENABLED
+
 		+ SHorizontalBox::Slot()
 		.FillWidth(1.0f)
-			
+
 		+ SHorizontalBox::Slot()
 		.HAlign(HAlign_Right)
 		.AutoWidth()
@@ -569,6 +591,30 @@ namespace UE::Audio::Insights
 		];
 	}
 
+#if MUTE_SOLO_ENABLED
+	TSharedRef<SWidget> FSignalFlowDashboardViewFactory::CreateToolbarClearMuteSoloButton()
+	{
+		using namespace FSignalFlowDashboardViewFactoryPrivate;
+
+		const FSlateStyle& Style = FSlateStyle::Get();
+
+		return SNew(SButton)
+			.ButtonStyle(&FAppStyle::Get().GetWidgetStyle<FButtonStyle>(ButtonStyleName))
+			.ToolTipText(LOCTEXT("AudioDashboard_SignalFlow_ClearMutesAndSolosTooltip", "Clears all assigned mute/solo states."))
+			.OnClicked_Lambda([]()
+			{
+				if (FAudioDeviceManager* const AudioDeviceManager = FAudioDeviceManager::Get())
+				{
+					AudioDeviceManager->GetDebugger().ClearMutesAndSolos();
+				}
+				return FReply::Handled();
+			})
+			[
+				Style.CreateButtonContentWidget("AudioInsights.Icon.SoundDashboard.Reset", LOCTEXT("AudioDashboard_SignalFlow_ClearMutesAndSolosText", "Clear All Mutes/Solos"), ButtonTextStyleName)
+			];
+	}
+#endif // MUTE_SOLO_ENABLED
+
 	TSharedRef<SWidget> FSignalFlowDashboardViewFactory::OnGetSettingsMenuContent()
 	{
 		if (!CommandList.IsValid())
diff --git a/Engine/Plugins/AudioInsights/Source/AudioInsights/Public/Views/SignalFlowDashboardViewFactory.h b/Engine/Plugins/AudioInsights/Source/AudioInsights/Public/Views/SignalFlowDashboardViewFactory.h
--- a/Engine/Plugins/AudioInsights/Source/AudioInsights/Public/Views/SignalFlowDashboardViewFactory.h
+++ b/Engine/Plugins/AudioInsights/Source/AudioInsights/Public/Views/SignalFlowDashboardViewFactory.h
@@ -75,6 +75,9 @@ namespace UE::Audio::Insights
 		TSharedRef<SWidget> CreateToolbarSettingsMenu();
 		TSharedRef<SWidget> CreateToolbarToggleControls();
 		TSharedRef<SWidget> CreateToolbarButtonActions();
+#if MUTE_SOLO_ENABLED
+		TSharedRef<SWidget> CreateToolbarClearMuteSoloButton();
+#endif // MUTE_SOLO_ENABLED
 
 		TSharedRef<SWidget> OnGetSettingsMenuContent();
 		void BuildJustificationMenuContent(FMenuBuilder& OutMenuBuilder);
